AggregateTicking and Pipe merged into Engine/Plugins/Runtime/MassEntity

diff --git a/Engine/Plugins/Runtime/MassEntity/Source/MassEntity/MassEntity.Build.cs b/Engine/Plugins/Runtime/MassEntity/Source/MassEntity/MassEntity.Build.cs
new file mode 100644
--- /dev/null
+++ b/Engine/Plugins/Runtime/MassEntity/Source/MassEntity/MassEntity.Build.cs
@@ -0,0 +1,40 @@
+// Copyright Epic Games, Inc. All Rights Reserved.
+
+using UnrealBuildTool;
+
+namespace UnrealBuildTool.Rules
+{
+	public class MassEntity : ModuleRules
+	{
+		public MassEntity(ReadOnlyTargetRules Target) : base(Target)
+		{
+			PCHUsage = ModuleRules.PCHUsageMode.UseExplicitOrSharedPCHs;
+
+			PublicIncludePaths.AddRange(
+				new string[] {
+				ModuleDirectory + "/Public",
+				}
+			);
+
+			PublicDependencyModuleNames.AddRange(
+				new string[] {
+					"Core",
+					"CoreUObject",
+					"Engine",
+					"StructUtils",
+					"DeveloperSettings",
+				}
+			);
+
+			if (Target.bBuildEditor == true)
+			{
+				PrivateDependencyModuleNames.Add("UnrealEd");
+			}
+
+			if (Target.bBuildDeveloperTools == true)
+			{
+				PrivateDependencyModuleNames.Add("MessageLog");
+			}
+		}
+	}
+}
\ No newline at end of file
diff --git a/Engine/Plugins/Runtime/MassEntity/Source/MassEntity/Private/ArchetypeData.cpp b/Engine/Plugins/Runtime/MassEntity/Source/MassEntity/Private/ArchetypeData.cpp
new file mode 100644
--- /dev/null
+++ b/Engine/Plugins/Runtime/MassEntity/Source/MassEntity/Private/ArchetypeData.cpp
@@ -0,0 +1,887 @@
+// Copyright Epic Games, Inc. All Rights Reserved.
+
+#include "ArchetypeData.h"
+#include "Misc/StringBuilder.h"
+#include "VisualLogger/VisualLogger.h"
+#include "Async/ParallelFor.h"
+
+//////////////////////////////////////////////////////////////////////
+// FArchetypeData
+namespace UE { namespace AggregateTicking 
+{
+	constexpr static bool bBitwiseRelocateComponents = true;
+}} // UE::AggregateTicking
+
+void FArchetypeData::ForEachComponentType(TFunction< void(const UScriptStruct* /*Component*/)> Function) const
+{
+	for (const FArchetypeComponentConfig& ComponentData : ComponentConfigs)
+	{
+		Function(ComponentData.ComponentType);
+	}
+}
+
+bool FArchetypeData::HasComponentType(const UScriptStruct* ComponentType) const
+{
+	return (ComponentType && CompositionDescriptor.Components.Contains(*ComponentType));
+}
+
+bool FArchetypeData::IsEquivalent(TConstArrayView<const UScriptStruct*> OtherList) const
+{
+	if (OtherList.Num() != ComponentConfigs.Num())
+	{
+		return false;
+	}
+
+	for (int32 Index = 0; Index < ComponentConfigs.Num(); ++Index)
+	{
+		if (OtherList[Index] != ComponentConfigs[Index].ComponentType)
+		{
+			return false;
+		}
+	}
+
+	return true;
+}
+
+void FArchetypeData::Initialize(const FLWComponentBitSet& Components, const FLWTagBitSet& Tags, const FLWChunkComponentBitSet& ChunkComponents)
+{
+	TArray<const UScriptStruct*, TInlineAllocator<16>> SortedComponentList;
+	Components.ExportTypes(SortedComponentList);
+
+	SortedComponentList.Sort(FLWComponentSorterOperator<UScriptStruct>());
+
+	// Figure out how many bytes all of the individual components (and metadata) will cost per entity
+	int32 ComponentSizeTallyBytes = 0;
+
+	// Alignment padding computation is currently very conservative and over-estimated.
+	int32 AlignmentPadding = 0;
+	{
+		// Save room for the 'metadata' (entity array)
+		ComponentSizeTallyBytes += sizeof(FLWEntity);
+
+		// Tally up the component sizes and place them in the index map
+		ComponentConfigs.AddDefaulted(SortedComponentList.Num());
+		ComponentIndexMap.Reserve(SortedComponentList.Num());
+
+		for (int32 ComponentIndex = 0; ComponentIndex < SortedComponentList.Num(); ++ComponentIndex)
+		{
+			const UScriptStruct* ComponentType = SortedComponentList[ComponentIndex];
+			checkSlow(ComponentType);
+			ComponentConfigs[ComponentIndex].ComponentType = ComponentType;
+			
+			AlignmentPadding += ComponentType->GetMinAlignment();
+			ComponentSizeTallyBytes += ComponentType->GetStructureSize();
+
+			ComponentIndexMap.Add(ComponentType, ComponentIndex);
+			CompositionDescriptor.Components.Add(*ComponentType);
+		}
+		checkSlow(Components.IsEquivalent(CompositionDescriptor.Components));
+	}
+
+	CompositionDescriptor.Tags = Tags;
+	CompositionDescriptor.ChunkComponents = ChunkComponents;
+	// create the chunk component instances
+	TArray<const UScriptStruct*, TInlineAllocator<16>> ChunkComponentList;
+	ChunkComponents.ExportTypes(ChunkComponentList);
+	ChunkComponentList.Sort(FLWComponentSorterOperator<UScriptStruct>());
+	for (const UScriptStruct* ChunkComponentType : ChunkComponentList)
+	{
+		check(ChunkComponentType);
+		ChunkComponentTemplates.Add(FInstancedStruct(ChunkComponentType));
+	}
+
+	TotalBytesPerEntity = ComponentSizeTallyBytes;
+	int32 ChunkAvailableSize = GetChunkAllocSize() - AlignmentPadding;
+	check(TotalBytesPerEntity <= ChunkAvailableSize);
+
+	NumEntitiesPerChunk = ChunkAvailableSize / TotalBytesPerEntity;
+
+	// Set up the offsets for each component into the chunk data
+	EntityListOffsetWithinChunk = 0;
+	int32 CurrentOffset = NumEntitiesPerChunk * sizeof(FLWEntity);
+	for (FArchetypeComponentConfig& ComponentData : ComponentConfigs)
+	{
+		CurrentOffset = Align(CurrentOffset, ComponentData.ComponentType->GetMinAlignment());
+		ComponentData.ArrayOffsetWithinChunk = CurrentOffset;
+		const int32 SizeOfThisComponentArray = NumEntitiesPerChunk * ComponentData.ComponentType->GetStructureSize();
+		CurrentOffset += SizeOfThisComponentArray;
+	}
+}
+
+void FArchetypeData::InitializeWithSibling(const FArchetypeData& SiblingArchetype, const FLWTagBitSet& OverrideTags)
+{
+	checkf(IsInitialized() == false, TEXT("Trying to %s but this archetype has already been initialized"));
+
+	ComponentConfigs = SiblingArchetype.ComponentConfigs;
+	ComponentIndexMap = SiblingArchetype.ComponentIndexMap;
+	CompositionDescriptor.Components = SiblingArchetype.CompositionDescriptor.Components;
+	CompositionDescriptor.Tags = OverrideTags;
+	CompositionDescriptor.ChunkComponents = SiblingArchetype.CompositionDescriptor.ChunkComponents;
+	ChunkComponentTemplates = SiblingArchetype.ChunkComponentTemplates;
+
+	TotalBytesPerEntity = SiblingArchetype.TotalBytesPerEntity;
+	NumEntitiesPerChunk = SiblingArchetype.NumEntitiesPerChunk;
+
+	// Set up the offsets for each component into the chunk data
+	EntityListOffsetWithinChunk = 0;
+}
+
+void FArchetypeData::AddEntity(FLWEntity Entity)
+{
+	AddEntityInternal(Entity, true/*bInitializeComponents*/);
+}
+
+int32 FArchetypeData::AddEntityInternal(FLWEntity Entity, const bool bInitializeComponents)
+{
+	int32 IndexWithinChunk = 0;
+	int32 AbsoluteIndex = 0;
+	int32 ChunkIndex = 0;
+	int32 EmptyChunkIndex = INDEX_NONE;
+	int32 EmptyAbsoluteIndex = INDEX_NONE;
+
+	FArchetypeChunk* DestinationChunk = nullptr;
+	// Check chunks for a free spot (trying to reuse the earlier ones first so later ones might get freed up) 
+	//@TODO: This could be accelerated to include a cached index to the first chunk with free spots or similar
+	for (FArchetypeChunk& Chunk : Chunks)
+	{
+		if (Chunk.GetNumInstances() == 0)
+		{
+			// Remember first empty chunk but continue looking for a chunk that has space and same group tag
+			if (EmptyChunkIndex == INDEX_NONE)
+			{
+				EmptyChunkIndex = ChunkIndex;
+				EmptyAbsoluteIndex = AbsoluteIndex;
+			}
+		}
+		else if (Chunk.GetNumInstances() < NumEntitiesPerChunk)
+		{
+			IndexWithinChunk = Chunk.GetNumInstances();
+			AbsoluteIndex += IndexWithinChunk;
+
+			Chunk.AddInstance();
+
+			DestinationChunk = &Chunk;
+			break;
+		}
+		AbsoluteIndex += NumEntitiesPerChunk;
+		++ChunkIndex;
+	}
+
+	if (DestinationChunk == nullptr)
+	{
+		// Check if it is a recycled chunk
+		if (EmptyChunkIndex != INDEX_NONE)
+		{
+			DestinationChunk = &Chunks[EmptyChunkIndex];
+			DestinationChunk->Recycle(ChunkComponentTemplates);
+			AbsoluteIndex = EmptyAbsoluteIndex;
+		}
+		else
+		{
+			DestinationChunk = &Chunks.Emplace_GetRef(GetChunkAllocSize(), ChunkComponentTemplates);
+		}
+
+		check(DestinationChunk);
+		DestinationChunk->AddInstance();
+	}
+
+	// Initialize the component memory
+	if (bInitializeComponents)
+	{
+		for (const FArchetypeComponentConfig& ComponentConfig : ComponentConfigs)
+		{
+			void* ComponentPtr = ComponentConfig.GetComponentData(DestinationChunk->GetRawMemory(), IndexWithinChunk);
+			ComponentConfig.ComponentType->InitializeStruct(ComponentPtr);
+		}
+	}
+
+	// Add to the table and map
+	EntityMap.Add(Entity.Index, AbsoluteIndex);
+	DestinationChunk->GetEntityArrayElementRef(EntityListOffsetWithinChunk, IndexWithinChunk) = Entity;
+
+	return AbsoluteIndex;
+}
+
+void FArchetypeData::RemoveEntity(FLWEntity Entity)
+{
+	const int32 AbsoluteIndex = EntityMap.FindAndRemoveChecked(Entity.Index);
+	RemoveEntityInternal(AbsoluteIndex, true/*bDestroyComponents*/);
+}
+
+void FArchetypeData::RemoveEntityInternal(const int32 AbsoluteIndex, const bool bDestroyComponents)
+{
+	const int32 ChunkIndex = AbsoluteIndex / NumEntitiesPerChunk;
+	const int32 IndexWithinChunk = AbsoluteIndex % NumEntitiesPerChunk;
+
+	FArchetypeChunk& Chunk = Chunks[ChunkIndex];
+
+	Chunk.RemoveInstance();
+	const int32 IndexToSwapFrom = Chunk.GetNumInstances();
+
+	checkf(bDestroyComponents || UE::AggregateTicking::bBitwiseRelocateComponents, TEXT("We allow not to destroy components only in bit wise relocation mode."));
+
+	// Remove and swap the last entry in the chunk to the location of the removed item (if it's not the same as the dying entry)
+	if (IndexToSwapFrom != IndexWithinChunk)
+	{
+		for (const FArchetypeComponentConfig& ComponentConfig : ComponentConfigs)
+		{
+			void* DyingComponentPtr = ComponentConfig.GetComponentData(Chunk.GetRawMemory(), IndexWithinChunk);
+			void* MovingComponentPtr = ComponentConfig.GetComponentData(Chunk.GetRawMemory(), IndexToSwapFrom);
+
+			if (UE::AggregateTicking::bBitwiseRelocateComponents)
+			{
+				// Destroy component data
+				if (bDestroyComponents)
+				{
+					ComponentConfig.ComponentType->DestroyStruct(DyingComponentPtr);
+				}
+
+				// Move last entry
+				FMemory::Memcpy(DyingComponentPtr, MovingComponentPtr, ComponentConfig.ComponentType->GetStructureSize());
+			}
+			else
+			{
+				// Destroy & initialize the component data
+				ComponentConfig.ComponentType->ClearScriptStruct(DyingComponentPtr);
+
+				// Copy last entry
+				ComponentConfig.ComponentType->CopyScriptStruct(DyingComponentPtr, MovingComponentPtr);
+
+				// Destroy last entry
+				ComponentConfig.ComponentType->DestroyStruct(MovingComponentPtr);
+			}
+		}
+
+		// Update the entity table and map
+		const FLWEntity EntityBeingSwapped = Chunk.GetEntityArrayElementRef(EntityListOffsetWithinChunk, IndexToSwapFrom);
+		Chunk.GetEntityArrayElementRef(EntityListOffsetWithinChunk, IndexWithinChunk) = EntityBeingSwapped;
+		EntityMap.FindChecked(EntityBeingSwapped.Index) = AbsoluteIndex;
+	}
+	else if (!UE::AggregateTicking::bBitwiseRelocateComponents || bDestroyComponents)
+	{
+		for (const FArchetypeComponentConfig& ComponentConfig : ComponentConfigs)
+		{
+			// Destroy the component data
+			void* DyingComponentPtr = ComponentConfig.GetComponentData(Chunk.GetRawMemory(), IndexWithinChunk);
+			ComponentConfig.ComponentType->DestroyStruct(DyingComponentPtr);
+		}
+	}
+
+	// If the chunk itself is empty now, see if we can remove it entirely
+	// Note: This is only possible for trailing chunks, to avoid messing up the absolute indices in the entities map
+	while ((Chunks.Num() > 0) && (Chunks.Last().GetNumInstances() == 0))
+	{
+		Chunks.RemoveAt(Chunks.Num() - 1, 1, /*bAllowShrinking=*/ false);
+	}
+}
+
+void FArchetypeData::BatchDestroyEntityChunks(const FArchetypeChunkCollection& ChunkCollection, TArray<FLWEntity>& OutEntitiesRemoved)
+{
+	const int32 InitialOutEntitiesCount = OutEntitiesRemoved.Num();
+
+	// Sorting the subchunks info so that subchunks of a given chunk are processed "from the back". Otherwise removing 
+	// a subchunk from the front of the chunk would inevitably invalidate following subchunks' information.
+	TArray<FArchetypeChunkCollection::FChunkInfo> Subchunks(ChunkCollection.GetChunks());
+	Subchunks.Sort([](const FArchetypeChunkCollection::FChunkInfo& A, const FArchetypeChunkCollection::FChunkInfo& B) 
+		{ 
+			return A.ChunkIndex < B.ChunkIndex || (A.ChunkIndex == B.ChunkIndex && A.SubchunkStart > B.SubchunkStart);
+		});
+
+	for (const FArchetypeChunkCollection::FChunkInfo SubchunkInfo : Subchunks)
+	{ 
+		FArchetypeChunk& Chunk = Chunks[SubchunkInfo.ChunkIndex];
+
+		// gather entities we're about to remove
+		FLWEntity* DyingEntityPtr = &Chunk.GetEntityArrayElementRef(EntityListOffsetWithinChunk, SubchunkInfo.SubchunkStart);
+		OutEntitiesRemoved.Append(DyingEntityPtr, SubchunkInfo.Length);
+
+		const int32 NumberToMove = FMath::Min(Chunk.GetNumInstances() - (SubchunkInfo.SubchunkStart + SubchunkInfo.Length), SubchunkInfo.Length);
+		checkf(NumberToMove >= 0, TEXT("Trying to move a negative number of elements indicates a problem with SubchunkInfo, it's possibly out of date."));
+		const int32 NumberToCut = FMath::Max(SubchunkInfo.Length - NumberToMove, 0);
+		
+		if (NumberToMove > 0)
+		{
+			const int32 SwapStartIndex = Chunk.GetNumInstances() - NumberToMove;
+			checkf((SubchunkInfo.SubchunkStart + NumberToMove - 1) < SwapStartIndex, TEXT("Remove and Move ranges overlap"));
+
+			for (const FArchetypeComponentConfig& ComponentConfig : ComponentConfigs)
+			{
+				void* DyingComponentPtr = ComponentConfig.GetComponentData(Chunk.GetRawMemory(), SubchunkInfo.SubchunkStart);
+				void* MovingComponentPtr = ComponentConfig.GetComponentData(Chunk.GetRawMemory(), SwapStartIndex);
+
+				if (UE::AggregateTicking::bBitwiseRelocateComponents)
+				{
+					// Destroy the components we'll replace by the following copy
+					ComponentConfig.ComponentType->DestroyStruct(DyingComponentPtr, NumberToMove);
+
+					// Swap components to the empty space just created.
+					FMemory::Memcpy(DyingComponentPtr, MovingComponentPtr, ComponentConfig.ComponentType->GetStructureSize() * NumberToMove);
+				}
+				else
+				{
+					// Clear components that we will copy over. Clear destroys and initializes the components, which is needed for CopyScriptStruct().
+					ComponentConfig.ComponentType->ClearScriptStruct(DyingComponentPtr, NumberToMove);
+
+					// Swap components into the empty space just created.
+					ComponentConfig.ComponentType->CopyScriptStruct(DyingComponentPtr, MovingComponentPtr, NumberToMove);
+
+					// Destroy the components that were moved.
+					ComponentConfig.ComponentType->DestroyStruct(MovingComponentPtr, NumberToMove);
+				}
+			}
+
+			// Update the entity table and map
+			const FLWEntity* MovingEntityPtr = &Chunk.GetEntityArrayElementRef(EntityListOffsetWithinChunk, SwapStartIndex);
+			int32 AbsoluteIndex = SubchunkInfo.ChunkIndex * NumEntitiesPerChunk + SubchunkInfo.SubchunkStart;
+
+			for (int i = 0; i < NumberToMove; ++i)
+			{
+				DyingEntityPtr[i] = MovingEntityPtr[i];
+				EntityMap.FindChecked(MovingEntityPtr[i].Index) = AbsoluteIndex++;
+			}
+		}
+
+		if (NumberToCut > 0)
+		{
+			// just clean up the rest. Note that we explicitly do not clean the spots vacated by entities moved from 
+			// the back of the chunk - if we did the risk calling DestroyStruct on them multiple times
+			const int32 CutStartIndex = SubchunkInfo.SubchunkStart + NumberToMove;
+			for (const FArchetypeComponentConfig& ComponentConfig : ComponentConfigs)
+			{
+				// Destroy the component data
+				void* DyingComponentPtr = ComponentConfig.GetComponentData(Chunk.GetRawMemory(), CutStartIndex);
+				ComponentConfig.ComponentType->DestroyStruct(DyingComponentPtr, NumberToCut);
+			}
+		}
+
+		Chunk.RemoveMultipleInstances(SubchunkInfo.Length);
+	}
+
+	for (int i = InitialOutEntitiesCount; i < OutEntitiesRemoved.Num(); ++i)
+	{
+		EntityMap.FindAndRemoveChecked(OutEntitiesRemoved[i].Index);
+	}
+
+	// If the chunk itself is empty now, see if we can remove it entirely
+	// Note: This is only possible for trailing chunks, to avoid messing up the absolute indices in the entities map
+	while ((Chunks.Num() > 0) && (Chunks.Last().GetNumInstances() == 0))
+	{
+		Chunks.RemoveAt(Chunks.Num() - 1, 1, /*bAllowShrinking=*/ false);
+	}
+}
+
+bool FArchetypeData::HasComponentDataForEntity(const UScriptStruct* ComponentType, int32 EntityIndex) const
+{
+	return (ComponentType && CompositionDescriptor.Components.Contains(*ComponentType));
+}
+
+void* FArchetypeData::GetComponentDataForEntityChecked(const UScriptStruct* ComponentType, int32 EntityIndex) const
+{
+	const FInternalEntityHandle InternalIndex = MakeEntityHandle(EntityIndex);
+	
+	// failing the below Find means given entity's archetype is missing given ComponentType
+	const int32 ComponentIndex = ComponentIndexMap.FindChecked(ComponentType);
+	return GetComponentData(ComponentIndex, InternalIndex);
+}
+
+void* FArchetypeData::GetComponentDataForEntity(const UScriptStruct* ComponentType, int32 EntityIndex) const
+{
+	if (const int32* ComponentIndex = ComponentIndexMap.Find(ComponentType))
+	{
+		FInternalEntityHandle InternalIndex = MakeEntityHandle(EntityIndex);
+		// failing the below Find means given entity's archetype is missing given ComponentType
+		return GetComponentData(*ComponentIndex, InternalIndex);
+	}
+	return nullptr;
+}
+
+void FArchetypeData::SetComponentsData(const FLWEntity Entity, TArrayView<const FInstancedStruct> ComponentInstances)
+{
+	FInternalEntityHandle InternalIndex = MakeEntityHandle(Entity);
+
+	for (const FInstancedStruct& Instance : ComponentInstances)
+	{
+		const UScriptStruct* ComponentType = Instance.GetScriptStruct();
+		check(ComponentType);
+		const int32 ComponentIndex = ComponentIndexMap.FindChecked(ComponentType);
+		void* ComponentMemory = GetComponentData(ComponentIndex, InternalIndex);
+		// No UE::AggregateTicking::bBitwiseRelocateComponents, this isn't a move component
+		ComponentType->CopyScriptStruct(ComponentMemory, Instance.GetMemory());
+	}
+}
+
+void FArchetypeData::SetComponentData(const FArchetypeChunkCollection& ChunkCollection, const FInstancedStruct& ComponentSource)
+{
+	check(ComponentSource.IsValid());
+	const UScriptStruct* ComponentType = ComponentSource.GetScriptStruct();
+	check(ComponentType);
+	const int32 ComponentIndex = ComponentIndexMap.FindChecked(ComponentType);
+	const int32 ComponentTypeSize = ComponentType->GetStructureSize();
+	const uint8* ComponentSourceMemory = ComponentSource.GetMemory();
+	check(ComponentSourceMemory);
+	
+	for (FArchetypeChunkIterator ChunkIterator(ChunkCollection); ChunkIterator; ++ChunkIterator)
+	{
+		uint8* ComponentMemory = (uint8*)ComponentConfigs[ComponentIndex].GetComponentData(Chunks[ChunkIterator->ChunkIndex].GetRawMemory(), ChunkIterator->SubchunkStart);
+		for (int i = ChunkIterator->Length; i; --i, ComponentMemory += ComponentTypeSize)
+		{
+			// No UE::AggregateTicking::bBitwiseRelocateComponents, this isn't a move of a component
+			ComponentType->CopyScriptStruct(ComponentMemory, ComponentSourceMemory);
+		}
+	}
+}
+
+void FArchetypeData::SetDefaultChunkComponentValue(FConstStructView InstancedStruct)
+{
+	if (ensure(InstancedStruct.GetScriptStruct()) && CompositionDescriptor.ChunkComponents.Contains(*InstancedStruct.GetScriptStruct()))
+	{
+		if (FInstancedStruct* Element = ChunkComponentTemplates.FindByPredicate([&InstancedStruct](const FInstancedStruct& Element) { return Element.GetScriptStruct() == InstancedStruct.GetScriptStruct(); }))
+		{
+			*Element = InstancedStruct;
+		}
+	}
+}
+
+void FArchetypeData::MoveEntityToAnotherArchetype(const FLWEntity Entity, FArchetypeData& NewArchetype)
+{
+	check(&NewArchetype != this);
+
+	const int32 AbsoluteIndex = EntityMap.FindAndRemoveChecked(Entity.Index);
+	const int32 ChunkIndex = AbsoluteIndex / NumEntitiesPerChunk;
+	const int32 IndexWithinChunk = AbsoluteIndex % NumEntitiesPerChunk;
+	FArchetypeChunk& Chunk = Chunks[ChunkIndex];
+
+	constexpr bool bInitializeComponentsDuringCreation = !UE::AggregateTicking::bBitwiseRelocateComponents;
+	const int32 NewAbsoluteIndex = NewArchetype.AddEntityInternal(Entity, bInitializeComponentsDuringCreation);
+	const int32 NewChunkIndex = NewAbsoluteIndex / NewArchetype.NumEntitiesPerChunk;
+	const int32 NewIndexWithinChunk = NewAbsoluteIndex % NewArchetype.NumEntitiesPerChunk;
+	FArchetypeChunk& NewChunk = NewArchetype.Chunks[NewChunkIndex];
+
+	// for every NewArchetype's component see if it was in the old component as well and if so copy it's value. 
+	// If not then initialize the component (unless it has already been initialized based on bInitializeComponentsDuringCreation
+	// value
+	for (const FArchetypeComponentConfig& NewComponentConfig : NewArchetype.ComponentConfigs)
+	{
+		const int32* OldComponentIndex = ComponentIndexMap.Find(NewComponentConfig.ComponentType);
+		void* Dst = NewComponentConfig.GetComponentData(NewChunk.GetRawMemory(), NewIndexWithinChunk);
+
+		// Only copy if the component type exists in both archetypes
+		if (OldComponentIndex)
+		{
+			const void* Src = ComponentConfigs[*OldComponentIndex].GetComponentData(Chunk.GetRawMemory(), IndexWithinChunk);
+			if (UE::AggregateTicking::bBitwiseRelocateComponents)
+			{
+				FMemory::Memcpy(Dst, Src, NewComponentConfig.ComponentType->GetStructureSize());
+			}
+			else
+			{
+				NewComponentConfig.ComponentType->CopyScriptStruct(Dst, Src);
+			}
+		}
+		else if (bInitializeComponentsDuringCreation == false)
+		{
+			// the component's unique to the NewArchetype need to be initialized
+			// @todo we're doing it for tags here as well. A tiny bit of perf lost. Probably not worth adding a check
+			// but something to keep in mind. Will go away once tags are more of an archetype fragment than entity's
+			NewComponentConfig.ComponentType->InitializeStruct(Dst);
+		}
+	}
+
+	constexpr bool bDestroyComponents = !UE::AggregateTicking::bBitwiseRelocateComponents;
+	RemoveEntityInternal(AbsoluteIndex, bDestroyComponents);
+}
+
+void FArchetypeData::ExecuteFunction(FLWComponentSystemExecutionContext& RunContext, const FLWComponentSystemExecuteFunction& Function, const FLWRequirementIndicesMapping& RequirementMapping, const FArchetypeChunkCollection& ChunkCollection)
+{
+	check(ChunkCollection.GetArchetype() == this);
+
+	// mz@todo to be removed
+	RunContext.SetCurrentArchetypeData(*this);
+	for (FArchetypeChunkIterator ChunkIterator(ChunkCollection); ChunkIterator; ++ChunkIterator)
+	{
+		FArchetypeChunk& Chunk = Chunks[ChunkIterator->ChunkIndex];
+		
+		const int32 ChunkLength = ChunkIterator->Length > 0 ? ChunkIterator->Length : (Chunk.GetNumInstances() - ChunkIterator->SubchunkStart);
+		if (ChunkLength)
+		{
+			checkf((ChunkIterator->SubchunkStart + ChunkLength) <= Chunk.GetNumInstances() && ChunkLength > 0, TEXT("Invalid subchunk, it is going over the number of instances in the chunk or it is empty."));
+
+			RunContext.SetCurrentChunkSerialModificationNumber(Chunk.GetSerialModificationNumber());
+			BindChunkComponentRequirements(RunContext, RequirementMapping.ChunkComponents, Chunk);
+			BindEntityRequirements(RunContext, RequirementMapping.EntityComponents, Chunk, ChunkIterator->SubchunkStart, ChunkLength);
+
+			Function(RunContext);
+		}
+	}
+}
+
+void FArchetypeData::ExecuteFunction(FLWComponentSystemExecutionContext& RunContext, const FLWComponentSystemExecuteFunction& Function, const FLWRequirementIndicesMapping& RequirementMapping, const FLWComponentSystemChunkConditionFunction& ChunkCondition)
+{
+	// mz@todo to be removed
+	RunContext.SetCurrentArchetypeData(*this);
+	for (FArchetypeChunk& Chunk : Chunks)
+	{
+		if (Chunk.GetNumInstances())
+		{
+			RunContext.SetCurrentChunkSerialModificationNumber(Chunk.GetSerialModificationNumber());
+			BindChunkComponentRequirements(RunContext, RequirementMapping.ChunkComponents, Chunk);
+
+			if (!ChunkCondition || ChunkCondition(RunContext))
+			{
+				BindEntityRequirements(RunContext, RequirementMapping.EntityComponents, Chunk, 0, Chunk.GetNumInstances());
+				Function(RunContext);
+			}
+		}
+	}
+}
+
+void FArchetypeData::ExecutionFunctionForChunk(FLWComponentSystemExecutionContext RunContext, const FLWComponentSystemExecuteFunction& Function, const FLWRequirementIndicesMapping& RequirementMapping, const FArchetypeChunkCollection::FChunkInfo& ChunkInfo, const FLWComponentSystemChunkConditionFunction& ChunkCondition)
+{
+	FArchetypeChunk& Chunk = Chunks[ChunkInfo.ChunkIndex];
+	const int32 ChunkLength = ChunkInfo.Length > 0 ? ChunkInfo.Length : (Chunk.GetNumInstances() - ChunkInfo.SubchunkStart);
+
+	if (ChunkLength)
+	{
+		RunContext.SetCurrentArchetypeData(*this);
+		RunContext.SetCurrentChunkSerialModificationNumber(Chunk.GetSerialModificationNumber());
+		BindChunkComponentRequirements(RunContext, RequirementMapping.ChunkComponents, Chunk);
+
+		if (!ChunkCondition || ChunkCondition(RunContext))
+		{
+			BindEntityRequirements(RunContext, RequirementMapping.EntityComponents, Chunk, ChunkInfo.SubchunkStart, ChunkLength);
+			Function(RunContext);
+		}
+	}
+}
+
+void FArchetypeData::CompactEntities(const double TimeAllowed)
+{
+	const double TimeAllowedEnd = FPlatformTime::Seconds() + TimeAllowed;
+
+	TArray<FArchetypeChunk*> SortedChunks;
+	for (FArchetypeChunk& Chunk : Chunks)
+	{
+		// Skip already full chunks
+		const int32 NumInstances = Chunk.GetNumInstances();
+		if(NumInstances > 0 && NumInstances < NumEntitiesPerChunk)
+		{
+			SortedChunks.Add(&Chunk);
+		}
+	}
+
+	// Check if there is anything to compact at all
+	if (SortedChunks.Num() <= 1)
+	{
+		return;
+	}
+
+	SortedChunks.Sort([](const FArchetypeChunk& LHS, const FArchetypeChunk& RHS)
+	{
+		return LHS.GetNumInstances() < RHS.GetNumInstances();
+	});
+
+	int32 ChunkToFillSortedIdx = 0;
+	int32 ChunkToEmptySortedIdx = SortedChunks.Num() -1;
+	while (ChunkToFillSortedIdx < ChunkToEmptySortedIdx && FPlatformTime::Seconds() < TimeAllowedEnd)
+	{
+		while (ChunkToFillSortedIdx < SortedChunks.Num() && SortedChunks[ChunkToFillSortedIdx]->GetNumInstances() == NumEntitiesPerChunk)
+		{
+			ChunkToFillSortedIdx++;
+		}
+		while (ChunkToEmptySortedIdx >= 0 && SortedChunks[ChunkToEmptySortedIdx]->GetNumInstances() == 0)
+		{
+			ChunkToEmptySortedIdx--;
+		}
+		if (ChunkToFillSortedIdx >= ChunkToEmptySortedIdx)
+		{
+			break;
+		}
+
+		FArchetypeChunk* ChunkToFill = SortedChunks[ChunkToFillSortedIdx];
+		FArchetypeChunk* ChunkToEmpty = SortedChunks[ChunkToEmptySortedIdx];
+		const int32 NumberOfEntitiesToMove =  FMath::Min(NumEntitiesPerChunk-ChunkToFill->GetNumInstances(), ChunkToEmpty->GetNumInstances());
+		const int32 FromIndex = ChunkToEmpty->GetNumInstances() - NumberOfEntitiesToMove;
+		const int32 ToIndex = ChunkToFill->GetNumInstances();
+		check(NumberOfEntitiesToMove > 0);
+
+		for (const FArchetypeComponentConfig& ComponentConfig : ComponentConfigs)
+		{
+			void* FromComponentPtr = ComponentConfig.GetComponentData(ChunkToEmpty->GetRawMemory(), FromIndex);
+			void* ToComponentPtr = ComponentConfig.GetComponentData(ChunkToFill->GetRawMemory(), ToIndex);
+
+			if (UE::AggregateTicking::bBitwiseRelocateComponents)
+			{
+				// Move all entries
+				FMemory::Memcpy(ToComponentPtr, FromComponentPtr, ComponentConfig.ComponentType->GetStructureSize() * NumberOfEntitiesToMove);
+			}
+			else
+			{
+				// Destroy & initialize the component data
+				ComponentConfig.ComponentType->ClearScriptStruct(ToComponentPtr, NumberOfEntitiesToMove);
+
+				// Copy all entries
+				ComponentConfig.ComponentType->CopyScriptStruct(ToComponentPtr, FromComponentPtr, NumberOfEntitiesToMove);
+
+				// Destroy all entries
+				ComponentConfig.ComponentType->DestroyStruct(FromComponentPtr, NumberOfEntitiesToMove);
+			}
+		}
+
+		FLWEntity* FromEntity = &ChunkToEmpty->GetEntityArrayElementRef(EntityListOffsetWithinChunk, FromIndex);
+		FLWEntity* ToEntity = &ChunkToFill->GetEntityArrayElementRef(EntityListOffsetWithinChunk, ToIndex);
+		FMemory::Memcpy(ToEntity, FromEntity, NumberOfEntitiesToMove * sizeof(FLWEntity));
+		ChunkToFill->AddMultipleInstances(NumberOfEntitiesToMove);
+		ChunkToEmpty->RemoveMultipleInstances(NumberOfEntitiesToMove);
+
+		const int32 ChunkToFillIdx = ChunkToFill - &Chunks[0];
+		check(ChunkToFillIdx >=0 && ChunkToFillIdx < Chunks.Num());
+		const int32 AbsoluteIndex = ChunkToFillIdx * NumEntitiesPerChunk + ToIndex;
+
+		for (int32 i =0; i < NumberOfEntitiesToMove; i++, ++ToEntity)
+		{
+			EntityMap.FindChecked(ToEntity->Index) = AbsoluteIndex+i;
+		}
+	}
+}
+
+void FArchetypeData::GetRequirementsComponentMapping(TConstArrayView<FLWComponentRequirement> Requirements, FLWComponentIndicesMapping& OutComponentIndices)
+{
+	OutComponentIndices.Reset(Requirements.Num());
+	for (const FLWComponentRequirement& Requirement : Requirements)
+	{
+		if (Requirement.RequiresBinding())
+		{
+			const int32* ComponentIndex = ComponentIndexMap.Find(Requirement.StructType);
+			check(ComponentIndex != nullptr || Requirement.IsOptional());
+			OutComponentIndices.Add(ComponentIndex ? *ComponentIndex : INDEX_NONE);
+		}
+	}
+}
+
+void FArchetypeData::GetRequirementsChunkComponentMapping(TConstArrayView<FLWComponentRequirement> ChunkRequirements, FLWComponentIndicesMapping& OutComponentIndices)
+{
+	int32 LastFoundComponentIndex = -1;
+	OutComponentIndices.Reset(ChunkRequirements.Num());
+	for (const FLWComponentRequirement& Requirement : ChunkRequirements)
+	{
+		if (Requirement.RequiresBinding())
+		{
+			int32 ComponentIndex = INDEX_NONE;
+			for (int32 i = LastFoundComponentIndex + 1; i < ChunkComponentTemplates.Num(); ++i)
+			{
+				if (ChunkComponentTemplates[i].GetScriptStruct()->IsChildOf(Requirement.StructType))
+				{
+					ComponentIndex = i;
+					break;
+				}
+			}
+
+			check(ComponentIndex != INDEX_NONE);
+			OutComponentIndices.Add(ComponentIndex);
+			LastFoundComponentIndex = ComponentIndex;
+		}
+	}
+}
+
+void FArchetypeData::BindEntityRequirements(FLWComponentSystemExecutionContext& RunContext, const FLWComponentIndicesMapping& EntityComponentsMapping, FArchetypeChunk& Chunk, const int32 SubchunkStart, const int32 SubchunkLength)
+{
+	// auto-correcting number of entities to process in case SubchunkStart +  SubchunkLength > Chunk.GetNumInstances()
+	const int32 NumEntities = SubchunkLength >= 0 ? FMath::Min(SubchunkLength, Chunk.GetNumInstances() - SubchunkStart) : Chunk.GetNumInstances();
+	check(SubchunkStart >= 0 && SubchunkStart < Chunk.GetNumInstances());
+
+	if (EntityComponentsMapping.Num() > 0)
+	{
+		check(RunContext.GetMutableRequirements().Num() == EntityComponentsMapping.Num());
+
+		for (int i = 0; i < EntityComponentsMapping.Num(); ++i)
+		{
+			FLWComponentSystemExecutionContext::FComponentView& Requirement = RunContext.ComponentViews[i];
+			const int32 ComponentIndex = EntityComponentsMapping[i];
+
+			check(ComponentIndex != INDEX_NONE || Requirement.Requirement.IsOptional());
+			if (ComponentIndex != INDEX_NONE)
+			{
+				Requirement.ComponentView = TArrayView<FLWComponentData>((FLWComponentData*)GetComponentData(ComponentIndex, Chunk.GetRawMemory(), SubchunkStart), NumEntities);
+			}
+			else
+			{
+				// @todo this might not be needed
+				Requirement.ComponentView = TArrayView<FLWComponentData>();
+			}
+		}
+	}
+	else
+	{
+		// Map in the required data arrays from the current chunk to the array views
+		for (FLWComponentSystemExecutionContext::FComponentView& Requirement : RunContext.GetMutableRequirements())
+		{
+			const int32* ComponentIndex = ComponentIndexMap.Find(Requirement.Requirement.StructType);
+			check(ComponentIndex != nullptr || Requirement.Requirement.IsOptional());
+			if (ComponentIndex)
+			{
+				Requirement.ComponentView = TArrayView<FLWComponentData>((FLWComponentData*)GetComponentData(*ComponentIndex, Chunk.GetRawMemory(), SubchunkStart), NumEntities);
+			}
+			else
+			{
+				Requirement.ComponentView = TArrayView<FLWComponentData>();
+			}
+		}
+	}
+
+	RunContext.EntityListView = TArrayView<FLWEntity>(&Chunk.GetEntityArrayElementRef(EntityListOffsetWithinChunk, SubchunkStart), NumEntities);
+}
+
+void FArchetypeData::BindChunkComponentRequirements(FLWComponentSystemExecutionContext& RunContext, const FLWComponentIndicesMapping& ChunkComponentsMapping, FArchetypeChunk& Chunk)
+{
+	if (ChunkComponentsMapping.Num() > 0)
+	{
+		check(RunContext.GetMutableChunkRequirements().Num() == ChunkComponentsMapping.Num());
+
+		for (int i = 0; i < ChunkComponentsMapping.Num(); ++i)
+		{
+			FLWComponentSystemExecutionContext::FChunkComponentView& ChunkRequirement = RunContext.ChunkComponents[i];
+			const int32 ChunkComponentIndex = ChunkComponentsMapping[i];
+
+			check(ChunkComponentIndex != INDEX_NONE);
+			ChunkRequirement.ChunkComponentView = Chunk.GetMutableChunkComponentViewChecked(ChunkComponentIndex);
+		}
+	}
+	else
+	{
+		for (FLWComponentSystemExecutionContext::FChunkComponentView& ChunkRequirement : RunContext.GetMutableChunkRequirements())
+		{
+			FInstancedStruct* ChunkComponentInstance = Chunk.FindMutableChunkComponent(ChunkRequirement.Requirement.StructType);
+			check(ChunkComponentInstance != nullptr);
+			ChunkRequirement.ChunkComponentView = FStructView(*ChunkComponentInstance);
+		}
+	}
+}
+
+SIZE_T FArchetypeData::GetAllocatedSize() const
+{
+	return sizeof(FArchetypeData) +
+		ComponentConfigs.GetAllocatedSize() +
+		Chunks.GetAllocatedSize() +
+		(Chunks.Num() * GetChunkAllocSize()) +
+		EntityMap.GetAllocatedSize() +
+		ComponentIndexMap.GetAllocatedSize();
+}
+
+FString FArchetypeData::DebugGetDescription() const
+{
+#if WITH_AGGREGATETICKING_DEBUG
+	TStringBuilder<256> ArchetypeDebugName;
+	
+	ArchetypeDebugName.Append(TEXT("<"));
+
+	bool bNeedsComma = false;
+	for (const FArchetypeComponentConfig& ComponentConfig : ComponentConfigs)
+	{
+		if (bNeedsComma)
+		{
+			ArchetypeDebugName.Append(TEXT(","));
+		}
+		ArchetypeDebugName.Append(ComponentConfig.ComponentType->GetName());
+		bNeedsComma = true;
+	}
+
+	ArchetypeDebugName.Append(TEXT(">"));
+	return ArchetypeDebugName.ToString();
+#else
+	return {};
+#endif
+}
+
+#if WITH_AGGREGATETICKING_DEBUG
+void FArchetypeData::DebugPrintArchetype(FOutputDevice& Ar)
+{
+	FStringOutputDevice TagsDecription;
+	CompositionDescriptor.Tags.DebugGetStringDesc(TagsDecription);
+	Ar.Logf(ELogVerbosity::Log, TEXT("Tags: %s"), *TagsDecription);
+	Ar.Logf(ELogVerbosity::Log, TEXT("Components: %s"), *DebugGetDescription());
+	Ar.Logf(ELogVerbosity::Log, TEXT("\tChunks: %d x %d KB = %d KB total"), Chunks.Num(), GetChunkAllocSize() / 1024, (GetChunkAllocSize()*Chunks.Num()) / 1024);
+	
+	int ChunkWithComponentsCount = 0;
+	for (FArchetypeChunk& Chunk : Chunks)
+	{
+		ChunkWithComponentsCount += Chunk.DebugGetChunkComponentCount() > 0 ? 1 : 0;
+	}
+	if (ChunkWithComponentsCount)
+	{
+		Ar.Logf(ELogVerbosity::Log, TEXT("\tChunks with components: %d"), ChunkWithComponentsCount);
+	}
+
+	const int32 CurrentEntityCapacity = Chunks.Num() * NumEntitiesPerChunk;
+	Ar.Logf(ELogVerbosity::Log, TEXT("\tEntity Count    : %d"), EntityMap.Num());
+	Ar.Logf(ELogVerbosity::Log, TEXT("\tEntity Capacity : %d"), CurrentEntityCapacity);
+	if (Chunks.Num() > 1)
+	{
+		const float Scaler = 100.0f / (float)CurrentEntityCapacity;
+		// count non-last chunks to see how occupied they are
+		int EntitiesPerChunkMin = CurrentEntityCapacity;
+		int EntitiesPerChunkMax = 0;
+		for (int ChunkIndex = 0; ChunkIndex < Chunks.Num() - 1; ++ChunkIndex)
+		{
+			const int Population = Chunks[ChunkIndex].GetNumInstances();
+			EntitiesPerChunkMin = FMath::Min(Population, EntitiesPerChunkMin);
+			EntitiesPerChunkMax = FMath::Max(Population, EntitiesPerChunkMax);
+		}
+		Ar.Logf(ELogVerbosity::Log, TEXT("\tEntity Occupancy: %.1f%% (min: %.1f%%, max: %.1f%%)"), Scaler * EntityMap.Num(), Scaler * EntitiesPerChunkMin, Scaler * EntitiesPerChunkMax);
+	}
+	else 
+	{
+		Ar.Logf(ELogVerbosity::Log, TEXT("\tEntity Occupancy: %.1f%%"), CurrentEntityCapacity > 0 ? ((EntityMap.Num() * 100.0f) / (float)CurrentEntityCapacity) : 0.f);
+	}
+	Ar.Logf(ELogVerbosity::Log, TEXT("\tBytes / Entity  : %d"), TotalBytesPerEntity);
+	Ar.Logf(ELogVerbosity::Log, TEXT("\tEntities / Chunk: %d"), NumEntitiesPerChunk);
+
+	Ar.Logf(ELogVerbosity::Log, TEXT("\tOffset 0x%04X: Entity[] (%d bytes each)"), EntityListOffsetWithinChunk, sizeof(FLWEntity));
+	int32 TotalBytesOfValidData = sizeof(FLWEntity) * NumEntitiesPerChunk;
+	for (const FArchetypeComponentConfig& ComponentConfig : ComponentConfigs)
+	{
+		TotalBytesOfValidData += ComponentConfig.ComponentType->GetStructureSize() * NumEntitiesPerChunk;
+		Ar.Logf(ELogVerbosity::Log, TEXT("\tOffset 0x%04X: %s[] (%d bytes each)"), ComponentConfig.ArrayOffsetWithinChunk, *ComponentConfig.ComponentType->GetName(), ComponentConfig.ComponentType->GetStructureSize());
+	}
+
+	//@TODO: Print out padding in between things?
+
+	const int32 UnusuablePaddingOffset = TotalBytesPerEntity * NumEntitiesPerChunk;
+	const int32 UnusuablePaddingAmount = GetChunkAllocSize() - UnusuablePaddingOffset;
+	if (UnusuablePaddingAmount > 0)
+	{
+		Ar.Logf(ELogVerbosity::Log, TEXT("\tOffset 0x%04X: WastePadding[] (%d bytes total)"), UnusuablePaddingOffset, UnusuablePaddingAmount);
+	}
+
+	if (GetChunkAllocSize() != TotalBytesOfValidData + UnusuablePaddingAmount)
+	{
+		Ar.Logf(ELogVerbosity::Log, TEXT("\t@TODO: EXTRA PADDING HERE:  TotalBytesOfValidData: %d (%d missing)"), TotalBytesOfValidData, GetChunkAllocSize() - TotalBytesOfValidData);
+	}
+}
+
+void FArchetypeData::DebugPrintEntity(FLWEntity Entity, FOutputDevice& Ar, const TCHAR* InPrefix) const
+{
+	for (const FArchetypeComponentConfig& ComponentConfig : ComponentConfigs)
+	{
+		void* Data = GetComponentDataForEntityChecked(ComponentConfig.ComponentType, Entity.Index);
+		
+		FString ComponentName = ComponentConfig.ComponentType->GetName();
+		ComponentName.RemoveFromStart(InPrefix);
+
+		FString ValueStr;
+		ComponentConfig.ComponentType->ExportText(ValueStr, Data, /*Default*/nullptr, /*OwnerObject*/nullptr, EPropertyPortFlags::PPF_IncludeTransient, /*ExportRootScope*/nullptr);
+
+		Ar.Logf(TEXT("%s: %s"), *ComponentName, *ValueStr);
+	}
+}
+
+#endif // WITH_AGGREGATETICKING_DEBUG
+
+void FArchetypeData::REMOVEME_GetArrayViewForComponentInChunk(int32 ChunkIndex, const UScriptStruct* ComponentType, void*& OutChunkBase, int32& OutNumEntities)
+{
+	const FArchetypeChunk& Chunk = Chunks[ChunkIndex];
+	const int32 ComponentIndex = ComponentIndexMap.FindChecked(ComponentType);
+
+	OutChunkBase = ComponentConfigs[ComponentIndex].GetComponentData(Chunk.GetRawMemory(), 0);
+	OutNumEntities = Chunk.GetNumInstances();
+}
+
diff --git a/Engine/Plugins/Runtime/MassEntity/Source/MassEntity/Private/ArchetypeData.h b/Engine/Plugins/Runtime/MassEntity/Source/MassEntity/Private/ArchetypeData.h
new file mode 100644
--- /dev/null
+++ b/Engine/Plugins/Runtime/MassEntity/Source/MassEntity/Private/ArchetypeData.h
@@ -0,0 +1,281 @@
+// Copyright Epic Games, Inc. All Rights Reserved.
+
+#pragma once
+
+#include "EntitySubsystem.h"
+#include "ArchetypeTypes.h"
+
+struct FLWComponentQuery;
+struct FLWComponentSystemExecutionContext;
+class FOutputDevice;
+struct FArchetypeChunkCollection;
+
+// This is one chunk within an archetype
+struct FArchetypeChunk
+{
+private:
+	uint8* RawMemory = nullptr;
+	int32 NumInstances = 0;
+	int32 SerialModificationNumber = 0;
+	TArray<FInstancedStruct> ChunkComponentData;
+
+public:
+	explicit FArchetypeChunk(int32 AllocSize, TConstArrayView<FInstancedStruct> InChunkComponentTemplates)
+		: ChunkComponentData(InChunkComponentTemplates)
+	{
+		RawMemory = (uint8*)FMemory::Malloc(AllocSize);
+	}
+
+	~FArchetypeChunk()
+	{
+		FMemory::Free(RawMemory);
+	}
+
+	// Returns the Entity array element at the specified index
+	FLWEntity& GetEntityArrayElementRef(int32 ChunkBase, int32 IndexWithinChunk)
+	{
+		return ((FLWEntity*)(RawMemory + ChunkBase))[IndexWithinChunk];
+	}
+
+	uint8* GetRawMemory() const
+	{
+		return RawMemory;
+	}
+
+	int32 GetNumInstances() const
+	{
+		return NumInstances;
+	}
+
+	void AddMultipleInstances(uint32 Count)
+	{
+		NumInstances += Count;
+		SerialModificationNumber++;
+	}
+
+	void RemoveMultipleInstances(uint32 Count)
+	{
+		NumInstances -= Count;
+		check(NumInstances >= 0);
+		SerialModificationNumber++;
+	}
+
+	void AddInstance()
+	{
+		AddMultipleInstances(1);
+	}
+
+	void RemoveInstance()
+	{
+		RemoveMultipleInstances(1);
+	}
+
+	int32 GetSerialModificationNumber() const
+	{
+		return SerialModificationNumber;
+	}
+
+	FStructView GetMutableChunkComponentViewChecked(const int32 Index) { return FStructView(ChunkComponentData[Index]); }
+
+	FInstancedStruct* FindMutableChunkComponent(const UScriptStruct* Type)
+	{
+		return ChunkComponentData.FindByPredicate([Type](const FInstancedStruct& Element)
+			{
+				return Element.GetScriptStruct()->IsChildOf(Type);
+			});
+	}
+
+	void Recycle(TConstArrayView<FInstancedStruct> InChunkComponentsTemplate)
+	{
+		checkf(NumInstances == 0, TEXT("Recycling a chunk that is not empty."));
+		SerialModificationNumber++;
+		ChunkComponentData = InChunkComponentsTemplate;
+	}
+
+#if WITH_AGGREGATETICKING_DEBUG
+	int32 DebugGetChunkComponentCount() const { return ChunkComponentData.Num(); }
+#endif // WITH_AGGREGATETICKING_DEBUG
+};
+
+// Information for a single component type in an archetype
+struct FArchetypeComponentConfig
+{
+	const UScriptStruct* ComponentType = nullptr;
+	int32 ArrayOffsetWithinChunk = 0;
+
+	void* GetComponentData(uint8* ChunkBase, int32 IndexWithinChunk) const
+	{
+		return ChunkBase + ArrayOffsetWithinChunk + (IndexWithinChunk * ComponentType->GetStructureSize());
+	}
+};
+
+// An archetype is defined by a collection of unique component types (no duplicates).
+// Order doesn't matter, there will only ever be one FArchetypeData per unique set of component types per entity manager subsystem
+struct FArchetypeData
+{
+private:
+	// One-stop-shop variable describing the archetype's component and tag composition 
+	FLWCompositionDescriptor CompositionDescriptor;
+
+	TArray<FInstancedStruct> ChunkComponentTemplates;
+
+	TArray<FArchetypeComponentConfig, TInlineAllocator<16>> ComponentConfigs;
+	
+	TArray<FArchetypeChunk> Chunks;
+
+	// Entity ID to index within archetype
+	//@TODO: Could be folded into FEntityData in the entity manager at the expense of a bit
+	// of loss of encapsulation and extra complexity during archetype changes
+	TMap<int32, int32> EntityMap;
+	
+	TMap<const UScriptStruct*, int32> ComponentIndexMap;
+
+	int32 NumEntitiesPerChunk;
+	int32 TotalBytesPerEntity;
+	int32 EntityListOffsetWithinChunk;
+
+	friend FLWComponentQuery;
+	friend FArchetypeChunkCollection;
+
+public:
+	TConstArrayView<FArchetypeComponentConfig> GetComponentConfigs() const { return ComponentConfigs; }
+	const FLWComponentBitSet& GetComponentBitSet() const { return CompositionDescriptor.Components; }
+	const FLWTagBitSet& GetTagBitSet() const { return CompositionDescriptor.Tags; }
+	const FLWChunkComponentBitSet& GetChunkComponentBitSet() const { return CompositionDescriptor.ChunkComponents; }
+	const FLWCompositionDescriptor& GetCompositionDescriptor() const { return CompositionDescriptor; }
+
+	/** Method to iterate on all the component types */
+	void ForEachComponentType(TFunction< void(const UScriptStruct* /*ComponentType*/)> Function) const;
+	bool HasComponentType(const UScriptStruct* ComponentType) const;
+	bool HasTagType(const UScriptStruct* ComponentType) const { check(ComponentType); return CompositionDescriptor.Tags.Contains(*ComponentType); }
+	bool IsEquivalent(TConstArrayView<const UScriptStruct*> OtherList) const;
+	bool IsEquivalent(const FLWComponentBitSet& InComponentBitSet, const FLWTagBitSet& InTagBitSet, const FLWChunkComponentBitSet& InChunkComponentsBitSet) const
+	{ 
+		return CompositionDescriptor.IsEquivalent(InComponentBitSet, InTagBitSet, InChunkComponentsBitSet);
+	}
+	bool IsEquivalent(const FLWCompositionDescriptor& OtherCompositionDescriptor) const
+	{
+		return CompositionDescriptor.IsEquivalent(OtherCompositionDescriptor);
+	}
+
+	void Initialize(const FLWComponentBitSet& Components, const FLWTagBitSet& Tags, const FLWChunkComponentBitSet& ChunkComponents);
+
+	/** 
+	 * A special way of initializing an archetype resulting in a copy of SiblingArchetype's setup with OverrideTags
+	 * replacing original tags of SiblingArchetype
+	 */
+	void InitializeWithSibling(const FArchetypeData& SiblingArchetype, const FLWTagBitSet& OverrideTags);
+
+	void AddEntity(FLWEntity Entity);
+	void RemoveEntity(FLWEntity Entity);
+	void BatchDestroyEntityChunks(const FArchetypeChunkCollection& ChunkCollection, TArray<FLWEntity>& OutEntitiesRemoved);
+
+	bool HasComponentDataForEntity(const UScriptStruct* ComponentType, int32 EntityIndex) const;
+	void* GetComponentDataForEntityChecked(const UScriptStruct* ComponentType, int32 EntityIndex) const;
+	void* GetComponentDataForEntity(const UScriptStruct* ComponentType, int32 EntityIndex) const;
+
+	FORCEINLINE int32 GetInternalIndexForEntity(const int32 EntityIndex) const { return EntityMap.FindChecked(EntityIndex); }
+	int32 GetNumEntitiesPerChunk() const { return NumEntitiesPerChunk; }
+
+	int32 GetNumEntities() const { return EntityMap.Num(); }
+
+	int32 GetChunkAllocSize() const { return 64*1024; }
+
+	int32 GetChunkCount() const { return Chunks.Num(); }
+
+	void ExecuteFunction(FLWComponentSystemExecutionContext& RunContext, const FLWComponentSystemExecuteFunction& Function, const FLWRequirementIndicesMapping& RequirementMapping, const FArchetypeChunkCollection& ChunkCollection);
+	void ExecuteFunction(FLWComponentSystemExecutionContext& RunContext, const FLWComponentSystemExecuteFunction& Function, const FLWRequirementIndicesMapping& RequirementMapping, const FLWComponentSystemChunkConditionFunction& ChunkCondition = FLWComponentSystemChunkConditionFunction());
+
+	void ExecutionFunctionForChunk(FLWComponentSystemExecutionContext RunContext, const FLWComponentSystemExecuteFunction& Function, const FLWRequirementIndicesMapping& RequirementMapping, const FArchetypeChunkCollection::FChunkInfo& ChunkInfo, const FLWComponentSystemChunkConditionFunction& ChunkCondition = FLWComponentSystemChunkConditionFunction());
+
+	/**
+	 * Compacts entities to fill up chunks as much as possible
+	 */
+	void CompactEntities(const double TimeAllowed);
+
+	/**
+	 * Moves the entity from this archetype to another, will only copy all matching component types
+	 * @param Entity is the entity to move
+	 * @param NewArchetype the archetype to move to
+	 */
+	void MoveEntityToAnotherArchetype(const FLWEntity Entity, FArchetypeData& NewArchetype);
+
+	/**
+	 * Set all component sources data on specified entity, will check if there are component sources type that does not exist in the archetype
+	 * @param Entity is the entity to set the data of all components
+	 * @param ComponentSources are the components to copy the data from
+	 */
+	void SetComponentsData(const FLWEntity Entity, TArrayView<const FInstancedStruct> ComponentSources);
+
+	/** For all entities indicated by ChunkCollection the function sets the value of component of type
+	 *  ComponentSource.GetScriptStruct to the value represented by ComponentSource.GetMemory */
+	void SetComponentData(const FArchetypeChunkCollection& ChunkCollection, const FInstancedStruct& ComponentSource);
+
+	void SetDefaultChunkComponentValue(FConstStructView InstancedStruct);
+
+	/** Returns conversion from given Requirements to archetype's component indices */
+	void GetRequirementsComponentMapping(TConstArrayView<FLWComponentRequirement> Requirements, FLWComponentIndicesMapping& OutComponentIndices);
+
+	/** Returns conversion from given ChunkRequirements to archetype's chunk component indices */
+	void GetRequirementsChunkComponentMapping(TConstArrayView<FLWComponentRequirement> ChunkRequirements, FLWComponentIndicesMapping& OutComponentIndices);
+
+	SIZE_T GetAllocatedSize() const;
+
+	// Converts the list of components into a user-readable debug string
+	FString DebugGetDescription() const;
+
+#if WITH_AGGREGATETICKING_DEBUG
+	/**
+	 * Prints out debug information about the archetype
+	 */
+	void DebugPrintArchetype(FOutputDevice& Ar);
+
+	/**
+	 * Prints out component's values for the specified entity. 
+	 * @param Entity The entity for which we want to print component values
+	 * @param Ar The output device
+	 * @param InPrefix Optional prefix to remove from component names
+	 */
+	void DebugPrintEntity(FLWEntity Entity, FOutputDevice& Ar, const TCHAR* InPrefix = TEXT("")) const;
+#endif // WITH_AGGREGATETICKING_DEBUG
+
+	void REMOVEME_GetArrayViewForComponentInChunk(int32 ChunkIndex, const UScriptStruct* ComponentType, void*& OutChunkBase, int32& OutNumEntities);
+
+	//////////////////////////////////////////////////////////////////////
+	// low level api
+	FORCEINLINE const int32* GetComponentIndex(const UScriptStruct* ComponentType) const { return ComponentIndexMap.Find(ComponentType); }
+	FORCEINLINE int32 GetComponentIndexChecked(const UScriptStruct* ComponentType) const { return ComponentIndexMap.FindChecked(ComponentType); }
+
+	FORCEINLINE void* GetComponentData(const int32 ComponentIndex, const FInternalEntityHandle EntityIndex) const
+	{
+		return ComponentConfigs[ComponentIndex].GetComponentData(EntityIndex.ChunkRawMemory, EntityIndex.IndexWithinChunk);
+	}
+
+	FORCEINLINE FInternalEntityHandle MakeEntityHandle(int32 EntityIndex) const
+	{
+		const int32 AbsoluteIndex = EntityMap.FindChecked(EntityIndex);
+		const int32 ChunkIndex = AbsoluteIndex / NumEntitiesPerChunk;
+	
+		return FInternalEntityHandle(Chunks[ChunkIndex].GetRawMemory(), AbsoluteIndex % NumEntitiesPerChunk); 
+	}
+
+	FORCEINLINE FInternalEntityHandle MakeEntityHandle(FLWEntity Entity) const
+	{
+		return MakeEntityHandle(Entity.Index); 
+	}
+
+	bool IsInitialized() const { return TotalBytesPerEntity > 0 && ComponentConfigs.IsEmpty() == false; }
+
+protected:
+	FORCEINLINE void* GetComponentData(const int32 ComponentIndex, uint8* ChunkRawMemory, const int32 IndexWithinChunk) const
+	{
+		return ComponentConfigs[ComponentIndex].GetComponentData(ChunkRawMemory, IndexWithinChunk);
+	}
+
+	void BindEntityRequirements(FLWComponentSystemExecutionContext& RunContext, const FLWComponentIndicesMapping& EntityComponentsMapping, FArchetypeChunk& Chunk, const int32 SubchunkStart, const int32 SubchunkLength);
+	void BindChunkComponentRequirements(FLWComponentSystemExecutionContext& RunContext, const FLWComponentIndicesMapping& ChunkComponentsMapping, FArchetypeChunk& Chunk);
+
+private:
+	int32 AddEntityInternal(FLWEntity Entity, const bool bInitializeComponents);
+	void RemoveEntityInternal(const int32 AbsoluteIndex, const bool bDestroyComponents);
+};
diff --git a/Engine/Plugins/Runtime/MassEntity/Source/MassEntity/Private/ArchetypeTypes.cpp b/Engine/Plugins/Runtime/MassEntity/Source/MassEntity/Private/ArchetypeTypes.cpp
new file mode 100644
--- /dev/null
+++ b/Engine/Plugins/Runtime/MassEntity/Source/MassEntity/Private/ArchetypeTypes.cpp
@@ -0,0 +1,92 @@
+// Copyright Epic Games, Inc. All Rights Reserved.
+
+#include "ArchetypeTypes.h"
+#include "EntitySubsystem.h"
+#include "ArchetypeData.h"
+#include "LWCCommandBuffer.h"
+
+
+//////////////////////////////////////////////////////////////////////
+// FArchetypeHandle
+bool FArchetypeHandle::operator==(const FArchetypeData* Other) const
+{
+	return DataPtr.Get() == Other;
+}
+
+uint32 GetTypeHash(const FArchetypeHandle& Instance)
+{
+	return GetTypeHash(Instance.DataPtr.Get());
+}
+
+//////////////////////////////////////////////////////////////////////
+// FArchetypeChunkCollection
+
+FArchetypeChunkCollection::FArchetypeChunkCollection(const FArchetypeHandle& InArchetype, TConstArrayView<FLWEntity> InEntities)
+	: Archetype(InArchetype)
+{
+	check(InArchetype.IsValid());
+	const FArchetypeData& ArchetypeData = *InArchetype.DataPtr.Get();
+	const int32 NumEntitiesPerChunk = ArchetypeData.GetNumEntitiesPerChunk();
+
+	// InEntities has a real chance of not being sorted by AbsoluteIndex. We gotta fix that to optimize how we process the data 
+	TArray<int32> TrueIndices;
+	TrueIndices.AddUninitialized(InEntities.Num());
+	int32 i = 0;
+	for (const FLWEntity& Entity : InEntities)
+	{
+		TrueIndices[i] = ArchetypeData.GetInternalIndexForEntity(Entity.Index);
+		++i;
+	}
+	TrueIndices.Sort();
+
+	// the following block of code is splitting up sorted AbsoluteIndices into 
+	// continuous chunks
+	int32 ChunkEnd = INDEX_NONE;
+	FChunkInfo DummyChunk;
+	FChunkInfo* SubChunkPtr = &DummyChunk;
+	int32 SubchunkLen = 0;
+	int32 PrevAbsoluteIndex = INDEX_NONE;
+	for (const int32 Index : TrueIndices)
+	{
+		// if run across a chunk border or run into an index discontinuity 
+		if (Index >= ChunkEnd || Index != (PrevAbsoluteIndex + 1))
+		{
+			SubChunkPtr->Length = SubchunkLen;
+			// note that both ChunkIndex and ChunkEnd will change only if AbsoluteIndex >= ChunkEnd
+			const int32 ChunkIndex = Index / NumEntitiesPerChunk;
+			ChunkEnd = (ChunkIndex + 1) * NumEntitiesPerChunk;
+			SubchunkLen = 0;
+			// new subchunk
+			const int32 SubchunkStart = Index % NumEntitiesPerChunk;
+			SubChunkPtr = &Chunks.Add_GetRef(FChunkInfo(ChunkIndex, SubchunkStart));
+		}
+		++SubchunkLen;
+		PrevAbsoluteIndex = Index;
+	}
+
+	SubChunkPtr->Length = SubchunkLen;
+}
+
+FArchetypeChunkCollection::FArchetypeChunkCollection(FArchetypeHandle& InArchetypeHandle)
+{
+	check(InArchetypeHandle.DataPtr.IsValid());
+	GatherChunksFromArchetype(InArchetypeHandle.DataPtr);
+}
+
+FArchetypeChunkCollection::FArchetypeChunkCollection(TSharedPtr<FArchetypeData>& InArchetype)
+{
+	GatherChunksFromArchetype(InArchetype);
+}
+
+void FArchetypeChunkCollection::GatherChunksFromArchetype(TSharedPtr<FArchetypeData>& InArchetype)
+{
+	check(InArchetype.IsValid());
+	Archetype.DataPtr = InArchetype;
+
+	const int32 ChunkCount = InArchetype->GetChunkCount();
+	Chunks.Reset(ChunkCount);
+	for (int32 i = 0; i < ChunkCount; ++i)
+	{
+		Chunks.Add(FChunkInfo(i));
+	}
+}
diff --git a/Engine/Plugins/Runtime/MassEntity/Source/MassEntity/Private/EntityQuery.cpp b/Engine/Plugins/Runtime/MassEntity/Source/MassEntity/Private/EntityQuery.cpp
new file mode 100644
--- /dev/null
+++ b/Engine/Plugins/Runtime/MassEntity/Source/MassEntity/Private/EntityQuery.cpp
@@ -0,0 +1,308 @@
+// Copyright Epic Games, Inc. All Rights Reserved.
+
+#include "EntityQuery.h"
+#include "EntitySubsystem.h"
+#include "ArchetypeData.h"
+#include "LWCCommandBuffer.h"
+#include "VisualLogger/VisualLogger.h"
+#include "Async/ParallelFor.h"
+
+//////////////////////////////////////////////////////////////////////
+// FLWComponentQuery
+
+
+FLWComponentQuery::FLWComponentQuery()
+{
+	ReadCommandlineParams();
+}
+
+FLWComponentQuery::FLWComponentQuery(std::initializer_list<UScriptStruct*> InitList)
+	: FLWComponentQuery()
+{
+	for (const UScriptStruct* ComponentType : InitList)
+	{
+		AddRequirement(ComponentType, ELWComponentAccess::ReadWrite, ELWComponentPresence::All);
+	}
+}
+
+FLWComponentQuery::FLWComponentQuery(TConstArrayView<const UScriptStruct*> InitList)
+	: FLWComponentQuery()
+{
+	for (const UScriptStruct* ComponentType : InitList)
+	{
+		AddRequirement(ComponentType, ELWComponentAccess::ReadWrite, ELWComponentPresence::All);
+	}
+}
+
+void FLWComponentQuery::ReadCommandlineParams()
+{
+	int AllowParallelQueries = -1;
+	if (FParse::Value(FCommandLine::Get(), TEXT("ParallelMassQueries="), AllowParallelQueries))
+	{
+		bAllowParallelExecution = (AllowParallelQueries != 0);
+	}
+}
+
+void FLWComponentQuery::SortRequirements()
+{
+	// we're sorting the Requirements the same way ArchetypeData's ComponentConfig is sorted (see FArchetypeData::Initialize)
+	// so that when we access ArchetypeData.ComponentConfigs in FArchetypeData::BindRequirementsWithMapping
+	// (via GetComponentData call) the access is sequential (i.e. not random) and there's a higher chance the memory
+	// ComponentConfigs we want to access have already been fetched and are available in processor cache.
+	Requirements.Sort(FLWComponentSorterOperator<FLWComponentRequirement>());
+	ChunkRequirements.Sort(FLWComponentSorterOperator<FLWComponentRequirement>());
+}
+
+void FLWComponentQuery::CacheArchetypes(UEntitySubsystem& InEntitySubsystem)
+{
+	const uint32 InEntitySubsystemHash = PointerHash(&InEntitySubsystem);
+	if (EntitySubsystemHash != InEntitySubsystemHash || InEntitySubsystem.GetArchetypeDataVersion() != ArchetypeDataVersion)
+	{
+		if (CheckValidity())
+		{
+			SortRequirements();
+
+			EntitySubsystemHash = InEntitySubsystemHash;
+			ValidArchetypes.Reset();
+			InEntitySubsystem.GetValidArchetypes(*this, ValidArchetypes);
+			ArchetypeDataVersion = InEntitySubsystem.GetArchetypeDataVersion();
+
+			TRACE_CPUPROFILER_EVENT_SCOPE_STR("Pipe RequirementsBinding")
+			const TConstArrayView<FLWComponentRequirement> LocalRequirements = GetRequirements();
+			ArchetypeComponentMapping.Reset(ValidArchetypes.Num());
+			ArchetypeComponentMapping.AddDefaulted(ValidArchetypes.Num());
+			for (int i = 0; i < ValidArchetypes.Num(); ++i)
+			{
+				ValidArchetypes[i].DataPtr->GetRequirementsComponentMapping(LocalRequirements, ArchetypeComponentMapping[i].EntityComponents);
+				if (ChunkRequirements.Num())
+				{
+					ValidArchetypes[i].DataPtr->GetRequirementsChunkComponentMapping(ChunkRequirements, ArchetypeComponentMapping[i].ChunkComponents);
+				}
+			}
+		}
+		else
+		{
+			UE_VLOG_UELOG(&InEntitySubsystem, LogAggregateTicking, Error, TEXT("FLWComponentQuery::CacheArchetypes: requirements not valid: %s"), *DebugGetDescription());
+		}
+	}
+}
+
+bool FLWComponentQuery::CheckValidity() const
+{
+	return RequiredAllComponents.IsEmpty() == false || RequiredAnyComponents.IsEmpty() == false || RequiredOptionalComponents.IsEmpty() == false;
+}
+
+bool FLWComponentQuery::DoesArchetypeMatchRequirements(const FArchetypeHandle& ArchetypeHandle) const
+{
+	check(ArchetypeHandle.IsValid());
+	const FArchetypeData* Archetype = ArchetypeHandle.DataPtr.Get();
+	CA_ASSUME(Archetype);
+
+	for (const FLWComponentRequirement& Requirement : GetRequirements())
+	{
+		check(Requirement.StructType);
+		if (Requirement.IsOptional())
+		{
+			// at this stage we don't care if the optional is present
+			continue;
+		}
+
+		const bool bNeedsComponent = (Requirement.Presence == ELWComponentPresence::All);
+		const bool bHasComponent = Archetype->HasComponentType(Requirement.StructType);
+		if (bNeedsComponent != bHasComponent)
+		{
+			return false;
+		}
+	}
+	return true;
+}
+
+void FLWComponentQuery::ForEachEntityChunk(const FArchetypeChunkCollection& Chunks, UEntitySubsystem& EntitySubsystem, FLWComponentSystemExecutionContext& ExecutionContext, const FLWComponentSystemExecuteFunction& ExecuteFunction)
+{
+	// mz@todo I don't like that we're copying data here.
+	ExecutionContext.SetChunkCollection(Chunks);
+	ForEachEntityChunk(EntitySubsystem, ExecutionContext, ExecuteFunction);
+	ExecutionContext.ClearChunkCollection();
+}
+
+void FLWComponentQuery::ForEachEntityChunk(UEntitySubsystem& EntitySubsystem, FLWComponentSystemExecutionContext& ExecutionContext, const FLWComponentSystemExecuteFunction& ExecuteFunction)
+{
+#if WITH_AGGREGATETICKING_DEBUG
+	int32 NumEntitiesToProcess = 0;
+#endif
+
+	// if there's a chunk collection set by the external code - use that
+	if (ExecutionContext.GetChunkCollection().IsSet())
+	{
+		// verify the archetype matches requirements
+		if (DoesArchetypeMatchRequirements(ExecutionContext.GetChunkCollection().GetArchetype()) == false)
+		{
+			// mz@todo add a unit test for this message
+			UE_VLOG_UELOG(&EntitySubsystem, LogAggregateTicking, Error, TEXT("Attempted to execute FLWComponentQuery with an incompatible Archetype"));
+			return;
+		}
+		ExecutionContext.SetRequirements(Requirements, ChunkRequirements);
+		ExecutionContext.GetChunkCollection().GetArchetype().DataPtr->ExecuteFunction(ExecutionContext, ExecuteFunction, {}, ExecutionContext.GetChunkCollection());
+#if WITH_AGGREGATETICKING_DEBUG
+		NumEntitiesToProcess = ExecutionContext.GetChunkCollection().GetArchetype().DataPtr->GetNumEntities();
+#endif
+	}
+	else
+	{
+		CacheArchetypes(EntitySubsystem);
+		// it's important to set requirements after caching archetypes due to that call potentially sorting the requirements and the order is relevant here.
+		ExecutionContext.SetRequirements(Requirements, ChunkRequirements);
+
+		for (int i = 0; i < ValidArchetypes.Num(); ++i)
+		{
+			FArchetypeHandle& Archetype = ValidArchetypes[i];
+			check(Archetype.IsValid());
+			Archetype.DataPtr->ExecuteFunction(ExecutionContext, ExecuteFunction, ArchetypeComponentMapping[i], ChunkCondition);
+			ExecutionContext.ClearComponentViews();
+#if WITH_AGGREGATETICKING_DEBUG
+			NumEntitiesToProcess += Archetype.DataPtr->GetNumEntities();
+#endif
+		}
+	}
+
+#if WITH_AGGREGATETICKING_DEBUG
+	// Not using VLOG to be thread safe
+	UE_CLOG(!ExecutionContext.DebugGetExecutionDesc().IsEmpty(), LogAggregateTicking, VeryVerbose,
+		TEXT("%s: %d entities sent for processing"), *ExecutionContext.DebugGetExecutionDesc(), NumEntitiesToProcess);
+#endif
+
+	ExecutionContext.ClearExecutionData();
+	ExecutionContext.FlushDeferred(EntitySubsystem);
+}
+
+void FLWComponentQuery::ParallelForEachEntityChunk(UEntitySubsystem& EntitySubsystem, FLWComponentSystemExecutionContext& ExecutionContext, const FLWComponentSystemExecuteFunction& ExecuteFunction)
+{
+	if (bAllowParallelExecution == false)
+	{
+		ForEachEntityChunk(EntitySubsystem, ExecutionContext, ExecuteFunction);
+		return;
+	}
+
+	struct FChunkJob
+	{
+		FArchetypeData& Archetype;
+		const int32 ArchetypeIndex;
+		const FArchetypeChunkCollection::FChunkInfo ChunkInfo;
+	};
+	TArray<FChunkJob> Jobs;
+
+	// if there's a chunk collection set by the external code - use that
+	if (ExecutionContext.GetChunkCollection().IsSet())
+	{
+		// verify the archetype matches requirements
+		FArchetypeHandle ArchetypeHandle = ExecutionContext.GetChunkCollection().GetArchetype();
+		if (DoesArchetypeMatchRequirements(ArchetypeHandle) == false)
+		{
+			// mz@todo add a unit test for this message
+			UE_VLOG_UELOG(&EntitySubsystem, LogAggregateTicking, Error, TEXT("Attempted to execute FLWComponentQuery with an incompatible Archetype"));
+			return;
+		}
+
+		ExecutionContext.SetRequirements(Requirements, ChunkRequirements);
+		check(ArchetypeHandle.IsValid());
+		FArchetypeData& ArchetypeRef = *ArchetypeHandle.DataPtr.Get();
+		const FArchetypeChunkCollection AsChunkCollection(ArchetypeHandle.DataPtr);
+		for (const FArchetypeChunkCollection::FChunkInfo& ChunkInfo : AsChunkCollection.GetChunks())
+		{
+			Jobs.Add({ ArchetypeRef, INDEX_NONE, ChunkInfo });
+		}
+	}
+	else
+	{
+		CacheArchetypes(EntitySubsystem);
+		ExecutionContext.SetRequirements(Requirements, ChunkRequirements);
+		for (int ArchetypeIndex = 0; ArchetypeIndex < ValidArchetypes.Num(); ++ArchetypeIndex)
+		{
+			FArchetypeHandle& Archetype = ValidArchetypes[ArchetypeIndex];
+			check(Archetype.IsValid());
+			FArchetypeData& ArchetypeRef = *Archetype.DataPtr.Get();
+			const FArchetypeChunkCollection AsChunkCollection(Archetype.DataPtr);
+			for (const FArchetypeChunkCollection::FChunkInfo& ChunkInfo : AsChunkCollection.GetChunks())
+			{
+				Jobs.Add({ArchetypeRef, ArchetypeIndex, ChunkInfo});
+			}
+		}
+	}
+	// ExecutionContext passed by copy on purpose
+	ParallelFor(Jobs.Num(), [this, ExecutionContext, &ExecuteFunction, Jobs](const int32 JobIndex)
+	{
+		Jobs[JobIndex].Archetype.ExecutionFunctionForChunk(ExecutionContext, ExecuteFunction
+			, Jobs[JobIndex].ArchetypeIndex != INDEX_NONE ? ArchetypeComponentMapping[Jobs[JobIndex].ArchetypeIndex] : FLWRequirementIndicesMapping()
+			, Jobs[JobIndex].ChunkInfo
+			, ChunkCondition);
+	});
+
+	ExecutionContext.ClearExecutionData();
+	ExecutionContext.FlushDeferred(EntitySubsystem);
+}
+
+int32 FLWComponentQuery::GetMatchingEntitiesNum(UEntitySubsystem& InEntitySubsystem)
+{
+	CacheArchetypes(InEntitySubsystem);
+	int32 TotalEntities = 0;
+	for (FArchetypeHandle& ArchetypeHandle : ValidArchetypes)
+	{
+		if (const FArchetypeData* Archetype = ArchetypeHandle.DataPtr.Get())
+		{
+			TotalEntities += Archetype->GetNumEntities();
+		}
+	}
+	return TotalEntities;
+}
+
+bool FLWComponentQuery::HasMatchingEntities(UEntitySubsystem& InEntitySubsystem)
+{
+	CacheArchetypes(InEntitySubsystem);
+
+	for (FArchetypeHandle& ArchetypeHandle : ValidArchetypes)
+	{
+		const FArchetypeData* Archetype = ArchetypeHandle.DataPtr.Get();
+		if (Archetype && Archetype->GetNumEntities() > 0)
+		{
+			return true;
+		}
+	}
+	return false;
+}
+
+FString FLWComponentQuery::DebugGetDescription() const
+{
+#if WITH_AGGREGATETICKING_DEBUG
+	TStringBuilder<256> StringBuilder;
+	StringBuilder.Append(TEXT("<"));
+
+	bool bNeedsComma = false;
+	for (const FLWComponentRequirement& Requirement : Requirements)
+	{
+		if (bNeedsComma)
+		{
+			StringBuilder.Append(TEXT(","));
+		}
+		StringBuilder.Append(*Requirement.DebugGetDescription());
+		bNeedsComma = true;
+	}
+
+	StringBuilder.Append(TEXT(">"));
+	return StringBuilder.ToString();
+#else
+	return {};
+#endif
+}
+
+//////////////////////////////////////////////////////////////////////
+// FLWComponentRequirement
+
+FString FLWComponentRequirement::DebugGetDescription() const
+{
+#if WITH_AGGREGATETICKING_DEBUG
+	return FString::Printf(TEXT("%s%s[%s]"), IsOptional() ? TEXT("?") : (Presence == ELWComponentPresence::None ? TEXT("-") : TEXT("+"))
+		, *GetNameSafe(StructType), *UE::AggregateTicking::DebugGetComponentAccessString(AccessMode));
+#else
+	return {};
+#endif
+}
