AggregateTicking and Pipe merged into Engine/Plugins/Runtime/MassEntity

diff --git a/Engine/Plugins/Runtime/MassEntity/Source/MassEntity/Private/EntitySubsystem.cpp b/Engine/Plugins/Runtime/MassEntity/Source/MassEntity/Private/EntitySubsystem.cpp
new file mode 100644
--- /dev/null
+++ b/Engine/Plugins/Runtime/MassEntity/Source/MassEntity/Private/EntitySubsystem.cpp
@@ -0,0 +1,1105 @@
+// Copyright Epic Games, Inc. All Rights Reserved.
+
+#include "EntitySubsystem.h"
+#include "ArchetypeData.h"
+#include "LWCCommandBuffer.h"
+#include "HAL/IConsoleManager.h"
+#include "Engine/World.h"
+#include "UObject/UObjectIterator.h"
+
+const FLWEntity UEntitySubsystem::InvalidEntity;
+
+DEFINE_LOG_CATEGORY(LogAggregateTicking);
+
+namespace UE { namespace AggregateTicking {
+
+FString DebugGetComponentAccessString(ELWComponentAccess Access)
+{
+#if WITH_AGGREGATETICKING_DEBUG
+	switch (Access)
+	{
+	case ELWComponentAccess::None:	return TEXT("--");
+	case ELWComponentAccess::ReadOnly:	return TEXT("RO");
+	case ELWComponentAccess::ReadWrite:	return TEXT("RW");
+	default:
+		ensureMsgf(false, TEXT("Missing string conversion for ELWComponentAccess=%d"), Access);
+		break;
+	}
+#endif // WITH_AGGREGATETICKING_DEBUG
+	return TEXT("Missing string conversion");
+}
+
+}} // UE::AggregateTicking
+
+//@TODO: Everything still alive leaks at shutdown
+//@TODO: No re-entrance safety while running a system (e.g., preventing someone from adding/removing entities or altering archetypes, etc...)
+//@TODO: Do we allow GCable types?  If so, need to implement AddReferencedObjects
+//@TODO: Do a UObject stat scope around the system Run call
+//@TODO: Tag components end up taking up space per-entity when they shouldn't
+
+//////////////////////////////////////////////////////////////////////
+// UEntitySubsystem
+
+UEntitySubsystem::UEntitySubsystem()
+{
+}
+
+void UEntitySubsystem::GetResourceSizeEx(FResourceSizeEx& CumulativeResourceSize)
+{
+	Super::GetResourceSizeEx(CumulativeResourceSize);
+
+	const SIZE_T MyExtraSize = Entities.GetAllocatedSize() + 
+		EntityFreeIndexList.GetAllocatedSize() +
+		DeferredCommandBuffer->GetAllocatedSize() +
+		ComponentHashToArchetypeMap.GetAllocatedSize() +
+		ComponentTypeToArchetypeMap.GetAllocatedSize();
+	CumulativeResourceSize.AddDedicatedSystemMemoryBytes(MyExtraSize);
+
+	for (const auto& KVP : ComponentHashToArchetypeMap)
+	{
+		for (const TSharedPtr<FArchetypeData>& ArchetypePtr : KVP.Value)
+		{
+			CumulativeResourceSize.AddDedicatedSystemMemoryBytes(ArchetypePtr->GetAllocatedSize());
+		}
+	}
+}
+
+void UEntitySubsystem::Initialize(FSubsystemCollectionBase& Collection)
+{
+	// Index 0 is reserved so we can treat that index as an invalid entity handle
+	Entities.Add();
+	SerialNumberGenerator.fetch_add(FMath::Max(1,NumReservedEntities));
+
+	DeferredCommandBuffer = MakeShareable(new FLWCCommandBuffer());
+
+	// initialize the bitsets
+
+	FLWComponentBitSet Components;
+	FLWTagBitSet Tags;
+	FLWChunkComponentBitSet ChunkComponents;
+
+	for (TObjectIterator<UStruct> StructIt; StructIt; ++StructIt)
+	{
+		if (StructIt->IsChildOf(FLWComponentData::StaticStruct()))
+		{
+			Components.Add(*CastChecked<UScriptStruct>(*StructIt));
+		}
+		else if (StructIt->IsChildOf(FComponentTag::StaticStruct()))
+		{
+			Tags.Add(*CastChecked<UScriptStruct>(*StructIt));
+		}
+		else if (StructIt->IsChildOf(FLWChunkComponent::StaticStruct()))
+		{
+			ChunkComponents.Add(*CastChecked<UScriptStruct>(*StructIt));
+		}
+	}
+}
+
+void UEntitySubsystem::Deinitialize()
+{
+	//@TODO: Should actually do this...
+}
+
+FArchetypeHandle UEntitySubsystem::CreateArchetype(TConstArrayView<const UScriptStruct*> ComponentsAngTagsList)
+{
+	FLWChunkComponentBitSet ChunkComponents;
+	FLWTagBitSet Tags;
+	TArray<const UScriptStruct*, TInlineAllocator<16>> ComponentList;
+	ComponentList.Reserve(ComponentsAngTagsList.Num());
+
+	for (const UScriptStruct* Type : ComponentsAngTagsList)
+	{
+		if (Type->IsChildOf(FLWComponentData::StaticStruct()))
+		{
+			ComponentList.Add(Type);
+		}
+		else if (Type->IsChildOf(FComponentTag::StaticStruct()))
+		{
+			Tags.Add(*Type);
+		}
+		else if (Type->IsChildOf(FLWChunkComponent::StaticStruct()))
+		{
+			ChunkComponents.Add(*Type);
+		}
+		else
+		{
+			UE_LOG(LogAggregateTicking, Warning, TEXT("%s: %s is not a valid component nor tag type. Ignoring.")
+				, ANSI_TO_TCHAR(__FUNCTION__), *GetNameSafe(Type));
+		}
+	}
+
+	return CreateArchetype(FLWCompositionDescriptor(FLWComponentBitSet(ComponentList), Tags, ChunkComponents));
+}
+
+FArchetypeHandle UEntitySubsystem::CreateArchetype(TConstArrayView<const UScriptStruct*> ComponentList, const FLWTagBitSet& Tags, const FLWChunkComponentBitSet& ChunkComponents)
+{
+	return CreateArchetype(FLWCompositionDescriptor(ComponentList, Tags, ChunkComponents));
+}
+
+FArchetypeHandle UEntitySubsystem::CreateArchetype(const TSharedPtr<FArchetypeData>& SourceArchetype, const FLWComponentBitSet& NewComponents)
+{
+	check(SourceArchetype.IsValid());
+	checkf(NewComponents.IsEmpty() == false, TEXT("%s Adding an empty component list to an archetype is not supported."), ANSI_TO_TCHAR(__FUNCTION__));
+
+	return CreateArchetype(FLWCompositionDescriptor(NewComponents + SourceArchetype->GetComponentBitSet(), SourceArchetype->GetTagBitSet(), SourceArchetype->GetChunkComponentBitSet()));
+}
+
+FArchetypeHandle UEntitySubsystem::CreateArchetype(const FLWCompositionDescriptor& Composition)
+{
+	const uint32 TypeHash = Composition.CalculateHash();
+
+	TArray<TSharedPtr<FArchetypeData>>& HashRow = ComponentHashToArchetypeMap.FindOrAdd(TypeHash);
+
+	FArchetypeHandle Result;
+	for (TSharedPtr<FArchetypeData>& Ptr : HashRow)
+	{
+		if (Ptr->IsEquivalent(Composition))
+		{
+			Result.DataPtr = Ptr;
+			break;
+		}
+	}
+
+	if (!Result.DataPtr.IsValid())
+	{
+		// Create a new archetype
+		TSharedPtr<FArchetypeData> NewArchetype = MakeShareable(new FArchetypeData);
+		NewArchetype->Initialize(Composition.Components, Composition.Tags, Composition.ChunkComponents);
+		HashRow.Add(NewArchetype);
+
+		for (const FArchetypeComponentConfig& ComponentConfig : NewArchetype->GetComponentConfigs())
+		{
+			checkSlow(ComponentConfig.ComponentType)
+			ComponentTypeToArchetypeMap.FindOrAdd(ComponentConfig.ComponentType).Add(NewArchetype);
+		}
+
+		Result.DataPtr = NewArchetype;
+
+		++ArchetypeDataVersion;
+	}
+
+	return MoveTemp(Result);
+}
+
+FArchetypeHandle UEntitySubsystem::InternalCreateSiblingArchetype(const TSharedPtr<FArchetypeData>& SourceArchetype, const FLWTagBitSet& OverrideTags)
+{
+	checkSlow(SourceArchetype.IsValid());
+	const FArchetypeData& SourceArchetypeRef = *SourceArchetype.Get();
+	const uint32 TypeHash = FLWCompositionDescriptor::CalculateHash(SourceArchetypeRef.GetComponentBitSet(), OverrideTags, SourceArchetypeRef.GetChunkComponentBitSet());
+
+	TArray<TSharedPtr<FArchetypeData>>& HashRow = ComponentHashToArchetypeMap.FindOrAdd(TypeHash);
+
+	FArchetypeHandle Result;
+	for (TSharedPtr<FArchetypeData>& Ptr : HashRow)
+	{
+		if (Ptr->IsEquivalent(SourceArchetypeRef.GetComponentBitSet(), OverrideTags, SourceArchetypeRef.GetChunkComponentBitSet()))
+		{
+			Result.DataPtr = Ptr;
+			break;
+		}
+	}
+
+	if (!Result.DataPtr.IsValid())
+	{
+		// Create a new archetype
+		TSharedPtr<FArchetypeData> NewArchetype = MakeShareable(new FArchetypeData);
+		NewArchetype->InitializeWithSibling(SourceArchetypeRef, OverrideTags);
+		HashRow.Add(NewArchetype);
+
+		for (const FArchetypeComponentConfig& ComponentConfig : NewArchetype->GetComponentConfigs())
+		{
+			checkSlow(ComponentConfig.ComponentType)
+			ComponentTypeToArchetypeMap.FindOrAdd(ComponentConfig.ComponentType).Add(NewArchetype);
+		}
+
+		Result.DataPtr = NewArchetype;
+
+		++ArchetypeDataVersion;
+	}
+
+	return MoveTemp(Result);
+}
+
+FArchetypeHandle UEntitySubsystem::GetArchetypeForEntity(FLWEntity Entity) const
+{
+	FArchetypeHandle Result;
+	if (IsEntityValid(Entity))
+	{
+		Result.DataPtr = Entities[Entity.Index].CurrentArchetype;
+	}
+	return Result;
+}
+
+void UEntitySubsystem::ForEachArchetypeComponentType(const FArchetypeHandle Archetype, TFunction< void(const UScriptStruct* /*ComponentType*/)> Function)
+{
+	check(Archetype.DataPtr.IsValid());
+	const FArchetypeData& ArchetypeData = *Archetype.DataPtr.Get();
+	ArchetypeData.ForEachComponentType(Function);
+}
+
+void UEntitySubsystem::SetDefaultChunkComponentValue(const FArchetypeHandle Archetype, FConstStructView InstancedStruct)
+{
+	check(Archetype.DataPtr.IsValid());
+	FArchetypeData& ArchetypeData = *Archetype.DataPtr.Get();
+	ArchetypeData.SetDefaultChunkComponentValue(InstancedStruct);
+}
+
+void UEntitySubsystem::DoEntityCompaction(const double TimeAllowed)
+{
+	const double TimeAllowedEnd = FPlatformTime::Seconds() + TimeAllowed;
+
+	bool bReachedTimeLimit = false;
+	for (const auto& KVP : ComponentHashToArchetypeMap)
+	{
+		for (const TSharedPtr<FArchetypeData>& ArchetypePtr : KVP.Value)
+		{
+			const double TimeAllowedLeft = TimeAllowedEnd - FPlatformTime::Seconds();
+			bReachedTimeLimit = TimeAllowedLeft <= 0.0;
+			if (bReachedTimeLimit)
+			{
+				break;
+			}
+			ArchetypePtr->CompactEntities(TimeAllowedLeft);
+		}
+		if (bReachedTimeLimit)
+		{
+			break;
+		}
+	}
+}
+
+FLWEntity UEntitySubsystem::CreateEntity(const FArchetypeHandle Archetype)
+{
+	check(Archetype.DataPtr.IsValid());
+
+	FLWEntity Entity = ReserveEntity();
+	InternalBuildEntity(Entity, Archetype);
+	return Entity;
+}
+
+FLWEntity UEntitySubsystem::CreateEntity(TConstArrayView<FInstancedStruct> ComponentInstanceList)
+{
+	check(ComponentInstanceList.Num() > 0);
+
+	FArchetypeHandle Archetype = CreateArchetype(FLWCompositionDescriptor(ComponentInstanceList, FLWTagBitSet(), FLWChunkComponentBitSet()));
+	check(Archetype.DataPtr.IsValid());
+
+	FLWEntity Entity = ReserveEntity();
+	InternalBuildEntity(Entity, Archetype);
+
+	FEntityData& EntityData = Entities[Entity.Index];
+	EntityData.CurrentArchetype->SetComponentsData(Entity, ComponentInstanceList);
+
+	return Entity;
+}
+
+FLWEntity UEntitySubsystem::ReserveEntity()
+{
+	// @todo: Need to add thread safety to the reservation of an entity
+	FLWEntity Result;
+	Result.Index = (EntityFreeIndexList.Num() > 0) ? EntityFreeIndexList.Pop(/*bAllowShrinking=*/ false) : Entities.Add();
+	Result.SerialNumber = SerialNumberGenerator.fetch_add(1);
+	Entities[Result.Index].SerialNumber = Result.SerialNumber;
+
+	return Result;
+}
+
+void UEntitySubsystem::ReleaseReservedEntity(FLWEntity Entity)
+{
+	checkf(!IsEntityBuilt(Entity), TEXT("Entity is already built, use DestroyEntity() instead"));
+
+	InternalReleaseEntity(Entity);
+}
+
+void UEntitySubsystem::BuildEntity(FLWEntity Entity, FArchetypeHandle Archetype)
+{
+	checkf(!IsEntityBuilt(Entity), TEXT("Expecting an entity that is not already built"));
+	check(Archetype.DataPtr.IsValid());
+
+	InternalBuildEntity(Entity, Archetype);
+}
+
+void UEntitySubsystem::BuildEntity(FLWEntity Entity, TConstArrayView<FInstancedStruct> ComponentInstanceList)
+{
+	check(ComponentInstanceList.Num() > 0);
+	checkf(!IsEntityBuilt(Entity), TEXT("Expecting an entity that is not already built"));
+
+	FArchetypeHandle Archetype = CreateArchetype(FLWCompositionDescriptor(ComponentInstanceList, FLWTagBitSet(), FLWChunkComponentBitSet()));
+	check(Archetype.DataPtr.IsValid());
+
+	InternalBuildEntity(Entity, Archetype);
+
+	FEntityData& EntityData = Entities[Entity.Index];
+	EntityData.CurrentArchetype->SetComponentsData(Entity, ComponentInstanceList);
+}
+
+int32 UEntitySubsystem::BatchCreateEntities(const FArchetypeHandle Archetype, const int32 Count, TArray<FLWEntity>& OutEntities)
+{
+	FArchetypeData* ArchetypePtr = Archetype.DataPtr.Get();
+	check(ArchetypePtr);
+	check(Count >= 0);
+	
+	int32 Index = OutEntities.Num();
+	OutEntities.AddDefaulted(Count);
+
+	// @todo optimize
+	for (; Index < OutEntities.Num(); ++Index)
+	{
+		FLWEntity& Result = OutEntities[Index];
+		Result.Index = (EntityFreeIndexList.Num() > 0) ? EntityFreeIndexList.Pop(/*bAllowShrinking=*/ false) : Entities.Add();
+		Result.SerialNumber = SerialNumberGenerator++;
+
+		FEntityData& EntityData = Entities[Result.Index];
+		EntityData.CurrentArchetype = Archetype.DataPtr;
+		EntityData.SerialNumber = Result.SerialNumber;
+		
+		ArchetypePtr->AddEntity(Result);
+	}
+
+	return Count;
+}
+
+void UEntitySubsystem::DestroyEntity(FLWEntity Entity)
+{
+	checkf(IsProcessing() == false, TEXT("Synchronous API function %s called during mass processing. Use asynchronous API instead."), ANSI_TO_TCHAR(__FUNCTION__));
+	
+	CheckIfEntityIsActive(Entity);
+
+	FEntityData& EntityData = Entities[Entity.Index];
+	FArchetypeData* Archetype = EntityData.CurrentArchetype.Get();
+	check(Archetype);
+	Archetype->RemoveEntity(Entity);
+
+	InternalReleaseEntity(Entity);
+
+	if (Archetype->GetNumEntities() == 0)
+	{
+		//@TODO: Drop our reference to the archetype?
+	}
+}
+
+void UEntitySubsystem::BatchDestroyEntities(TConstArrayView<FLWEntity> InEntities)
+{
+	TRACE_CPUPROFILER_EVENT_SCOPE_STR("BatchDestroyEntities");
+
+	checkf(IsProcessing() == false, TEXT("Synchronous API function %s called during mass processing. Use asynchronous API instead."), ANSI_TO_TCHAR(__FUNCTION__));
+	
+	EntityFreeIndexList.Reserve(EntityFreeIndexList.Num() + InEntities.Num());
+
+	// @todo optimize, we can make savings by implementing Archetype->RemoveEntities()
+	for (const FLWEntity Entity : InEntities)
+	{
+		if (Entities.IsValidIndex(Entity.Index) == false)
+		{
+			continue;
+		}
+
+		FEntityData& EntityData = Entities[Entity.Index];
+		if (EntityData.SerialNumber != Entity.SerialNumber)
+		{
+			continue;
+		}
+
+		FArchetypeData* Archetype = EntityData.CurrentArchetype.Get();
+		check(Archetype);
+		Archetype->RemoveEntity(Entity);
+
+		EntityData.Reset();
+		EntityFreeIndexList.Add(Entity.Index);
+	}
+}
+
+void UEntitySubsystem::BatchDestroyEntityChunks(const FArchetypeChunkCollection& Chunks)
+{
+	TRACE_CPUPROFILER_EVENT_SCOPE_STR("BatchDestroyEntityChunks");
+
+	checkf(IsProcessing() == false, TEXT("Synchronous API function %s called during mass processing. Use asynchronous API instead."), ANSI_TO_TCHAR(__FUNCTION__));
+
+	check(Chunks.GetArchetype().IsValid());
+	TArray<FLWEntity> EntitiesRemoved;
+	Chunks.GetArchetype().DataPtr->BatchDestroyEntityChunks(Chunks, EntitiesRemoved);
+	
+	EntityFreeIndexList.Reserve(EntityFreeIndexList.Num() + EntitiesRemoved.Num());
+
+	for (const FLWEntity& Entity : EntitiesRemoved)
+	{
+		check(Entities.IsValidIndex(Entity.Index));
+
+		FEntityData& EntityData = Entities[Entity.Index];
+		if (EntityData.SerialNumber == Entity.SerialNumber)
+		{
+			EntityData.Reset();
+			EntityFreeIndexList.Add(Entity.Index);
+		}
+	}
+}
+
+void UEntitySubsystem::AddComponentToEntity(FLWEntity Entity, const UScriptStruct* ComponentType)
+{
+	checkf(ComponentType, TEXT("Null component type passed in to %s"), ANSI_TO_TCHAR(__FUNCTION__));
+	checkf(IsProcessing() == false, TEXT("Synchronous API function %s called during mass processing. Use asynchronous API instead."), ANSI_TO_TCHAR(__FUNCTION__));
+
+	CheckIfEntityIsActive(Entity);
+
+	InternalAddComponentListToEntityChecked(Entity, FLWComponentBitSet(*ComponentType));
+}
+
+void UEntitySubsystem::AddComponentListToEntity(FLWEntity Entity, TConstArrayView<const UScriptStruct*> ComponentList)
+{
+	CheckIfEntityIsActive(Entity);
+
+	InternalAddComponentListToEntityChecked(Entity, FLWComponentBitSet(ComponentList));
+}
+
+void UEntitySubsystem::AddCompositionToEntity_GetDelta(FLWEntity Entity, FLWCompositionDescriptor& Composition)
+{
+	CheckIfEntityIsActive(Entity);
+
+	FEntityData& EntityData = Entities[Entity.Index];
+	FArchetypeData* OldArchetype = EntityData.CurrentArchetype.Get();
+	check(OldArchetype);
+
+	Composition.Components -= OldArchetype->GetCompositionDescriptor().Components;
+	Composition.Tags -= OldArchetype->GetCompositionDescriptor().Tags;
+
+	if (Composition.IsEmpty() == false)
+	{
+		FLWCompositionDescriptor NewDescriptor = OldArchetype->GetCompositionDescriptor();
+		NewDescriptor.Components += Composition.Components;
+		NewDescriptor.Tags+= Composition.Tags;
+
+		const FArchetypeHandle NewArchetypeHandle = CreateArchetype(NewDescriptor);
+
+		if (ensure(NewArchetypeHandle.DataPtr != EntityData.CurrentArchetype))
+		{
+			// Move the entity over
+			FArchetypeData* NewArchetype = NewArchetypeHandle.DataPtr.Get();
+			check(NewArchetype);
+			EntityData.CurrentArchetype->MoveEntityToAnotherArchetype(Entity, *NewArchetype);
+			EntityData.CurrentArchetype = NewArchetypeHandle.DataPtr;
+		}
+	}
+}
+
+void UEntitySubsystem::RemoveCompositionFromEntity(FLWEntity Entity, const FLWCompositionDescriptor& InDescriptor)
+{
+	CheckIfEntityIsActive(Entity);
+
+	FEntityData& EntityData = Entities[Entity.Index];
+	FArchetypeData* OldArchetype = EntityData.CurrentArchetype.Get();
+	check(OldArchetype);
+
+	FLWCompositionDescriptor NewDescriptor = OldArchetype->GetCompositionDescriptor();
+	NewDescriptor.Components -= InDescriptor.Components;
+	NewDescriptor.Tags -= InDescriptor.Tags;
+
+	if (NewDescriptor.IsEmpty() == false)
+	{
+		const FArchetypeHandle NewArchetypeHandle = CreateArchetype(NewDescriptor);
+
+		if (ensure(NewArchetypeHandle.DataPtr != EntityData.CurrentArchetype))
+		{
+			// Move the entity over
+			FArchetypeData* NewArchetype = NewArchetypeHandle.DataPtr.Get();
+			check(NewArchetype);
+			EntityData.CurrentArchetype->MoveEntityToAnotherArchetype(Entity, *NewArchetype);
+			EntityData.CurrentArchetype = NewArchetypeHandle.DataPtr;
+		}
+	}
+}
+
+void UEntitySubsystem::InternalBuildEntity(FLWEntity Entity, const FArchetypeHandle Archetype)
+{
+	FEntityData& EntityData = Entities[Entity.Index];
+	EntityData.CurrentArchetype = Archetype.DataPtr;
+	EntityData.CurrentArchetype->AddEntity(Entity);
+}
+
+void UEntitySubsystem::InternalReleaseEntity(FLWEntity Entity)
+{
+	FEntityData& EntityData = Entities[Entity.Index];
+	EntityData.Reset();
+	EntityFreeIndexList.Add(Entity.Index);
+}
+
+void UEntitySubsystem::InternalAddComponentListToEntityChecked(FLWEntity Entity, const FLWComponentBitSet& InComponents)
+{
+	FEntityData& EntityData = Entities[Entity.Index];
+	FArchetypeData* OldArchetype = EntityData.CurrentArchetype.Get();
+	check(OldArchetype);
+	
+	ensureMsgf(OldArchetype->GetComponentBitSet().HasAny(InComponents) == false, TEXT("Trying to add a new component type to an entity, but it already has some of them."));
+
+	const FLWComponentBitSet NewComponents = InComponents - OldArchetype->GetComponentBitSet();
+	if (NewComponents.IsEmpty() == false)
+	{
+		InternalAddComponentListToEntity(Entity, NewComponents);
+	}
+}
+
+void UEntitySubsystem::InternalAddComponentListToEntity(FLWEntity Entity, const FLWComponentBitSet& NewComponents)
+{
+	checkf(NewComponents.IsEmpty() == false, TEXT("%s is intended for internal calls with non empty NewComponents parameter"), ANSI_TO_TCHAR(__FUNCTION__));
+	check(Entities.IsValidIndex(Entity.Index));
+	FEntityData& EntityData = Entities[Entity.Index];
+	check(EntityData.CurrentArchetype.IsValid());
+
+	// fetch or create the new archetype
+	const FArchetypeHandle NewArchetypeHandle = CreateArchetype(EntityData.CurrentArchetype, NewComponents);
+
+	if (NewArchetypeHandle.DataPtr != EntityData.CurrentArchetype)
+	{
+		// Move the entity over
+		FArchetypeData* NewArchetype = NewArchetypeHandle.DataPtr.Get();
+		check(NewArchetype);
+		EntityData.CurrentArchetype->MoveEntityToAnotherArchetype(Entity, *NewArchetype);
+		EntityData.CurrentArchetype = NewArchetypeHandle.DataPtr;
+	}
+}
+
+void UEntitySubsystem::AddComponentInstanceListToEntity(FLWEntity Entity, TConstArrayView<FInstancedStruct> ComponentInstanceList)
+{
+	checkf(IsProcessing() == false, TEXT("Synchronous API function %s called during mass processing. Use asynchronous API instead."), ANSI_TO_TCHAR(__FUNCTION__));
+
+	CheckIfEntityIsActive(Entity);
+	checkf(ComponentInstanceList.Num() > 0, TEXT("Need to specify at least one component instances for this operation"));
+
+	InternalAddComponentListToEntityChecked(Entity, FLWComponentBitSet(ComponentInstanceList));
+
+	FEntityData& EntityData = Entities[Entity.Index];
+	EntityData.CurrentArchetype->SetComponentsData(Entity, ComponentInstanceList);
+}
+
+void UEntitySubsystem::RemoveComponentFromEntity(FLWEntity Entity, const UScriptStruct* ComponentType)
+{
+	RemoveComponentListFromEntity(Entity, MakeArrayView(&ComponentType, 1));
+}
+
+void UEntitySubsystem::RemoveComponentListFromEntity(FLWEntity Entity, TConstArrayView<const UScriptStruct*> ComponentList)
+{
+	checkf(IsProcessing() == false, TEXT("Synchronous API function %s called during mass processing. Use asynchronous API instead."), ANSI_TO_TCHAR(__FUNCTION__));
+
+	CheckIfEntityIsActive(Entity);
+	
+	FEntityData& EntityData = Entities[Entity.Index];
+	FArchetypeData* OldArchetype = EntityData.CurrentArchetype.Get();
+	check(OldArchetype);
+
+	const FLWComponentBitSet ComponentsToRemove(ComponentList);
+
+	if (ensureMsgf(OldArchetype->GetComponentBitSet().HasAny(ComponentsToRemove), TEXT("Trying to remove a list of components from an entity but none of the components given was found.")))
+	{
+		// If all the components got removed this will result in fetching of the empty archetype
+		const FArchetypeHandle NewArchetypeHandle = CreateArchetype(FLWCompositionDescriptor(OldArchetype->GetComponentBitSet() - ComponentsToRemove, OldArchetype->GetTagBitSet(), OldArchetype->GetChunkComponentBitSet()));
+
+		// Move the entity over
+		FArchetypeData* NewArchetype = NewArchetypeHandle.DataPtr.Get();
+		check(NewArchetype);
+		OldArchetype->MoveEntityToAnotherArchetype(Entity, *NewArchetype);
+		EntityData.CurrentArchetype = NewArchetypeHandle.DataPtr;
+	}
+}
+
+void UEntitySubsystem::SwapTagsForEntity(FLWEntity Entity, const UScriptStruct* OldTagType, const UScriptStruct* NewTagType)
+{
+	checkf(IsProcessing() == false, TEXT("Synchronous API function %s called during mass processing. Use asynchronous API instead."), ANSI_TO_TCHAR(__FUNCTION__));
+
+	CheckIfEntityIsActive(Entity);
+
+	checkf((OldTagType != nullptr) && OldTagType->IsChildOf(FComponentTag::StaticStruct()), TEXT("%s works only with tags while '%s' is not one."), ANSI_TO_TCHAR(__FUNCTION__), *GetPathNameSafe(OldTagType));
+	checkf((NewTagType != nullptr) && NewTagType->IsChildOf(FComponentTag::StaticStruct()), TEXT("%s works only with tags while '%s' is not one."), ANSI_TO_TCHAR(__FUNCTION__), *GetPathNameSafe(NewTagType));
+
+	FEntityData& EntityData = Entities[Entity.Index];
+	FArchetypeData* CurrentArchetype = EntityData.CurrentArchetype.Get();
+	check(CurrentArchetype);
+
+	FLWTagBitSet NewTagBitSet = CurrentArchetype->GetTagBitSet();
+	NewTagBitSet.Remove(*OldTagType);
+	NewTagBitSet.Add(*NewTagType);
+	
+	if (NewTagBitSet != CurrentArchetype->GetTagBitSet())
+	{
+		FArchetypeHandle NewArchetypeHandle = InternalCreateSiblingArchetype(EntityData.CurrentArchetype, NewTagBitSet);
+		checkSlow(NewArchetypeHandle.IsValid());
+
+		// Move the entity over
+		EntityData.CurrentArchetype->MoveEntityToAnotherArchetype(Entity, *NewArchetypeHandle.DataPtr.Get());
+		EntityData.CurrentArchetype = NewArchetypeHandle.DataPtr;
+	}
+}
+
+void UEntitySubsystem::AddTagToEntity(FLWEntity Entity, const UScriptStruct* TagType)
+{
+	checkf((TagType != nullptr) && TagType->IsChildOf(FComponentTag::StaticStruct()), TEXT("%s works only with tags while '%s' is not one."), ANSI_TO_TCHAR(__FUNCTION__), *GetPathNameSafe(TagType));
+
+	CheckIfEntityIsActive(Entity);
+
+	FEntityData& EntityData = Entities[Entity.Index];
+	FArchetypeData* CurrentArchetype = EntityData.CurrentArchetype.Get();
+	check(CurrentArchetype);
+
+	if (CurrentArchetype->HasTagType(TagType) == false)
+	{
+		//FLWTagBitSet NewTags = CurrentArchetype->GetTagBitSet() - *TagType;
+		FLWTagBitSet NewTags = CurrentArchetype->GetTagBitSet();
+		NewTags.Add(*TagType);
+		FArchetypeHandle NewArchetypeHandle = InternalCreateSiblingArchetype(EntityData.CurrentArchetype, NewTags);
+		checkSlow(NewArchetypeHandle.IsValid());
+
+		// Move the entity over
+		EntityData.CurrentArchetype->MoveEntityToAnotherArchetype(Entity, *NewArchetypeHandle.DataPtr.Get());
+		EntityData.CurrentArchetype = NewArchetypeHandle.DataPtr;
+	}
+}
+	
+void UEntitySubsystem::RemoveTagFromEntity(FLWEntity Entity, const UScriptStruct* TagType)
+{
+	checkf((TagType != nullptr) && TagType->IsChildOf(FComponentTag::StaticStruct()), TEXT("%s works only with tags while '%s' is not one."), ANSI_TO_TCHAR(__FUNCTION__), *GetPathNameSafe(TagType));
+
+	CheckIfEntityIsActive(Entity);
+
+	FEntityData& EntityData = Entities[Entity.Index];
+	FArchetypeData* CurrentArchetype = EntityData.CurrentArchetype.Get();
+	check(CurrentArchetype);
+
+	if (CurrentArchetype->HasTagType(TagType))
+	{
+		// CurrentArchetype->GetTagBitSet() -  *TagType
+		FLWTagBitSet NewTags = CurrentArchetype->GetTagBitSet();
+		NewTags.Remove(*TagType);
+		FArchetypeHandle NewArchetypeHandle = InternalCreateSiblingArchetype(EntityData.CurrentArchetype, NewTags);
+		checkSlow(NewArchetypeHandle.IsValid());
+
+		// Move the entity over
+		EntityData.CurrentArchetype->MoveEntityToAnotherArchetype(Entity, *NewArchetypeHandle.DataPtr.Get());
+		EntityData.CurrentArchetype = NewArchetypeHandle.DataPtr;
+	}
+}
+
+
+void UEntitySubsystem::MoveEntityToAnotherArchetype(FLWEntity Entity, FArchetypeHandle NewArchetypeHandle)
+{
+	CheckIfEntityIsActive(Entity);
+
+	FArchetypeData* NewArchetype = NewArchetypeHandle.DataPtr.Get();
+	check(NewArchetype);
+
+	// Move the entity over
+	FEntityData& EntityData = Entities[Entity.Index];
+	EntityData.CurrentArchetype->MoveEntityToAnotherArchetype(Entity, *NewArchetype);
+	EntityData.CurrentArchetype = NewArchetypeHandle.DataPtr;
+}
+
+void UEntitySubsystem::SetEntityComponentsValues(FLWEntity Entity, TArrayView<const FInstancedStruct> ComponentInstanceList)
+{
+	CheckIfEntityIsActive(Entity);
+
+	FEntityData& EntityData = Entities[Entity.Index];
+	EntityData.CurrentArchetype->SetComponentsData(Entity, ComponentInstanceList);
+}
+
+void UEntitySubsystem::BatchSetEntityComponentsValues(const FArchetypeChunkCollection& SparseEntities, TArrayView<const FInstancedStruct> ComponentInstanceList)
+{
+	FArchetypeData* Archetype = SparseEntities.GetArchetype().DataPtr.Get();
+	check(Archetype);
+
+	for (const FInstancedStruct& ComponentTemplate : ComponentInstanceList)
+	{
+		Archetype->SetComponentData(SparseEntities, ComponentTemplate);
+	}
+}
+
+void* UEntitySubsystem::InternalGetComponentDataChecked(FLWEntity Entity, const UScriptStruct* ComponentType) const
+{
+	CheckIfEntityIsActive(Entity);
+
+	checkf((ComponentType != nullptr) && ComponentType->IsChildOf(FLWComponentData::StaticStruct()), TEXT("InternalGetComponentDataChecked called with an invalid component type '%s'"), *GetPathNameSafe(ComponentType));
+	const FEntityData& EntityData = Entities[Entity.Index];
+	return EntityData.CurrentArchetype->GetComponentDataForEntityChecked(ComponentType, Entity.Index);
+}
+
+void* UEntitySubsystem::InternalGetComponentDataPtr(FLWEntity Entity, const UScriptStruct* ComponentType) const
+{
+	CheckIfEntityIsActive(Entity);
+	checkf((ComponentType != nullptr) && ComponentType->IsChildOf(FLWComponentData::StaticStruct()), TEXT("InternalGetComponentData called with an invalid component type '%s'"), *GetPathNameSafe(ComponentType));
+	const FEntityData& EntityData = Entities[Entity.Index];
+	return EntityData.CurrentArchetype->GetComponentDataForEntity(ComponentType, Entity.Index);
+}
+
+bool UEntitySubsystem::IsEntityValid(FLWEntity Entity) const
+{
+	return (Entity.Index > 0) && Entities.IsValidIndex(Entity.Index) && (Entities[Entity.Index].SerialNumber == Entity.SerialNumber);
+}
+
+bool UEntitySubsystem::IsEntityBuilt(FLWEntity Entity) const
+{
+	CheckIfEntityIsValid(Entity);
+	return Entities[Entity.Index].CurrentArchetype.IsValid();
+}
+
+void UEntitySubsystem::CheckIfEntityIsValid(FLWEntity Entity) const
+{
+	checkf(IsEntityValid(Entity), TEXT("Invalid entity (ID: %d, SN:%d, %s)"), Entity.Index, Entity.SerialNumber,
+		   (Entity.Index == 0) ? TEXT("was never initialized") : TEXT("already destroyed"));
+}
+
+void UEntitySubsystem::CheckIfEntityIsActive(FLWEntity Entity) const
+{
+	checkf(IsEntityBuilt(Entity), TEXT("Entity not yet created(ID: %d, SN:%d)"));
+}
+
+void UEntitySubsystem::GetValidArchetypes(const FLWComponentQuery& Query, TArray<FArchetypeHandle>& OutValidArchetypes)
+{
+	//@TODO: Not optimized yet, but we call this rarely now, so not a big deal.
+
+	// First get set of all archetypes that contain *any* component
+	TSet<TSharedPtr<FArchetypeData>> AnyArchetypes;
+	for (const FLWComponentRequirement& Requirement : Query.GetRequirements())
+	{
+		check(Requirement.StructType);
+		if (Requirement.Presence != ELWComponentPresence::None)
+		{
+			if (TArray<TSharedPtr<FArchetypeData>>* pData = ComponentTypeToArchetypeMap.Find(Requirement.StructType))
+			{
+				AnyArchetypes.Append(*pData);
+			}
+		}
+	}
+
+	// Then verify that they contain *all* required components
+	for (TSharedPtr<FArchetypeData>& ArchetypePtr : AnyArchetypes)
+	{
+		FArchetypeData& Archetype = *(ArchetypePtr.Get());
+
+		if (Archetype.GetTagBitSet().HasAll(Query.GetRequiredAllTags()) == false)
+		{
+			// missing some required tags, skip.
+#if WITH_AGGREGATETICKING_DEBUG
+			const FLWTagBitSet UnsatisfiedTags = Query.GetRequiredAllTags() - Archetype.GetTagBitSet();
+			FStringOutputDevice Description;
+			UnsatisfiedTags.DebugGetStringDesc(Description);
+			UE_LOG(LogAggregateTicking, VeryVerbose, TEXT("Archetype did not match due to missing tags: %s")
+				, *Description);
+#endif // WITH_AGGREGATETICKING_DEBUG
+			continue;
+		}
+
+		if (Archetype.GetTagBitSet().HasNone(Query.GetRequiredNoneTags()) == false)
+		{
+			// has some tags required to be absent
+#if WITH_AGGREGATETICKING_DEBUG
+			const FLWTagBitSet UnwantedTags = Query.GetRequiredAllTags().GetOverlap(Archetype.GetTagBitSet());
+			FStringOutputDevice Description;
+			UnwantedTags.DebugGetStringDesc(Description);
+			UE_LOG(LogAggregateTicking, VeryVerbose, TEXT("Archetype has tags required absent: %s")
+				, *Description);
+#endif // WITH_AGGREGATETICKING_DEBUG
+			continue;
+		}
+
+		if (Query.GetRequiredAnyTags().IsEmpty() == false 
+			&& Archetype.GetTagBitSet().HasAny(Query.GetRequiredAnyTags()) == false)
+		{
+#if WITH_AGGREGATETICKING_DEBUG
+			FStringOutputDevice Description;
+			Query.GetRequiredAnyTags().DebugGetStringDesc(Description);
+			UE_LOG(LogAggregateTicking, VeryVerbose, TEXT("Archetype did not match due to missing \'any\' tags: %s")
+				, *Description);
+#endif // WITH_AGGREGATETICKING_DEBUG
+			continue;
+		}
+		
+		if (Archetype.GetComponentBitSet().HasAll(Query.GetRequiredAllComponents()) == false)
+		{
+			// missing some required components, skip.
+#if WITH_AGGREGATETICKING_DEBUG
+			const FLWComponentBitSet UnsatisfiedComponents = Query.GetRequiredAllComponents() - Archetype.GetComponentBitSet();
+			FStringOutputDevice Description;
+			UnsatisfiedComponents.DebugGetStringDesc(Description);
+			UE_LOG(LogAggregateTicking, VeryVerbose, TEXT("Archetype did not match due to missing Components: %s")
+				, *Description);
+#endif // WITH_AGGREGATETICKING_DEBUG
+			continue;
+		}
+
+		if (Archetype.GetComponentBitSet().HasNone(Query.GetRequiredNoneComponents()) == false)
+		{
+			// has some Components required to be absent
+#if WITH_AGGREGATETICKING_DEBUG
+			const FLWComponentBitSet UnwantedComponents = Query.GetRequiredAllComponents().GetOverlap(Archetype.GetComponentBitSet());
+			FStringOutputDevice Description;
+			UnwantedComponents.DebugGetStringDesc(Description);
+			UE_LOG(LogAggregateTicking, VeryVerbose, TEXT("Archetype has Components required absent: %s")
+				, *Description);
+#endif // WITH_AGGREGATETICKING_DEBUG
+			continue;
+		}
+
+		if (Query.GetRequiredAnyComponents().IsEmpty() == false 
+			&& Archetype.GetComponentBitSet().HasAny(Query.GetRequiredAnyComponents()) == false)
+		{
+#if WITH_AGGREGATETICKING_DEBUG
+			FStringOutputDevice Description;
+			Query.GetRequiredAnyComponents().DebugGetStringDesc(Description);
+			UE_LOG(LogAggregateTicking, VeryVerbose, TEXT("Archetype did not match due to missing \'any\' components: %s")
+				, *Description);
+#endif // WITH_AGGREGATETICKING_DEBUG
+			continue;
+		}
+
+		if (Archetype.GetChunkComponentBitSet().HasAll(Query.GetRequiredAllChunkComponents()) == false)
+		{
+			// missing some required components, skip.
+#if WITH_AGGREGATETICKING_DEBUG
+			const FLWChunkComponentBitSet UnsatisfiedComponents = Query.GetRequiredAllChunkComponents() - Archetype.GetChunkComponentBitSet();
+			FStringOutputDevice Description;
+			UnsatisfiedComponents.DebugGetStringDesc(Description);
+			UE_LOG(LogAggregateTicking, VeryVerbose, TEXT("Archetype did not match due to missing Chunk Components: %s")
+				, *Description);
+#endif // WITH_AGGREGATETICKING_DEBUG
+			continue;
+		}
+
+		if (Archetype.GetChunkComponentBitSet().HasNone(Query.GetRequiredNoneChunkComponents()) == false)
+		{
+			// has some Components required to be absent
+#if WITH_AGGREGATETICKING_DEBUG
+			const FLWChunkComponentBitSet UnwantedComponents = Query.GetRequiredNoneChunkComponents().GetOverlap(Archetype.GetChunkComponentBitSet());
+			FStringOutputDevice Description;
+			UnwantedComponents.DebugGetStringDesc(Description);
+			UE_LOG(LogAggregateTicking, VeryVerbose, TEXT("Archetype has Chunk Components required absent: %s")
+				, *Description);
+#endif // WITH_AGGREGATETICKING_DEBUG
+			continue;
+		}
+
+		OutValidArchetypes.Add(ArchetypePtr);
+	}
+}
+
+FLWComponentSystemExecutionContext UEntitySubsystem::CreateExecutionContext(const float DeltaSeconds) const
+{
+	FLWComponentSystemExecutionContext ExecutionContext(DeltaSeconds);
+	ExecutionContext.SetDeferredCommandBuffer(DeferredCommandBuffer);
+	return MoveTemp(ExecutionContext);
+}
+
+#if WITH_AGGREGATETICKING_DEBUG
+void UEntitySubsystem::DebugPrintArchetypes(FOutputDevice& Ar) const
+{
+	Ar.Logf(ELogVerbosity::Log, TEXT("Listing archetypes contained in %s"), *GetPathNameSafe(this));
+
+	int32 NumBuckets = 0;
+	int32 NumArchetypes = 0;
+	int32 LongestArchetypeBucket = 0;
+	for (const auto& KVP : ComponentHashToArchetypeMap)
+	{
+		for (const TSharedPtr<FArchetypeData>& ArchetypePtr : KVP.Value)
+		{
+			ArchetypePtr->DebugPrintArchetype(Ar);
+		}
+
+		const int32 NumArchetypesInBucket = KVP.Value.Num();
+		LongestArchetypeBucket = FMath::Max(LongestArchetypeBucket, NumArchetypesInBucket);
+		NumArchetypes += NumArchetypesInBucket;
+		++NumBuckets;
+	}
+
+	Ar.Logf(ELogVerbosity::Log, TEXT("ComponentHashToArchetypeMap: %d archetypes across %d buckets, longest bucket is %d"),
+		NumArchetypes, NumBuckets, LongestArchetypeBucket);
+}
+
+void UEntitySubsystem::DebugPrintEntity(int32 Index, FOutputDevice& Ar, const TCHAR* InPrefix) const
+{
+	if (Index >= Entities.Num())
+	{
+		Ar.Logf(ELogVerbosity::Log, TEXT("Unable to list components values for out of range index in %s"), *GetPathNameSafe(this));
+		return;
+	}
+
+	const FEntityData& EntityData = Entities[Index];
+	if (!EntityData.IsValid())
+	{
+		Ar.Logf(ELogVerbosity::Log, TEXT("Unable to list components values for invalid entity in %s"), *GetPathNameSafe(this));
+	}
+
+	FLWEntity Entity;
+	Entity.Index = Index;
+	Entity.SerialNumber = EntityData.SerialNumber;
+	DebugPrintEntity(Entity, Ar, InPrefix);
+}
+
+void UEntitySubsystem::DebugPrintEntity(FLWEntity Entity, FOutputDevice& Ar, const TCHAR* InPrefix) const
+{
+	if (!IsEntityActive(Entity))
+	{
+		Ar.Logf(ELogVerbosity::Log, TEXT("Unable to list components values for invalid entity in %s"), *GetPathNameSafe(this));
+	}
+
+	Ar.Logf(ELogVerbosity::Log, TEXT("Listing components values for Entity[%s] in %s"), *Entity.DebugGetDescription(), *GetPathNameSafe(this));
+
+	const FEntityData& EntityData = Entities[Entity.Index];
+	FArchetypeData* Archetype = EntityData.CurrentArchetype.Get();
+	if (Archetype == nullptr)
+	{
+		Ar.Logf(ELogVerbosity::Log, TEXT("Unable to list components values for invalid entity in %s"), *GetPathNameSafe(this));
+	}
+	else
+	{
+		Archetype->DebugPrintEntity(Entity, Ar, InPrefix);
+	}
+}
+
+void UEntitySubsystem::DebugGetStringDesc(const FArchetypeHandle& Archetype, FOutputDevice& Ar)
+{
+	Ar.Logf(TEXT("%s"), Archetype.IsValid() ? *Archetype.DataPtr->DebugGetDescription() : TEXT("INVALID"));
+}
+
+void UEntitySubsystem::DebugGetArchetypesStringDetails(FOutputDevice& Ar, const bool bIncludeEmpty)
+{
+#if WITH_AGGREGATETICKING_DEBUG
+	Ar.SetAutoEmitLineTerminator(true);
+	for (auto Pair : ComponentHashToArchetypeMap)
+	{
+		Ar.Logf(ELogVerbosity::Log, TEXT("\n-----------------------------------\nHash: %u"), Pair.Key);
+		for (TSharedPtr<FArchetypeData> Archetype : Pair.Value)
+		{
+			if (Archetype.IsValid() && (bIncludeEmpty == true || Archetype->GetChunkCount() > 0))
+			{
+				Archetype->DebugPrintArchetype(Ar);
+				Ar.Logf(ELogVerbosity::Log, TEXT("+++++++++++++++++++++++++\n"));
+			}
+		}
+	}
+#endif // WITH_AGGREGATETICKING_DEBUG
+}
+
+void UEntitySubsystem::DebugGetArchetypeComponentTypes(const FArchetypeHandle& Archetype, TArray<const UScriptStruct*>& InOutComponentList) const
+{
+	if (Archetype.DataPtr.IsValid())
+	{
+		Archetype.DataPtr->GetCompositionDescriptor().Components.DebugGetStructTypes(InOutComponentList);
+	}
+}
+
+int32 UEntitySubsystem::DebugGetArchetypeEntitiesCount(const FArchetypeHandle& Archetype) const
+{
+	return Archetype.DataPtr.IsValid() ? Archetype.DataPtr->GetNumEntities() : 0;
+}
+
+void UEntitySubsystem::DebugRemoveAllEntities()
+{
+	for (int EntityIndex = NumReservedEntities; EntityIndex < Entities.Num(); ++EntityIndex)
+	{
+		FEntityData& EntityData = Entities[EntityIndex];
+		if (EntityData.IsValid() == false)
+		{
+			// already dead
+			continue;
+		}
+		const TSharedPtr<FArchetypeData>& Archetype = EntityData.CurrentArchetype;
+		FLWEntity Entity;
+		Entity.Index = EntityIndex;
+		Entity.SerialNumber = EntityData.SerialNumber;
+		Archetype->RemoveEntity(Entity);
+
+		EntityData.Reset();
+		EntityFreeIndexList.Add(EntityIndex);
+	}
+}
+
+void UEntitySubsystem::DebugGetArchetypeStrings(const FArchetypeHandle& Archetype, TArray<FName>& OutComponentNames, TArray<FName>& OutTagNames)
+{
+	if (Archetype.IsValid() == false)
+	{
+		return;
+	}
+
+	FArchetypeData& ArchetypeRef = *Archetype.DataPtr.Get();
+	
+	OutComponentNames.Reserve(ArchetypeRef.GetComponentConfigs().Num());
+	for (const FArchetypeComponentConfig& ComponentConfig : ArchetypeRef.GetComponentConfigs())
+	{
+		checkSlow(ComponentConfig.ComponentType);
+		OutComponentNames.Add(ComponentConfig.ComponentType->GetFName());
+	}
+
+	ArchetypeRef.GetTagBitSet().DebugGetIndividualNames(OutTagNames);
+}
+
+//////////////////////////////////////////////////////////////////////
+// Debug commands
+
+FAutoConsoleCommandWithWorldArgsAndOutputDevice GPrintArchetypesCmd(
+	TEXT("EntityManager.PrintArchetypes"),
+	TEXT("Prints information about all archetypes in the current world"),
+	FConsoleCommandWithWorldArgsAndOutputDeviceDelegate::CreateStatic(
+		[](const TArray<FString>& Params, UWorld* World, FOutputDevice& Ar)
+{
+	if (UEntitySubsystem* EntitySystem = World ? World->GetSubsystem<UEntitySubsystem>() : nullptr)
+	{
+		EntitySystem->DebugPrintArchetypes(Ar);
+	}
+	else
+	{
+		Ar.Logf(ELogVerbosity::Error, TEXT("Failed to find Entity Subsystem for world %s"), *GetPathNameSafe(World));
+	}
+}));
+#endif // WITH_AGGREGATETICKING_DEBUG
+
+//////////////////////////////////////////////////////////////////////
+// FLWComponentSystemExecutionContext
+
+void FLWComponentSystemExecutionContext::FlushDeferred(UEntitySubsystem& EntitySystem) const
+{
+	if (bFlushDeferredCommands && DeferredCommandBuffer)
+	{
+		DeferredCommandBuffer->ReplayBufferAgainstSystem(&EntitySystem);
+	}
+}
+
+void FLWComponentSystemExecutionContext::ClearExecutionData()
+{
+	ComponentViews.Reset();
+	CurrentArchetypesTagBitSet.Reset();
+	ChunkSerialModificationNumber = -1;
+}
+
+void FLWComponentSystemExecutionContext::SetCurrentArchetypeData(FArchetypeData& ArchetypeData)
+{
+	CurrentArchetypesTagBitSet = ArchetypeData.GetTagBitSet();
+}
+
+void FLWComponentSystemExecutionContext::SetChunkCollection(const FArchetypeChunkCollection& InChunkCollection)
+{
+	check(ChunkCollection.IsEmpty());
+	ChunkCollection = InChunkCollection;
+}
+
+void FLWComponentSystemExecutionContext::SetChunkCollection(FArchetypeChunkCollection&& InChunkCollection)
+{
+	check(ChunkCollection.IsEmpty());
+	ChunkCollection = MoveTemp(InChunkCollection);
+}
+
+void FLWComponentSystemExecutionContext::SetRequirements(TConstArrayView<FLWComponentRequirement> InRequirements, TConstArrayView<FLWComponentRequirement> InChunkRequirements)
+{ 
+	ComponentViews.Reset();
+	for (const FLWComponentRequirement& Requirement : InRequirements)
+	{
+		if (Requirement.RequiresBinding())
+		{
+			ComponentViews.Add(FComponentView(Requirement));
+		}
+	}
+
+	ChunkComponents.Reset();
+	for (const FLWComponentRequirement& Requirement : InChunkRequirements)
+	{
+		if (Requirement.RequiresBinding())
+		{
+			ChunkComponents.Add(FChunkComponentView(Requirement));
+		}
+	}
+}
+
diff --git a/Engine/Plugins/Runtime/MassEntity/Source/MassEntity/Private/EntityView.cpp b/Engine/Plugins/Runtime/MassEntity/Source/MassEntity/Private/EntityView.cpp
new file mode 100644
--- /dev/null
+++ b/Engine/Plugins/Runtime/MassEntity/Source/MassEntity/Private/EntityView.cpp
@@ -0,0 +1,49 @@
+// Copyright Epic Games, Inc. All Rights Reserved.
+
+#include "EntityView.h"
+#include "EntitySubsystem.h"
+#include "ArchetypeData.h"
+
+
+//////////////////////////////////////////////////////////////////////
+// FEntityView
+FEntityView::FEntityView(const FArchetypeHandle& ArchetypeHandle, FLWEntity InEntity)
+{
+	Entity = InEntity;
+	check(ArchetypeHandle.IsValid());
+	Archetype = ArchetypeHandle.DataPtr.Get();
+	EntityHandle = Archetype->MakeEntityHandle(Entity);
+}
+
+FEntityView::FEntityView(const UEntitySubsystem& EntitySubsystem, FLWEntity InEntity)
+{
+	Entity = InEntity;
+	const FArchetypeHandle ArchetypeHandle = EntitySubsystem.GetArchetypeForEntity(Entity);
+	check(ArchetypeHandle.IsValid());
+	Archetype = ArchetypeHandle.DataPtr.Get();
+	EntityHandle = Archetype->MakeEntityHandle(Entity);
+}
+
+void* FEntityView::GetComponentPtr(const UScriptStruct& ComponentType) const
+{
+	checkSlow(Archetype && EntityHandle.IsValid());
+	if (const int32* ComponentIndex = Archetype->GetComponentIndex(&ComponentType))
+	{
+		// failing the below Find means given entity's archetype is missing given ComponentType
+		return Archetype->GetComponentData(*ComponentIndex, EntityHandle);
+	}
+	return nullptr;
+}
+
+void* FEntityView::GetComponentPtrChecked(const UScriptStruct& ComponentType) const
+{
+	checkSlow(Archetype && EntityHandle.IsValid());
+	const int32 ComponentIndex = Archetype->GetComponentIndexChecked(&ComponentType);
+	return Archetype->GetComponentData(ComponentIndex, EntityHandle);
+}
+
+bool FEntityView::HasTag(const UScriptStruct& TagType) const
+{
+	checkSlow(Archetype && EntityHandle.IsValid());
+	return Archetype->HasTagType(&TagType);
+}
diff --git a/Engine/Plugins/Runtime/MassEntity/Source/MassEntity/Private/LWCCommandBuffer.cpp b/Engine/Plugins/Runtime/MassEntity/Source/MassEntity/Private/LWCCommandBuffer.cpp
new file mode 100644
--- /dev/null
+++ b/Engine/Plugins/Runtime/MassEntity/Source/MassEntity/Private/LWCCommandBuffer.cpp
@@ -0,0 +1,196 @@
+// Copyright Epic Games, Inc. All Rights Reserved.
+
+#include "LWCCommandBuffer.h"
+#include "HAL/IConsoleManager.h"
+#include "ProfilingDebugging/CsvProfiler.h"
+
+//////////////////////////////////////////////////////////////////////
+// FLWCCommandBuffer
+
+CSV_DEFINE_CATEGORY(MassEntities, true);
+CSV_DEFINE_CATEGORY(MassEntitiesCounters, true);
+
+namespace UE::FLWCCommand {
+
+#if CSV_PROFILER
+bool bEnableDetailedStats = false;
+
+FAutoConsoleVariableRef CVarEnableDetailedCommandStats(TEXT("massentities.EnableCommandDetailedStats"), bEnableDetailedStats,
+	TEXT("Set to true create a dedicated stat per type of command."), ECVF_Default);
+
+/** CSV stat names */
+static FString DefaultName = TEXT("Command");
+static TArray<FName> CommandFNames;
+
+/** CSV custom stat names (ANSI) */
+static const int32 MaxNameLength = 64;
+typedef ANSICHAR ANSIName[MaxNameLength];
+static ANSIName DefaultANSIName = "Command";
+static TArray<ANSIName> CommandANSINames;
+
+/**
+ * Provides valid names for CSV profiling.
+ * @param Entry is the view on the command
+ * @param OutName is the name to use for csv custom stats
+ * @param OutANSIName is the name to use for csv stats
+ */
+void GetCommandStatNames(FStructView Entry, FString& OutName, ANSIName*& OutANSIName)
+{
+	OutANSIName = &DefaultANSIName;
+	OutName = DefaultName;
+	if (!bEnableDetailedStats ||
+		!ensureMsgf(Entry.GetScriptStruct() != nullptr, TEXT("Unable to get stat name for an invalid entry. Using default name.")))
+	{
+		return;
+	}
+
+	const FName CommandFName = Entry.GetScriptStruct()->GetFName();
+	OutName = CommandFName.ToString();
+
+	const int32 Index = CommandFNames.Find(CommandFName);
+	if (Index == INDEX_NONE)
+	{
+		CommandFNames.Emplace(CommandFName);
+		OutANSIName = &CommandANSINames.AddZeroed_GetRef();
+		// Use prefix for easier parsing in reports
+		//const FString CounterName = FString::Printf(TEXT("Num%s"), *OutName);
+		//FMemory::Memcpy(OutANSIName, StringCast<ANSICHAR>(*CounterName).Get(), FMath::Min(CounterName.Len(), MaxNameLength - 1) * sizeof(ANSICHAR));
+		FMemory::Memcpy(OutANSIName, StringCast<ANSICHAR>(*OutName).Get(), FMath::Min(OutName.Len(), MaxNameLength - 1) * sizeof(ANSICHAR));
+	}
+	else
+	{
+		OutANSIName = &CommandANSINames[Index];
+	}
+}
+
+#endif
+} // UE::FLWCCommand
+
+void FLWCCommandBuffer::ReplayBufferAgainstSystem(UEntitySubsystem* System)
+{
+	check(System);
+
+	UE_MT_SCOPED_WRITE_ACCESS(PendingCommandsDetector);
+	PendingCommands.ForEach([System](FStructView Entry)
+	{
+		const FCommandBufferEntryBase* Command = Entry.GetPtr<FCommandBufferEntryBase>();
+		checkf(Command, TEXT("Either the entry is null or the command does not derive from FCommandBufferEntryBase"));
+
+#if CSV_PROFILER
+		using namespace UE::FLWCCommand;
+
+		// Extract name (default or detailed)
+		ANSIName* ANSIName = &DefaultANSIName;
+		FString Name = DefaultName;
+		GetCommandStatNames(Entry, Name, ANSIName);
+
+		// Push stats
+		FScopedCsvStat ScopedCsvStat(*ANSIName, CSV_CATEGORY_INDEX(MassEntities));
+		FCsvProfiler::RecordCustomStat(*Name, CSV_CATEGORY_INDEX(MassEntitiesCounters), 1, ECsvCustomStatOp::Accumulate);
+#endif // CSV_PROFILER
+
+		Command->Execute(*System);
+	});
+
+	// Using Clear() instead of Reset(), as otherwise the chunks moved into the PendingCommands in MoveAppend() can accumulate.
+	PendingCommands.Clear();
+
+	if (EntitiesToDestroy.Num())
+	{
+		TRACE_CPUPROFILER_EVENT_SCOPE_STR("ECS Deferred Commands Destroy Entities");
+		System->BatchDestroyEntities(MakeArrayView(EntitiesToDestroy));
+		EntitiesToDestroy.Reset();
+	}
+}
+
+void FLWCCommandBuffer::MoveAppend(FLWCCommandBuffer& Other)
+{
+	// @todo optimize, there surely a way to do faster then this.
+	UE_MT_SCOPED_READ_ACCESS(Other.PendingCommandsDetector);
+	if (Other.GetPendingCommandsCount() || Other.EntitiesToDestroy.Num())
+	{
+		FScopeLock Lock(&AppendingCommandsCS);
+		UE_MT_SCOPED_WRITE_ACCESS(PendingCommandsDetector);
+		PendingCommands.Append(MoveTemp(Other.PendingCommands));
+		EntitiesToDestroy.Append(MoveTemp(Other.EntitiesToDestroy));
+	}
+}
+
+//////////////////////////////////////////////////////////////////////
+// Command implementations
+
+void FCommandAddComponentInstance::Execute(UEntitySubsystem& System) const
+{
+	System.AddComponentInstanceListToEntity(TargetEntity, MakeArrayView(&Struct, 1));
+}
+
+void FCommandSwapTags::Execute(UEntitySubsystem& System) const
+{
+	if (System.IsEntityValid(TargetEntity) == false)
+	{
+		return;
+}
+
+	if (OldTagType && NewTagType)
+	{
+		System.SwapTagsForEntity(TargetEntity, OldTagType, NewTagType);
+	}
+	else if (OldTagType)
+	{
+		System.RemoveTagFromEntity(TargetEntity, OldTagType);
+	}
+	else if (NewTagType)
+	{
+		System.AddTagToEntity(TargetEntity, NewTagType);
+	}
+}
+
+void FBuildEntityFromComponentInstance::Execute(UEntitySubsystem& System) const
+{
+	System.BuildEntity(TargetEntity, MakeArrayView(&Struct, 1));
+}
+
+void FBuildEntityFromComponentInstances::Execute(UEntitySubsystem& System) const
+{
+	System.BuildEntity(TargetEntity, Instances);
+}
+
+void FCommandAddComponent::Execute(UEntitySubsystem& System) const
+{
+	System.AddComponentToEntity(TargetEntity, StructParam);
+}
+
+void FCommandRemoveComponent::Execute(UEntitySubsystem& System) const
+{
+	System.RemoveComponentFromEntity(TargetEntity, StructParam);
+}
+
+void FCommandAddComponentList::Execute(UEntitySubsystem& System) const
+{
+	System.AddComponentListToEntity(TargetEntity, ComponentList);
+}
+
+void FCommandRemoveComponentList::Execute(UEntitySubsystem& System) const
+{
+	System.RemoveComponentListFromEntity(TargetEntity, ComponentList);
+}
+
+void FCommandAddTag::Execute(UEntitySubsystem & System) const 
+{
+	if (System.IsEntityValid(TargetEntity) == false)
+{
+		return;
+	}
+
+	System.AddTagToEntity(TargetEntity, StructParam);
+}
+
+void FCommandRemoveTag::Execute(UEntitySubsystem& System) const
+{
+	System.RemoveTagFromEntity(TargetEntity, StructParam);
+}
+
+void FCommandRemoveComposition::Execute(UEntitySubsystem& System) const
+{
+	System.RemoveCompositionFromEntity(TargetEntity, Descriptor);
+}
diff --git a/Engine/Plugins/Runtime/MassEntity/Source/MassEntity/Private/MassEntityDebug.cpp b/Engine/Plugins/Runtime/MassEntity/Source/MassEntity/Private/MassEntityDebug.cpp
new file mode 100644
--- /dev/null
+++ b/Engine/Plugins/Runtime/MassEntity/Source/MassEntity/Private/MassEntityDebug.cpp
@@ -0,0 +1,34 @@
+// Copyright Epic Games, Inc. All Rights Reserved.
+
+#include "MassEntityDebug.h"
+#include "Misc/OutputDevice.h"
+#include "MassProcessor.h"
+
+
+DEFINE_ENUM_TO_STRING(EPipeProcessingPhase);
+
+namespace UE { namespace Pipe { namespace Debug {
+
+void DebugOutputDescription(TConstArrayView<UPipeProcessor*> Processors, FOutputDevice& Ar)
+{
+#if WITH_PIPE_DEBUG
+	const bool bAutoLineEnd = Ar.GetAutoEmitLineTerminator();
+	Ar.SetAutoEmitLineTerminator(false);
+	for (const UPipeProcessor* Proc : Processors)
+	{
+		if (Proc)
+		{
+			Proc->DebugOutputDescription(Ar);
+			Ar.Logf(TEXT("\n"));
+		}
+		else
+		{
+			Ar.Logf(TEXT("NULL\n"));
+		}
+	}
+	Ar.SetAutoEmitLineTerminator(bAutoLineEnd);
+#endif // WITH_PIPE_DEBUG
+}
+
+}}} // namespace UE::Pipe::Debug
+
diff --git a/Engine/Plugins/Runtime/MassEntity/Source/MassEntity/Private/MassEntityModule.cpp b/Engine/Plugins/Runtime/MassEntity/Source/MassEntity/Private/MassEntityModule.cpp
new file mode 100644
--- /dev/null
+++ b/Engine/Plugins/Runtime/MassEntity/Source/MassEntity/Private/MassEntityModule.cpp
@@ -0,0 +1,62 @@
+// Copyright Epic Games, Inc. All Rights Reserved.
+
+#include "CoreMinimal.h"
+#include "IMassEntityModule.h"
+#if WITH_UNREAL_DEVELOPER_TOOLS
+#include "MessageLogModule.h"
+#include "Engine/World.h"
+#include "Logging/MessageLog.h"
+#endif // WITH_UNREAL_DEVELOPER_TOOLS
+
+#define LOCTEXT_NAMESPACE "Pipe"
+
+class FPipeModuleModule : public IPipeModule
+{
+	/** IModuleInterface implementation */
+	virtual void StartupModule() override;
+	virtual void ShutdownModule() override;
+#if WITH_UNREAL_DEVELOPER_TOOLS
+	static void OnWorldCleanup(UWorld* /*World*/, bool /*bSessionEnded*/, bool /*bCleanupResources*/);
+	FDelegateHandle OnWorldCleanupHandle;
+#endif // WITH_UNREAL_DEVELOPER_TOOLS
+};
+
+IMPLEMENT_MODULE(FPipeModuleModule, Pipe)
+
+
+void FPipeModuleModule::StartupModule()
+{
+	// This code will execute after your module is loaded into memory (but after global variables are initialized, of course.)
+
+#if WITH_UNREAL_DEVELOPER_TOOLS
+	{
+		FMessageLogModule& MessageLogModule = FModuleManager::LoadModuleChecked<FMessageLogModule>("MessageLog");
+		FMessageLogInitializationOptions InitOptions;
+		InitOptions.bShowPages = true;
+		InitOptions.bShowFilters = true;
+		MessageLogModule.RegisterLogListing("Pipe", LOCTEXT("Pipe", "Pipe"), InitOptions);
+
+		OnWorldCleanupHandle = FWorldDelegates::OnWorldCleanup.AddStatic(&FPipeModuleModule::OnWorldCleanup);
+	}
+#endif // WITH_UNREAL_DEVELOPER_TOOLS
+}
+
+
+void FPipeModuleModule::ShutdownModule()
+{
+	// This function may be called during shutdown to clean up your module.  For modules that support dynamic reloading,
+	// we call this function before unloading the module.
+#if WITH_UNREAL_DEVELOPER_TOOLS
+	FWorldDelegates::OnWorldCleanup.Remove(OnWorldCleanupHandle);
+#endif // WITH_UNREAL_DEVELOPER_TOOLS
+}
+
+#if WITH_UNREAL_DEVELOPER_TOOLS
+void FPipeModuleModule::OnWorldCleanup(UWorld* /*World*/, bool /*bSessionEnded*/, bool /*bCleanupResources*/)
+{
+	// clearing out messages from the world being cleaned up
+	FMessageLog("Pipe").NewPage(FText::FromString(TEXT("Pipe")));
+}
+#endif // WITH_UNREAL_DEVELOPER_TOOLS
+
+#undef LOCTEXT_NAMESPACE 
diff --git a/Engine/Plugins/Runtime/MassEntity/Source/MassEntity/Private/MassEntitySettings.cpp b/Engine/Plugins/Runtime/MassEntity/Source/MassEntity/Private/MassEntitySettings.cpp
new file mode 100644
--- /dev/null
+++ b/Engine/Plugins/Runtime/MassEntity/Source/MassEntity/Private/MassEntitySettings.cpp
@@ -0,0 +1,188 @@
+// Copyright Epic Games, Inc. All Rights Reserved.
+
+#include "MassEntitySettings.h"
+#include "MassProcessingPhase.h"
+#include "VisualLogger/VisualLogger.h"
+#include "UObject/UObjectHash.h"
+#include "MassEntityDebug.h"
+#include "Misc/CoreDelegates.h"
+#if WITH_EDITOR
+#include "ObjectEditorUtils.h"
+#endif // WITH_EDITOR
+
+//----------------------------------------------------------------------//
+//  UPipeSettings
+//----------------------------------------------------------------------//
+UPipeSettings::UPipeSettings(const FObjectInitializer& ObjectInitializer)
+	: Super(ObjectInitializer)
+{
+	for (int i = 0; i < (int)EPipeProcessingPhase::MAX; ++i)
+	{
+		ProcessingPhasesConfig[i].PhaseName = *EnumToString(EPipeProcessingPhase(i));
+	}
+
+	FCoreDelegates::OnPostEngineInit.AddUObject(this, &UPipeSettings::BuildProcessorListAndPhases);
+}
+
+void UPipeSettings::BeginDestroy()
+{
+	FCoreDelegates::OnPostEngineInit.RemoveAll(this);
+	Super::BeginDestroy();
+}
+
+void UPipeSettings::BuildProcessorListAndPhases()
+{
+	if (bInitialized)
+	{
+		return;
+	}
+
+	BuildProcessorList();
+	BuildPhases();
+	bInitialized = true;
+}
+
+void UPipeSettings::BuildPhases()
+{
+#if WITH_EDITOR
+	for (int i = 0; i < int(EPipeProcessingPhase::MAX); ++i)
+	{
+		FPipeProcessingPhaseConfig& PhaseConfig = ProcessingPhasesConfig[i];
+		PhaseConfig.PhaseProcessor = NewObject<UPipeCompositeProcessor>(this, PhaseConfig.PhaseGroupClass
+			, *FString::Printf(TEXT("ProcessingPhase_%s"), *PhaseConfig.PhaseName.ToString()));
+		PhaseConfig.PhaseProcessor->SetGroupName(PhaseConfig.PhaseName);
+		PhaseConfig.PhaseProcessor->SetProcessingPhase(EPipeProcessingPhase(i));
+		const FString PhaseDumpDependencyGraphFileName = !DumpDependencyGraphFileName.IsEmpty() ? DumpDependencyGraphFileName + TEXT("_") + PhaseConfig.PhaseName.ToString() : FString();
+		PhaseConfig.PhaseProcessor->CopyAndSort(PhaseConfig, PhaseDumpDependencyGraphFileName);
+		
+		FStringOutputDevice Ar;
+		PhaseConfig.PhaseProcessor->DebugOutputDescription(Ar);
+		PhaseConfig.Description = MoveTemp(Ar);
+	}
+#endif // WITH_EDITOR
+}
+
+void UPipeSettings::BuildProcessorList()
+{
+	ProcessorCDOs.Reset();
+	for (FPipeProcessingPhaseConfig& PhaseConfig : ProcessingPhasesConfig)
+	{
+		PhaseConfig.ProcessorCDOs.Reset();
+	}
+
+	TArray<UClass*> SubClassess;
+	GetDerivedClasses(UPipeProcessor::StaticClass(), SubClassess);
+
+	for (int i = SubClassess.Num() - 1; i; --i)
+	{
+		if (SubClassess[i]->HasAnyClassFlags(CLASS_Abstract))
+		{
+			continue;
+		}
+
+		UPipeProcessor* ProcessorCDO = GetMutableDefault<UPipeProcessor>(SubClassess[i]);
+		// we explicitly restrict adding UPipeCompositeProcessor. If needed by specific project a derived class can be added
+		if (ProcessorCDO && SubClassess[i] != UPipeCompositeProcessor::StaticClass()
+#if WITH_EDITOR
+			&& ProcessorCDO->ShouldShowUpInSettings()
+#endif // WITH_EDITOR
+		)
+		{
+			ProcessorCDOs.Add(ProcessorCDO);
+			if (ProcessorCDO->ShouldAutoAddToGlobalList())
+			{
+				ProcessingPhasesConfig[int(ProcessorCDO->GetProcessingPhase())].ProcessorCDOs.Add(ProcessorCDO);
+			}
+		}
+	}
+
+	ProcessorCDOs.Sort([](UPipeProcessor& LHS, UPipeProcessor& RHS) {
+		return LHS.GetName().Compare(RHS.GetName()) < 0;
+	});
+}
+
+void UPipeSettings::AddToActiveProcessorsList(TSubclassOf<UPipeProcessor> ProcessorClass)
+{
+	if (UPipeProcessor* ProcessorCDO = GetMutableDefault<UPipeProcessor>(ProcessorClass))
+	{
+		if (ProcessorClass == UPipeCompositeProcessor::StaticClass())
+		{
+			UE_VLOG_UELOG(this, LogPipe, Log, TEXT("%s adding PipeCompositeProcessor to the global processor list is unsupported"), ANSI_TO_TCHAR(__FUNCTION__));
+		}
+		else if (ProcessorClass->HasAnyClassFlags(CLASS_Abstract))
+		{
+			UE_VLOG_UELOG(this, LogPipe, Log, TEXT("%s unable to add %s due to it being an abstract class"), ANSI_TO_TCHAR(__FUNCTION__), *ProcessorClass->GetName());
+		}
+		else if (ProcessorCDOs.Find(ProcessorCDO) != INDEX_NONE)
+		{
+			UE_VLOG_UELOG(this, LogPipe, Log, TEXT("%s already in global processor list"), *ProcessorCDO->GetName());
+		}
+		else 
+		{
+			ensureMsgf(ProcessorCDO->ShouldAutoAddToGlobalList() == false, TEXT("%s missing from the global list while it's already marked to be auto-added"), *ProcessorCDO->GetName());
+			ProcessorCDOs.Add(ProcessorCDO);
+			ProcessorCDO->SetShouldAutoRegisterWithGlobalList(true);
+		}
+	}
+}
+
+const FPipeProcessingPhaseConfig* UPipeSettings::GetProcessingPhasesConfig()
+{
+	BuildProcessorListAndPhases();
+	return ProcessingPhasesConfig;
+}
+
+#if WITH_EDITOR
+void UPipeSettings::PostEditChangeProperty(FPropertyChangedEvent& PropertyChangedEvent)
+{
+	// @todo replacing a valid entry in the array is not handled right now. 
+
+	static const FName PipeName = TEXT("Pipe");
+	static const FName ProcessorCDOsName = GET_MEMBER_NAME_CHECKED(UPipeSettings, ProcessorCDOs);
+
+	Super::PostEditChangeProperty(PropertyChangedEvent);
+
+	if (PropertyChangedEvent.ChangeType == EPropertyChangeType::ArrayAdd)
+	{
+		// ignore adding elements to arrays since it would be 'None' at first
+		return;
+	}
+
+	if (PropertyChangedEvent.Property)
+	{
+		const FName PropName = PropertyChangedEvent.Property->GetFName();
+		if (PropName == ProcessorCDOsName)
+		{
+			BuildProcessorList();
+		}
+
+		BuildPhases();
+		OnSettingsChange.Broadcast(PropertyChangedEvent);
+	}
+}
+
+void UPipeSettings::PostEditChangeChainProperty(struct FPropertyChangedChainEvent& PropertyChangedEvent)
+{
+	static const FName AutoRegisterName = TEXT("bAutoRegisterWithProcessingPhases");
+
+	Super::PostEditChangeChainProperty(PropertyChangedEvent);
+
+	FProperty* Property = PropertyChangedEvent.Property;
+	FProperty* MemberProperty = nullptr;
+	FEditPropertyChain::TDoubleLinkedListNode* LastPropertyNode = PropertyChangedEvent.PropertyChain.GetActiveMemberNode();
+	while (LastPropertyNode && LastPropertyNode->GetNextNode())
+	{
+		LastPropertyNode = LastPropertyNode->GetNextNode();
+	}
+
+	if (LastPropertyNode)
+	{
+		MemberProperty = LastPropertyNode->GetValue();
+	}
+
+	if (MemberProperty && MemberProperty->GetFName() == AutoRegisterName)
+	{
+		BuildProcessorList();
+	}
+}
+#endif // WITH_EDITOR
diff --git a/Engine/Plugins/Runtime/MassEntity/Source/MassEntity/Private/MassEntitySystem.cpp b/Engine/Plugins/Runtime/MassEntity/Source/MassEntity/Private/MassEntitySystem.cpp
new file mode 100644
--- /dev/null
+++ b/Engine/Plugins/Runtime/MassEntity/Source/MassEntity/Private/MassEntitySystem.cpp
@@ -0,0 +1,114 @@
+// Copyright Epic Games, Inc. All Rights Reserved.
+
+#include "MassEntitySystem.h"
+#include "MassEntityTypes.h"
+#include "Engine/Engine.h"
+#include "Engine/World.h"
+#include "Containers/UnrealString.h"
+#include "Logging/LogScopedVerbosityOverride.h"
+
+#if WITH_PIPE_DEBUG && WITH_AGGREGATETICKING_DEBUG
+namespace UE::Pipe::Debug
+{
+// @todo rename to `brute.*` when time comes
+
+FAutoConsoleCommandWithWorldArgsAndOutputDevice PrintEntityFragmentsCmd(
+	TEXT("pipe.PrintEntityFragments"),
+	TEXT("Prints all fragment types and values (uproperties) for the specified Entity index"),
+	FConsoleCommandWithWorldArgsAndOutputDeviceDelegate::CreateLambda(
+		[](const TArray<FString>& Params, UWorld* World, FOutputDevice& Ar)
+		{
+			if (UPipeEntitySubsystem* EntitySystem = UPipeEntitySubsystem::GetCurrent(World))
+			{
+				int32 Index = INDEX_NONE;
+				if (LexTryParseString<int32>(Index, *Params[0]))
+				{
+					EntitySystem->DebugPrintEntity(Index, Ar);
+				}
+				else
+				{
+					Ar.Logf(ELogVerbosity::Error, TEXT("Entity index parameter must be an integer"));
+				}
+			}
+			else
+			{
+				Ar.Logf(ELogVerbosity::Error, TEXT("Failed to find PipeEntitySubsystem for world %s"), *GetPathNameSafe(World));
+			}
+		})
+);
+
+FAutoConsoleCommandWithWorldArgsAndOutputDevice LogArchetypesCmd(
+	TEXT("pipe.LogArchetypes"),
+	TEXT("Dumps description of archetypes to log. Optional parameter controls whether to include or exclude non-occupied archetypes. Defaults to 'include'."),
+	FConsoleCommandWithWorldArgsAndOutputDeviceDelegate::CreateLambda([](const TArray<FString>& Params, UWorld*, FOutputDevice& Ar)
+	{
+		const TIndirectArray<FWorldContext>& WorldContexts = GEngine->GetWorldContexts();
+		for (const FWorldContext& Context : WorldContexts)
+		{
+			UWorld* World = Context.World();
+			if (World == nullptr || World->IsPreviewWorld())
+			{
+				continue;
+			}
+
+			Ar.Logf(ELogVerbosity::Log, TEXT("Dumping description of archetypes for world: %s (%s - %s)"),
+				*GetPathNameSafe(World),
+				LexToString(World->WorldType),
+				*ToString(World->GetNetMode()));
+
+			if (UPipeEntitySubsystem* EntitySystem = UPipeEntitySubsystem::GetCurrent(World))
+			{
+				bool bIncludeEmpty = true;
+				if (Params.Num())
+				{
+					LexTryParseString(bIncludeEmpty, *Params[0]);
+				}
+				Ar.Logf(ELogVerbosity::Log, TEXT("Include empty archetypes: %s"), bIncludeEmpty ? TEXT("TRUE") : TEXT("FALSE"));
+				EntitySystem->DebugGetArchetypesStringDetails(Ar, bIncludeEmpty);
+			}
+			else
+			{
+				Ar.Logf(ELogVerbosity::Error, TEXT("Failed to find PipeEntitySubsystem for world: %s (%s - %s)"),
+					*GetPathNameSafe(World),
+					LexToString(World->WorldType),
+					*ToString(World->GetNetMode()));
+			}
+		}
+	})
+);
+
+// @todo these console commands will be reparented to "massentities" domain once we rename and shuffle the modules around 
+FAutoConsoleCommandWithWorld RecacheQueries(
+	TEXT("pipe.RecacheQueries"),
+	TEXT("Forces EntityQueries to recache their valid archetypes"),
+	FConsoleCommandWithWorldDelegate::CreateLambda([](UWorld* InWorld)
+		{
+			if (UPipeEntitySubsystem* System = UPipeEntitySubsystem::GetCurrent(InWorld))
+			{
+				System->DebugForceArchetypeDataVersionBump();
+			}
+		}
+	));
+
+FAutoConsoleCommandWithWorldArgsAndOutputDevice LogFragmentSizes(
+	TEXT("pipe.LogFragmentSizes"),
+	TEXT("Logs all the fragment types being used along with their sizes."),
+	FConsoleCommandWithWorldArgsAndOutputDeviceDelegate::CreateLambda([](const TArray<FString>& Params, UWorld* World, FOutputDevice& Ar)
+		{
+			for (const TWeakObjectPtr<const UScriptStruct>& WeakStruct : FLWComponentBitSet::DebugGetAllStructTypes())
+			{
+				if (const UScriptStruct* StructType = WeakStruct.Get())
+				{
+					Ar.Logf(ELogVerbosity::Log, TEXT("%s, size: %d"), *StructType->GetName(), StructType->GetStructureSize());
+				}
+			}
+		})
+	);
+
+} // namespace UE::Pipe::Debug
+#endif // WITH_PIPE_DEBUG && WITH_AGGREGATETICKING_DEBUG
+
+TFunctionRef<UPipeEntitySubsystem* (UWorld*)> UPipeEntitySubsystem::InstanceGetter = [](UWorld* World)
+{ 
+	return World ? World->GetSubsystem<UPipeEntitySubsystem>() : nullptr;
+};
diff --git a/Engine/Plugins/Runtime/MassEntity/Source/MassEntity/Private/MassEntityTypes.cpp b/Engine/Plugins/Runtime/MassEntity/Source/MassEntity/Private/MassEntityTypes.cpp
new file mode 100644
--- /dev/null
+++ b/Engine/Plugins/Runtime/MassEntity/Source/MassEntity/Private/MassEntityTypes.cpp
@@ -0,0 +1,226 @@
+// Copyright Epic Games, Inc. All Rights Reserved.
+
+#include "MassEntityTypes.h"
+#include "MassProcessor.h"
+#include "MassSchematic.h"
+#include "Misc/OutputDevice.h"
+#include "MassEntityUtils.h"
+#include "Engine/EngineBaseTypes.h"
+#include "VisualLogger/VisualLogger.h"
+#include "MassEntityDebug.h"
+
+DEFINE_LOG_CATEGORY(LogPipe);
+
+//----------------------------------------------------------------------//
+//  FPipeContext
+//----------------------------------------------------------------------//
+FPipeContext::FPipeContext(UEntitySubsystem& InEntities, const float InDeltaSeconds)
+	: EntitySubsystem(&InEntities), DeltaSeconds(InDeltaSeconds)
+{
+
+}
+
+//----------------------------------------------------------------------//
+//  FRuntimePipeline
+//----------------------------------------------------------------------//
+void FRuntimePipeline::Reset()
+{
+	Processors.Reset();
+}
+
+void FRuntimePipeline::Initialize(UObject& Owner)
+{
+	for (UPipeProcessor* Proc : Processors)
+	{
+		if (Proc)
+		{
+			REDIRECT_OBJECT_TO_VLOG(Proc, &Owner);
+			Proc->Initialize(Owner);
+		}
+	}
+}
+
+void FRuntimePipeline::SetProcessors(TArray<UPipeProcessor*>&& InProcessors)
+{
+	Processors = InProcessors;
+}
+
+void FRuntimePipeline::InitializeFromSchematics(TConstArrayView<TSoftObjectPtr<UPipeSchematic>> Schematics, UObject& InOwner)
+{
+	Reset();
+	
+	// @todo we'll sometimes end up with duplicated PipeProcessors in the resulting array. We need to come up with a consistent policy for handling that 
+	for (const TSoftObjectPtr<UPipeSchematic>& Schematic : Schematics)
+	{
+		UPipeSchematic* Pipe = Schematic.LoadSynchronous();
+		if (Pipe)
+		{
+			AppendOrOverrideRuntimeProcessorCopies(Pipe->GetProcessors(), InOwner);
+		}
+		else
+		{
+			UE_LOG(LogPipe, Error, TEXT("Unable to resolve PipeSchematic %s while creating FRuntimePipeline"), *Schematic.GetLongPackageName());
+		}
+	}
+
+	Initialize(InOwner);
+}
+
+void FRuntimePipeline::CreateFromArray(TConstArrayView<const UPipeProcessor*> InProcessors, UObject& InOwner)
+{
+	Reset();
+	AppendOrOverrideRuntimeProcessorCopies(InProcessors, InOwner);
+}
+
+void FRuntimePipeline::InitializeFromArray(TConstArrayView<const UPipeProcessor*> InProcessors, UObject& InOwner)
+{
+	CreateFromArray(InProcessors, InOwner);
+	Initialize(InOwner);
+}
+
+void FRuntimePipeline::InitializeFromClassArray(TConstArrayView<TSubclassOf<UPipeProcessor>> InProcessorClasses, UObject& InOwner)
+{
+	Reset();
+
+	const UWorld* World = InOwner.GetWorld();
+	const EProcessorExecutionFlags WorldExecutionFlags = World ? UE::Pipe::Utils::GetProcessorExecutionFlagsForWold(*World) : EProcessorExecutionFlags::All;
+
+	for (const TSubclassOf<UPipeProcessor>& ProcessorClass : InProcessorClasses)
+	{
+		if (ProcessorClass)
+		{
+			UPipeProcessor* CDO = ProcessorClass.GetDefaultObject();
+			if (CDO && CDO->ShouldExecute(WorldExecutionFlags))
+			{
+				UPipeProcessor* ProcInstance = NewObject<UPipeProcessor>(&InOwner, ProcessorClass);
+				Processors.Add(ProcInstance);
+			}
+			else
+			{
+				UE_CVLOG(CDO, &InOwner, LogPipe, Log, TEXT("Skipping %s due to ExecutionFlags"), *CDO->GetName());
+			}
+		}
+	}
+
+	Initialize(InOwner);
+}
+
+bool FRuntimePipeline::HasProcessorOfExactClass(TSubclassOf<UPipeProcessor> InClass) const
+{
+	UClass* TestClass = InClass.Get();
+	return Processors.FindByPredicate([TestClass](const UPipeProcessor* Proc){ return Proc != nullptr && Proc->GetClass() == TestClass; }) != nullptr;
+}
+
+void FRuntimePipeline::AppendUniqueRuntimeProcessorCopies(TConstArrayView<const UPipeProcessor*> InProcessors, UObject& InOwner)
+{
+	const UWorld* World = InOwner.GetWorld();
+	const EProcessorExecutionFlags WorldExecutionFlags = World ? UE::Pipe::Utils::GetProcessorExecutionFlagsForWold(*World) : EProcessorExecutionFlags::All;
+	const int32 StartingCount = Processors.Num();
+		
+	for (const UPipeProcessor* Proc : InProcessors)
+	{
+		if (Proc && Proc->ShouldExecute(WorldExecutionFlags)
+			&& (Proc->AllowDuplicates() || (HasProcessorOfExactClass(Proc->GetClass()) == false)))
+		{
+			// unfortunately the const cast is required since NewObject doesn't support const Template object
+			UPipeProcessor* ProcCopy = NewObject<UPipeProcessor>(&InOwner, Proc->GetClass(), FName(), RF_NoFlags, const_cast<UPipeProcessor*>(Proc));
+			Processors.Add(ProcCopy);
+		}
+#if WITH_PIPE_DEBUG
+		else if (Proc)
+		{
+			if (Proc->ShouldExecute(WorldExecutionFlags) == false)
+			{
+				UE_VLOG(&InOwner, LogPipe, Log, TEXT("Skipping %s due to ExecutionFlags"), *Proc->GetName());
+			}
+			else if (Proc->AllowDuplicates() == false)
+			{
+				UE_VLOG(&InOwner, LogPipe, Log, TEXT("Skipping %s due to it being a duplicate"), *Proc->GetName());
+			}
+		}
+#endif // WITH_PIPE_DEBUG
+	}
+
+	for (int32 NewProcIndex = StartingCount; NewProcIndex < Processors.Num(); ++NewProcIndex)
+	{
+		UPipeProcessor* Proc = Processors[NewProcIndex];
+		check(Proc);
+		REDIRECT_OBJECT_TO_VLOG(Proc, &InOwner);
+		Proc->Initialize(InOwner);
+	}
+}
+
+void FRuntimePipeline::AppendOrOverrideRuntimeProcessorCopies(TConstArrayView<const UPipeProcessor*> InProcessors, UObject& InOwner)
+{
+	const UWorld* World = InOwner.GetWorld();
+	const EProcessorExecutionFlags WorldExecutionFlags = World ? UE::Pipe::Utils::GetProcessorExecutionFlagsForWold(*World) : EProcessorExecutionFlags::All;
+	
+	for (const UPipeProcessor* Proc : InProcessors)
+	{
+		if (Proc && Proc->ShouldExecute(WorldExecutionFlags))
+		{
+			// unfortunately the const cast is required since NewObject doesn't support const Template object
+			UPipeProcessor* ProcCopy = NewObject<UPipeProcessor>(&InOwner, Proc->GetClass(), FName(), RF_NoFlags, const_cast<UPipeProcessor*>(Proc));
+			check(ProcCopy);
+
+			if (ProcCopy->AllowDuplicates())
+			{
+				// we don't care if there are instances of this class in Processors already
+				Processors.Add(ProcCopy);
+			}
+			else 
+			{
+				const UClass* TestClass = Proc->GetClass();
+				UPipeProcessor** PrevProcessor = Processors.FindByPredicate([TestClass, ProcCopy](const UPipeProcessor* Proc) {
+					return Proc != nullptr && Proc->GetClass() == TestClass;
+				});
+
+				if (PrevProcessor)
+				{
+					*PrevProcessor = ProcCopy;
+				}
+				else
+				{
+					Processors.Add(ProcCopy);
+				}
+			}
+		}
+		else
+		{
+			UE_CVLOG(Proc, &InOwner, LogPipe, Log, TEXT("Skipping %s due to ExecutionFlags"), *Proc->GetName());
+		}
+	}
+}
+
+void FRuntimePipeline::AppendProcessor(UPipeProcessor& Processor)
+{
+	Processors.Add(&Processor);
+}
+
+UPipeCompositeProcessor* FRuntimePipeline::FindTopLevelGroupByName(FName GroupName)
+{
+	for (UPipeProcessor* Processor : Processors)
+	{
+		UPipeCompositeProcessor* CompositeProcessor = Cast<UPipeCompositeProcessor>(Processor);
+		if (CompositeProcessor && CompositeProcessor->GetGroupName() == GroupName)
+		{
+			return CompositeProcessor;
+		}
+	}
+	return nullptr;
+}
+
+void FRuntimePipeline::DebugOutputDescription(FOutputDevice& Ar) const
+{
+	UE::Pipe::Debug::DebugOutputDescription(Processors, Ar);
+}
+
+uint32 GetTypeHash(const FRuntimePipeline& Instance)
+{ 
+	uint32 Hash = 0;
+	for (const UPipeProcessor* Proc : Instance.Processors)
+	{
+		Hash = HashCombine(Hash, PointerHash(Proc));
+	}
+	return Hash;
+}
diff --git a/Engine/Plugins/Runtime/MassEntity/Source/MassEntity/Private/MassEntityUtils.cpp b/Engine/Plugins/Runtime/MassEntity/Source/MassEntity/Private/MassEntityUtils.cpp
new file mode 100644
--- /dev/null
+++ b/Engine/Plugins/Runtime/MassEntity/Source/MassEntity/Private/MassEntityUtils.cpp
@@ -0,0 +1,34 @@
+// Copyright Epic Games, Inc. All Rights Reserved.
+
+#include "MassEntityUtils.h"
+#include "Engine/EngineBaseTypes.h"
+#include "Engine/World.h"
+
+namespace UE { namespace Pipe { namespace Utils
+{
+
+EProcessorExecutionFlags GetProcessorExecutionFlagsForWold(const UWorld& World)
+{
+	EProcessorExecutionFlags ExecutionFlags = EProcessorExecutionFlags::None;
+	const ENetMode NetMode = World.GetNetMode();
+	switch (NetMode)
+	{
+		case NM_ListenServer:
+			ExecutionFlags = EProcessorExecutionFlags::Client | EProcessorExecutionFlags::Server;
+			break;
+		case NM_DedicatedServer:
+			ExecutionFlags = EProcessorExecutionFlags::Server;
+			break;
+		case NM_Client:
+			ExecutionFlags = EProcessorExecutionFlags::Client;
+			break;
+		default:
+			check(NetMode == NM_Standalone);
+			ExecutionFlags = EProcessorExecutionFlags::Standalone;
+			break;
+	}
+
+	return ExecutionFlags;
+}
+
+} } }// namespace UE::Pipe::Utils
\ No newline at end of file
diff --git a/Engine/Plugins/Runtime/MassEntity/Source/MassEntity/Private/MassExecutor.cpp b/Engine/Plugins/Runtime/MassEntity/Source/MassEntity/Private/MassExecutor.cpp
new file mode 100644
--- /dev/null
+++ b/Engine/Plugins/Runtime/MassEntity/Source/MassEntity/Private/MassExecutor.cpp
@@ -0,0 +1,190 @@
+// Copyright Epic Games, Inc. All Rights Reserved.
+
+#include "MassExecutor.h"
+#include "MassEntityTypes.h"
+#include "MassProcessor.h"
+#include "LWCCommandBuffer.h"
+#include "ProfilingDebugging/CpuProfilerTrace.h"
+
+namespace UE::Pipe::Executor
+{
+void Run(FRuntimePipeline& RuntimePipeline, FPipeContext& PipeContext)
+{
+	if (!ensure(PipeContext.EntitySubsystem) || 
+		!ensure(PipeContext.DeltaSeconds >= 0.f) ||
+		!ensure(RuntimePipeline.Processors.Find(nullptr) == INDEX_NONE))
+	{
+		return;
+	}
+
+	TRACE_CPUPROFILER_EVENT_SCOPE_STR("PipeExecutor Run Pipeline")
+	RunProcessorsView(RuntimePipeline.Processors, PipeContext);
+}
+
+void RunSparse(FRuntimePipeline& RuntimePipeline, FPipeContext& PipeContext, FArchetypeHandle Archetype, TConstArrayView<FLWEntity> Entities)
+{
+	if (!ensure(PipeContext.EntitySubsystem) ||
+		!ensure(RuntimePipeline.Processors.Find(nullptr) == INDEX_NONE) ||
+		RuntimePipeline.Processors.Num() == 0 ||
+		!ensureMsgf(Archetype.IsValid(), TEXT("The Archetype passed in to UE::Pipe::Executor::RunSparse is invalid")))
+	{
+		return;
+	}
+
+	TRACE_CPUPROFILER_EVENT_SCOPE_STR("PipeExecutor RunSparseEntities");
+
+	const FArchetypeChunkCollection ChunkCollection(Archetype, Entities);
+	RunProcessorsView(RuntimePipeline.Processors, PipeContext, &ChunkCollection);
+}
+
+void RunSparse(FRuntimePipeline& RuntimePipeline, FPipeContext& PipeContext, const FArchetypeChunkCollection& ChunkCollection)
+{
+	if (!ensure(PipeContext.EntitySubsystem) ||
+		!ensure(RuntimePipeline.Processors.Find(nullptr) == INDEX_NONE) ||
+		RuntimePipeline.Processors.Num() == 0 ||
+		!ensureMsgf(ChunkCollection.GetArchetype().IsValid(), TEXT("The Archetype of ChunkCollection passed in to UE::Pipe::Executor::RunSparse is invalid")))
+	{
+		return;
+	}
+
+	TRACE_CPUPROFILER_EVENT_SCOPE_STR("PipeExecutor RunSparse");
+
+	RunProcessorsView(RuntimePipeline.Processors, PipeContext, &ChunkCollection);
+}
+
+void Run(UPipeProcessor& Processor, FPipeContext& PipeContext)
+{
+	if (!ensure(PipeContext.EntitySubsystem) || !ensure(PipeContext.DeltaSeconds >= 0.f))
+	{
+		return;
+	}
+
+	TRACE_CPUPROFILER_EVENT_SCOPE_STR("PipeExecutor Run")
+
+	UPipeProcessor* ProcPtr = &Processor;
+	RunProcessorsView(MakeArrayView(&ProcPtr, 1), PipeContext);
+}
+
+void RunProcessorsView(TArrayView<UPipeProcessor*> Processors, FPipeContext& PipeContext, const FArchetypeChunkCollection* ChunkCollection)
+{
+	TRACE_CPUPROFILER_EVENT_SCOPE(RunProcessorsView);
+
+	if (PipeContext.EntitySubsystem == nullptr)
+	{
+		UE_LOG(LogPipe, Error, TEXT("%s PipeContext.EntitySubsystem is null. Baling out."), ANSI_TO_TCHAR(__FUNCTION__));
+		return;
+	}
+#if WITH_PIPE_DEBUG
+	if (Processors.Find(nullptr) != INDEX_NONE)
+	{
+		UE_LOG(LogPipe, Error, TEXT("%s input Processors contains nullptr. Baling out."), ANSI_TO_TCHAR(__FUNCTION__));
+		return;
+	}
+#endif // WITH_PIPE_DEBUG
+
+	TRACE_CPUPROFILER_EVENT_SCOPE_STR("PipeExecutor RunProcessorsView")
+
+	FLWComponentSystemExecutionContext ExecutionContext(PipeContext.DeltaSeconds);
+	if (ChunkCollection)
+	{
+		ExecutionContext.SetChunkCollection(*ChunkCollection);
+	}
+	// manually creating a new command buffer to let the default one still be used by code unaware of pipe processing
+	ExecutionContext.SetDeferredCommandBuffer(MakeShareable(new FLWCCommandBuffer()));
+	ExecutionContext.SetFlushDeferredCommands(false);
+	ExecutionContext.SetAuxData(PipeContext.AuxData);
+
+	{
+		TRACE_CPUPROFILER_EVENT_SCOPE_STR("Execute Processors")
+		
+		UEntitySubsystem::FScopedProcessing ProcessingScope = PipeContext.EntitySubsystem->NewProcessingScope();
+
+		for (UPipeProcessor* Proc : Processors)
+		{
+			Proc->CallExecute(*PipeContext.EntitySubsystem, ExecutionContext);
+		}
+	}
+	
+	{		
+		TRACE_CPUPROFILER_EVENT_SCOPE_STR("Flush Deferred Commands")
+		
+		ExecutionContext.SetFlushDeferredCommands(true);
+		// append the commands added from other, non-processor sources (like MassAgentSubsystem)
+		ExecutionContext.Defer().MoveAppend(PipeContext.EntitySubsystem->Defer());
+		ExecutionContext.FlushDeferred(*PipeContext.EntitySubsystem);
+	}
+}
+
+struct FPipeExecutorDoneTask
+{
+	FPipeExecutorDoneTask(const FLWComponentSystemExecutionContext& InExecutionContext, UEntitySubsystem& InEntitySubsystem, TFunction<void()> InOnDoneNotification, const FString& InDebugName)
+		: ExecutionContext(InExecutionContext)
+		, EntitySubsystem(InEntitySubsystem)
+		, OnDoneNotification(InOnDoneNotification)
+		, DebugName(InDebugName)
+	{
+	}
+	static TStatId GetStatId()
+	{
+		RETURN_QUICK_DECLARE_CYCLE_STAT(FPipeExecutorDoneTask, STATGROUP_TaskGraphTasks);
+	}
+
+	static ENamedThreads::Type GetDesiredThread() { return ENamedThreads::GameThread; }
+	static ESubsequentsMode::Type GetSubsequentsMode() { return ESubsequentsMode::TrackSubsequents; }
+	void DoTask(ENamedThreads::Type CurrentThread, const FGraphEventRef& MyCompletionGraphEvent)
+	{
+		TRACE_CPUPROFILER_EVENT_SCOPE_STR("Flush Deferred Commands Parallel");
+
+		if (&ExecutionContext.Defer() != &EntitySubsystem.Defer())
+		{
+			ExecutionContext.Defer().MoveAppend(EntitySubsystem.Defer());
+		}
+
+		UE_LOG(LogPipe, Log, TEXT("PipeExecutor %s tasks DONE"), *DebugName);
+		ExecutionContext.SetFlushDeferredCommands(true);
+		ExecutionContext.FlushDeferred(EntitySubsystem);
+
+		OnDoneNotification();
+	}
+private:
+	FLWComponentSystemExecutionContext ExecutionContext;
+	UEntitySubsystem& EntitySubsystem;
+	TFunction<void()> OnDoneNotification;
+	FString DebugName;
+};
+
+FGraphEventRef TriggerParallelTasks(UPipeProcessor& Processor, FPipeContext& PipeContext, TFunction<void()> OnDoneNotification)
+{
+	TRACE_CPUPROFILER_EVENT_SCOPE(RunProcessorsView);
+
+	if (PipeContext.EntitySubsystem == nullptr)
+	{
+		UE_LOG(LogPipe, Error, TEXT("%s PipeContext.EntitySubsystem is null. Baling out."), ANSI_TO_TCHAR(__FUNCTION__));
+		return FGraphEventRef();
+	}
+
+	TRACE_CPUPROFILER_EVENT_SCOPE_STR("PipeExecutor RunParallel")
+
+	// not going through UEntitySubsystem::CreateExecutionContext on purpose - we do need a separate command buffer
+	FLWComponentSystemExecutionContext ExecutionContext(PipeContext.DeltaSeconds);
+	ExecutionContext.SetDeferredCommandBuffer(MakeShareable(new FLWCCommandBuffer()));
+	ExecutionContext.SetFlushDeferredCommands(false);
+	ExecutionContext.SetAuxData(PipeContext.AuxData);
+
+	FGraphEventRef CompletionEvent;
+	{
+		TRACE_CPUPROFILER_EVENT_SCOPE_STR("Dispatch Processors")
+		CompletionEvent = Processor.DispatchProcessorTasks(*PipeContext.EntitySubsystem, ExecutionContext, {});
+	}
+
+	if (CompletionEvent.IsValid())
+	{
+		const FGraphEventArray Prerequisites = { CompletionEvent };
+		CompletionEvent = TGraphTask<FPipeExecutorDoneTask>::CreateTask(&Prerequisites)
+			.ConstructAndDispatchWhenReady(ExecutionContext, *PipeContext.EntitySubsystem, OnDoneNotification, Processor.GetName());
+	}
+
+	return CompletionEvent;
+}
+
+} // namespace UE::Pipe::Executor
diff --git a/Engine/Plugins/Runtime/MassEntity/Source/MassEntity/Private/MassProcessingPhase.cpp b/Engine/Plugins/Runtime/MassEntity/Source/MassEntity/Private/MassProcessingPhase.cpp
new file mode 100644
--- /dev/null
+++ b/Engine/Plugins/Runtime/MassEntity/Source/MassEntity/Private/MassProcessingPhase.cpp
@@ -0,0 +1,279 @@
+// Copyright Epic Games, Inc. All Rights Reserved.
+
+#include "MassProcessingPhase.h"
+#include "MassEntityTypes.h"
+#include "MassEntityDebug.h"
+#include "MassEntitySettings.h"
+#include "MassProcessor.h"
+#include "MassExecutor.h"
+#include "MassEntitySystem.h"
+#include "VisualLogger/VisualLogger.h"
+#include "Engine/World.h"
+#include "LWCCommandBuffer.h"
+
+#define LOCTEXT_NAMESPACE "Pipe"
+
+DECLARE_CYCLE_STAT(TEXT("Pipe Phase Done"), STAT_PipePhaseDone, STATGROUP_TaskGraphTasks);
+
+namespace FPipeTweakables
+{
+	bool bFullyParallel = false;
+
+	FAutoConsoleVariableRef CVars[] = {
+		{TEXT("pipe.FullyParallel"), bFullyParallel, TEXT("Enables pipe processing distribution to all available thread (via the task graph)")},
+	};
+}
+
+//----------------------------------------------------------------------//
+//  FPipeProcessingPhase
+//----------------------------------------------------------------------//
+FPipeProcessingPhase::FPipeProcessingPhase()
+{
+	bCanEverTick = true;
+	bStartWithTickEnabled = false;
+}
+
+void FPipeProcessingPhase::ExecuteTick(float DeltaTime, ELevelTick TickType, ENamedThreads::Type CurrentThread, const FGraphEventRef& MyCompletionGraphEvent)
+{
+	if (TickType == LEVELTICK_ViewportsOnly || TickType == LEVELTICK_PauseTick)
+	{
+		return;
+	}
+
+	checkf(Manager, TEXT("Manager is null which is not a supported case. Either this FPipeProcessingPhase has not been initialized properly or it's been left dangling after the FPipeProcessingPhase owner got destroyed."));
+
+	Manager->OnPhaseStart(*this);
+	
+	OnPhaseStart.Broadcast(DeltaTime);
+
+	check(PhaseProcessor);
+	
+	FPipeContext Context(Manager->GetEntitySubsystemRef(), DeltaTime);
+
+	if (bRunInParallelMode)
+	{
+		bIsDuringPipeProcessing = true;
+		const FGraphEventRef PipelineCompletionEvent = UE::Pipe::Executor::TriggerParallelTasks(*PhaseProcessor, Context, [this, DeltaTime]()
+			{
+				OnParallelExecutionDone(DeltaTime);
+			});
+
+		if (PipelineCompletionEvent.IsValid())
+		{
+			MyCompletionGraphEvent->DontCompleteUntil(PipelineCompletionEvent);
+		}
+		else
+		{
+			OnParallelExecutionDone(DeltaTime);
+		}
+	}
+	else
+	{
+		TGuardValue<bool> PipeRunningGuard(bIsDuringPipeProcessing, true);
+		UE::Pipe::Executor::Run(*PhaseProcessor, Context);
+
+		OnPhaseEnd.Broadcast(DeltaTime);
+		Manager->OnPhaseEnd(*this);
+	}
+}
+
+void FPipeProcessingPhase::OnParallelExecutionDone(const float DeltaTime)
+{
+	bIsDuringPipeProcessing = false;
+	OnPhaseEnd.Broadcast(DeltaTime);
+	check(Manager);
+	Manager->OnPhaseEnd(*this);
+}
+
+FString FPipeProcessingPhase::DiagnosticMessage()
+{
+	return (Manager ? Manager->GetFullName() : TEXT("NULL-PipeProcessingPhaseManager")) + TEXT("[ProcessorTick]");
+}
+
+FName FPipeProcessingPhase::DiagnosticContext(bool bDetailed)
+{
+	return Manager ? Manager->GetClass()->GetFName() : TEXT("NULL-PipeProcessingPhaseManager");
+}
+
+//----------------------------------------------------------------------//
+// UPipeProcessingPhaseManager  
+//----------------------------------------------------------------------//
+void UPipeProcessingPhaseManager::PostInitProperties()
+{
+	Super::PostInitProperties();
+	if (HasAnyFlags(RF_ClassDefaultObject) == false)
+	{
+#if WITH_EDITOR
+		UWorld* World = GetWorld();
+		if (World && World->IsGameWorld() == false)
+		{
+			UPipeSettings* PipeSettings = GetMutableDefault<UPipeSettings>();
+			check(PipeSettings);
+			PipeSettingsChangeHandle = PipeSettings->GetOnSettingsChange().AddUObject(this, &UPipeProcessingPhaseManager::OnPipeSettingsChange);
+		}
+#endif // WITH_EDITOR
+	}
+	CreatePhases();
+}
+
+void UPipeProcessingPhaseManager::BeginDestroy()
+{
+	Stop();
+	Super::BeginDestroy();
+
+#if WITH_EDITOR
+	if (UPipeSettings* PipeSettings = GetMutableDefault<UPipeSettings>())
+	{
+		PipeSettings->GetOnSettingsChange().Remove(PipeSettingsChangeHandle);
+	}
+#endif // WITH_EDITOR
+}
+
+void UPipeProcessingPhaseManager::InitializePhases(UObject& InProcessorOwner)
+{
+	const FPipeProcessingPhaseConfig* ProcessingPhasesConfig = GET_PIPE_CONFIG_VALUE(GetProcessingPhasesConfig());
+
+	FString DependencyGraphFileName;
+
+#if WITH_EDITOR
+	const UWorld* World = InProcessorOwner.GetWorld();
+	const UPipeSettings* PipeSettings = GetMutableDefault<UPipeSettings>();
+	if (World != nullptr && PipeSettings != nullptr && !PipeSettings->DumpDependencyGraphFileName.IsEmpty())
+	{
+		DependencyGraphFileName = FString::Printf(TEXT("%s_%s"), *PipeSettings->DumpDependencyGraphFileName,*ToString(World->GetNetMode()));
+	}
+#endif // WITH_EDITOR
+
+	for (int i = 0; i < int(EPipeProcessingPhase::MAX); ++i)
+	{
+		const FPipeProcessingPhaseConfig& PhaseConfig = ProcessingPhasesConfig[i];
+		UPipeCompositeProcessor* PhaseProcessor = ProcessingPhases[i].PhaseProcessor;
+		FString FileName = !DependencyGraphFileName.IsEmpty() ? FString::Printf(TEXT("%s_%s"), *DependencyGraphFileName, *PhaseConfig.PhaseName.ToString()) : FString();
+		PhaseProcessor->CopyAndSort(PhaseConfig, FileName);
+		PhaseProcessor->Initialize(InProcessorOwner);
+	}
+
+#if WITH_PIPE_DEBUG
+	// print it all out to vislog
+	UE_VLOG_UELOG(this, LogPipe, Verbose, TEXT("Phases initialization done. Current composition:"));
+
+	FStringOutputDevice OutDescription;
+	for (int i = 0; i < int(EPipeProcessingPhase::MAX); ++i)
+	{
+		if (ProcessingPhases[i].PhaseProcessor)
+		{
+			ProcessingPhases[i].PhaseProcessor->DebugOutputDescription(OutDescription);
+			UE_VLOG_UELOG(this, LogPipe, Verbose, TEXT("--- %s"), *OutDescription);
+			OutDescription.Reset();
+		}
+	}	
+#endif // WITH_PIPE_DEBUG
+}
+
+void UPipeProcessingPhaseManager::Start(UWorld& World)
+{
+	EntitySubsystem = UPipeEntitySubsystem::GetCurrent(&World);
+
+	if (ensure(EntitySubsystem))
+	{
+		EnableTickFunctions(World);
+	}
+	else
+	{
+		UE_VLOG_UELOG(this, LogPipe, Error, TEXT("Called %s while missing the EntitySubsystem"), ANSI_TO_TCHAR(__FUNCTION__));
+	}
+}
+
+void UPipeProcessingPhaseManager::Start(UPipeEntitySubsystem& InEntitySubsystem)
+{
+	UWorld* World = InEntitySubsystem.GetWorld();
+	check(World);
+	EntitySubsystem = &InEntitySubsystem;
+	EnableTickFunctions(*World);
+}
+
+void UPipeProcessingPhaseManager::EnableTickFunctions(const UWorld& World)
+{
+	check(EntitySubsystem);
+	for (FPipeProcessingPhase& Phase : ProcessingPhases)
+	{
+		Phase.Manager = this;
+		Phase.RegisterTickFunction(World.PersistentLevel);
+		Phase.SetTickFunctionEnable(true);
+	}
+	UE_VLOG_UELOG(this, LogPipe, Log, TEXT("UPipeTickManager %s.%s has been started")
+		, *GetNameSafe(GetOuter()), *GetName());
+}
+
+void UPipeProcessingPhaseManager::Stop()
+{
+	EntitySubsystem = nullptr;
+	
+	for (FPipeProcessingPhase& Phase : ProcessingPhases)
+	{
+		Phase.SetTickFunctionEnable(false);
+	}
+
+	UE_VLOG_UELOG(this, LogPipe, Log, TEXT("UPipeTickManager %s.%s has been stopped")
+		, *GetNameSafe(GetOuter()), *GetName());
+}
+
+void UPipeProcessingPhaseManager::CreatePhases() 
+{
+	static ETickingGroup PhaseToTickingGroup[int(EPipeProcessingPhase::MAX)]
+	{
+		ETickingGroup::TG_PrePhysics, // EPipeProcessingPhase::PrePhysics
+		ETickingGroup::TG_StartPhysics, // EPipeProcessingPhase::StartPhysics
+		ETickingGroup::TG_DuringPhysics, // EPipeProcessingPhase::DuringPhysics
+		ETickingGroup::TG_EndPhysics,	// EPipeProcessingPhase::EndPhysics
+		ETickingGroup::TG_PostPhysics,	// EPipeProcessingPhase::PostPhysics
+		ETickingGroup::TG_LastDemotable, // EPipeProcessingPhase::FrameEnd
+	};
+
+	// @todo copy from settings instead of blindly creating from scratch
+	for (int i = 0; i < int(EPipeProcessingPhase::MAX); ++i)
+	{
+		ProcessingPhases[i].Phase = EPipeProcessingPhase(i);
+		ProcessingPhases[i].TickGroup = PhaseToTickingGroup[i];
+		ProcessingPhases[i].PhaseProcessor = NewObject<UPipeCompositeProcessor>(this, UPipeCompositeProcessor::StaticClass()
+			, *FString::Printf(TEXT("ProcessingPhase_%s"), *EnumToString(EPipeProcessingPhase(i))));
+		REDIRECT_OBJECT_TO_VLOG(ProcessingPhases[i].PhaseProcessor, this);
+		ProcessingPhases[i].PhaseProcessor->SetGroupName(*EnumToString(EPipeProcessingPhase(i)));
+		ProcessingPhases[i].PhaseProcessor->SetProcessingPhase(EPipeProcessingPhase(i));
+	}
+}
+
+void UPipeProcessingPhaseManager::OnPhaseStart(const FPipeProcessingPhase& Phase)
+{
+	ensure(CurrentPhase == EPipeProcessingPhase::MAX);
+	CurrentPhase = Phase.Phase;
+}
+
+void UPipeProcessingPhaseManager::OnPhaseEnd(FPipeProcessingPhase& Phase)
+{
+	ensure(CurrentPhase == Phase.Phase);
+	CurrentPhase = EPipeProcessingPhase::MAX;
+
+	// switch between parallel and single-thread versions only after a given batch of processing has been wrapped up	
+	if (Phase.IsConfiguredForParallelMode() != FPipeTweakables::bFullyParallel)
+	{
+		if (FPipeTweakables::bFullyParallel)
+		{
+			Phase.ConfigureForParallelMode();
+		}
+		else
+		{
+			Phase.ConfigureForSingleThreadMode();
+		}
+	}
+}
+
+#if WITH_EDITOR
+void UPipeProcessingPhaseManager::OnPipeSettingsChange(const FPropertyChangedEvent& PropertyChangedEvent)
+{
+	check(GetOuter());
+	InitializePhases(*GetOuter());
+}
+#endif // WITH_EDITOR
+
+#undef LOCTEXT_NAMESPACE
\ No newline at end of file
