AggregateTicking and Pipe merged into Engine/Plugins/Runtime/MassEntity

diff --git a/Engine/Plugins/Runtime/MassEntity/Source/MassEntity/Private/MassProcessor.cpp b/Engine/Plugins/Runtime/MassEntity/Source/MassEntity/Private/MassProcessor.cpp
new file mode 100644
--- /dev/null
+++ b/Engine/Plugins/Runtime/MassEntity/Source/MassEntity/Private/MassProcessor.cpp
@@ -0,0 +1,626 @@
+// Copyright Epic Games, Inc. All Rights Reserved.
+
+#include "MassProcessor.h"
+#include "MassEntitySettings.h"
+#include "MassProcessorDependencySolver.h"
+#include "VisualLogger/VisualLogger.h"
+#include "Engine/World.h"
+#include "LWCCommandBuffer.h"
+
+DECLARE_CYCLE_STAT(TEXT("PipeProcessor Group Completed"), Pipe_GroupCompletedTask, STATGROUP_TaskGraphTasks);
+
+#define PARALLELIZED_TRAFFIC_HACK 1
+
+#if PARALLELIZED_TRAFFIC_HACK
+namespace UE::MassTraffic
+{
+	int32 bParallelizeTraffic = 1;
+	FAutoConsoleVariableRef CVarParallelizeTraffic(TEXT("ai.traffic.parallelize"), bParallelizeTraffic, TEXT("Whether to parallelize traffic or not"), ECVF_Cheat);
+}
+#endif // PARALLELIZED_TRAFFIC_HACK
+
+#if WITH_PIPE_DEBUG
+namespace UE::Mass::Debug
+{
+	bool bLogProcessingGraph = false;
+	FAutoConsoleVariableRef CVarLogProcessingGraph(TEXT("pipe.LogProcessingGraph"), bLogProcessingGraph
+		, TEXT("When enabled will log task graph tasks created while dispatching processors to other threads, along with their dependencies"), ECVF_Cheat);
+}
+#endif // WITH_PIPE_DEBUG
+
+// change to && 1 to enable more detailed processing tasks logging
+#if WITH_PIPE_DEBUG && 0
+#define PROCESSOR_LOG(Fmt, ...) UE_LOG(LogPipe, Log, Fmt, ##__VA_ARGS__)
+#else // WITH_PIPE_DEBUG
+#define PROCESSOR_LOG(...) 
+#endif // WITH_PIPE_DEBUG
+
+namespace FPipeTweakables
+{
+	bool bParallelGroups = false;
+	float PostponedTaskWaitTimeWarningLevel = 0.002f;
+
+	FAutoConsoleVariableRef CVarsPipeProcessor[] = {
+		{TEXT("pipe.ParallelGroups"), bParallelGroups, TEXT("Enables pipe processing groups distribution to all available threads (via the task graph)")},
+		{TEXT("pipe.PostponedTaskWaitTimeWarningLevel"), PostponedTaskWaitTimeWarningLevel, TEXT("if waiting for postponed task\'s dependencies exceeds this number an error will be logged")},
+	};
+}
+
+class FPipeProcessorTask
+{
+public:
+	FPipeProcessorTask(UEntitySubsystem& InEntitySubsystem, const FLWComponentSystemExecutionContext& InExecutionContext, UPipeProcessor& InProc, bool bInManageCommandBuffer = true)
+		: EntitySubsystem(&InEntitySubsystem)
+		, ExecutionContext(InExecutionContext)
+		, Processor(&InProc)
+		, bManageCommandBuffer(bInManageCommandBuffer)
+	{}
+
+	static TStatId GetStatId()
+	{
+		RETURN_QUICK_DECLARE_CYCLE_STAT(FPipeProcessorTask, STATGROUP_TaskGraphTasks);
+	}
+
+	static ESubsequentsMode::Type GetSubsequentsMode() { return ESubsequentsMode::TrackSubsequents; }
+
+	static ENamedThreads::Type GetDesiredThread()
+	{
+		return ENamedThreads::AnyHiPriThreadHiPriTask;
+	}
+
+	void DoTask(ENamedThreads::Type CurrentThread, const FGraphEventRef& MyCompletionGraphEvent)
+	{
+		checkf(Processor, TEXT("Expecting a valid processor to execute"));
+		checkf(EntitySubsystem, TEXT("Expecting a valid entity subsystem to execute processor"));
+
+		PROCESSOR_LOG(TEXT("+--+ Task %s started on %u"), *Processor->GetProcessorName(), FPlatformTLS::GetCurrentThreadId());
+
+		UEntitySubsystem::FScopedProcessing ProcessingScope = EntitySubsystem->NewProcessingScope();
+
+		TRACE_CPUPROFILER_EVENT_SCOPE_STR("Pipe Processor Task");
+		
+		if (bManageCommandBuffer)
+		{
+			TSharedPtr<FLWCCommandBuffer> MainSharedPtr = ExecutionContext.GetSharedDeferredCommandBuffer();
+			ExecutionContext.SetDeferredCommandBuffer(MakeShareable(new FLWCCommandBuffer()));
+			Processor->CallExecute(*EntitySubsystem, ExecutionContext);
+			MainSharedPtr->MoveAppend(ExecutionContext.Defer());
+		}
+		else
+		{
+			Processor->CallExecute(*EntitySubsystem, ExecutionContext);
+		}
+		PROCESSOR_LOG(TEXT("+--+ Task %s finished"), *Processor->GetProcessorName());
+	}
+
+private:
+	UEntitySubsystem* EntitySubsystem = nullptr;
+	FLWComponentSystemExecutionContext ExecutionContext;
+	UPipeProcessor* Processor = nullptr;
+	/** 
+	 * indicates whether this task is responsible for creation of a dedicated command buffer and transferring over the 
+	 * commands after processor's execution;
+	 */
+	bool bManageCommandBuffer = true;
+};
+
+class FPipeProcessorsTask_GameThread : public FPipeProcessorTask
+{
+public:
+	FPipeProcessorsTask_GameThread(UEntitySubsystem& InEntitySubsystem, const FLWComponentSystemExecutionContext& InExecutionContext, UPipeProcessor& InProc)
+		: FPipeProcessorTask(InEntitySubsystem, InExecutionContext, InProc)
+	{}
+
+	static ENamedThreads::Type GetDesiredThread()
+	{
+		return ENamedThreads::GameThread;
+	}
+};
+
+//----------------------------------------------------------------------//
+// UPipeProcessor 
+//----------------------------------------------------------------------//
+UPipeProcessor::UPipeProcessor(const FObjectInitializer& ObjectInitializer)
+	: ExecutionFlags((int32)(EProcessorExecutionFlags::Server | EProcessorExecutionFlags::Standalone))
+{
+}
+
+void UPipeProcessor::SetShouldAutoRegisterWithGlobalList(const bool bAutoRegister)
+{	
+	if (ensureMsgf(HasAnyFlags(RF_ClassDefaultObject), TEXT("Setting bAutoRegisterWithProcessingPhases for non-CDOs has no effect")))
+	{
+		bAutoRegisterWithProcessingPhases = bAutoRegister;
+#if WITH_EDITOR
+		SaveConfig(CPF_Config, *GetDefaultConfigFilename());
+#endif // WITH_EDITOR
+	}
+}
+
+void UPipeProcessor::PostInitProperties()
+{
+	Super::PostInitProperties();
+	if (HasAnyFlags(RF_ClassDefaultObject) == false) 
+	{
+		ConfigureQueries();
+	}
+#if CPUPROFILERTRACE_ENABLED
+	StatId = GetProcessorName();
+#endif
+}
+
+void UPipeProcessor::CallExecute(UEntitySubsystem& EntitySubsystem, FLWComponentSystemExecutionContext& Context)
+{
+	TRACE_CPUPROFILER_EVENT_SCOPE_TEXT(*StatId);
+#if WITH_PIPE_DEBUG
+	Context.DebugSetExecutionDesc(FString::Printf(TEXT("%s (%s)"), *GetProcessorName(), *ToString(EntitySubsystem.GetWorld()->GetNetMode())));
+#endif
+	Execute(EntitySubsystem, Context);
+}
+
+FGraphEventRef UPipeProcessor::DispatchProcessorTasks(UEntitySubsystem& EntitySubsystem, FLWComponentSystemExecutionContext& ExecutionContext, const FGraphEventArray& Prerequisites)
+{
+	FGraphEventRef ReturnVal;
+	if (bRequiresGameThreadExecution)
+	{
+		ReturnVal = TGraphTask<FPipeProcessorsTask_GameThread>::CreateTask(&Prerequisites).ConstructAndDispatchWhenReady(EntitySubsystem, ExecutionContext, *this);
+	}
+	else
+	{
+		ReturnVal = TGraphTask<FPipeProcessorTask>::CreateTask(&Prerequisites).ConstructAndDispatchWhenReady(EntitySubsystem, ExecutionContext, *this);
+	}	
+	return ReturnVal;
+}
+
+void UPipeProcessor::DebugOutputDescription(FOutputDevice& Ar, int32 Indent) const
+{
+#if WITH_PIPE_DEBUG
+	Ar.Logf(TEXT("%*s%s"), Indent, TEXT(""), *GetProcessorName());
+#endif // WITH_PIPE_DEBUG
+}
+
+#if WITH_EDITOR
+void UPipeProcessor::PostEditChangeProperty(FPropertyChangedEvent& PropertyChangedEvent)
+{
+	Super::PostEditChangeProperty(PropertyChangedEvent);
+
+	if (HasAnyFlags(RF_ClassDefaultObject))
+	{
+		// this is here to make sure all the changes to CDOs we do via PipeSettings gets serialized to the ini file
+		SaveConfig(CPF_Config, *GetDefaultConfigFilename());
+	}
+}
+#endif // WITH_EDITOR
+
+//----------------------------------------------------------------------//
+//  UPipeCompositeProcessor
+//----------------------------------------------------------------------//
+UPipeCompositeProcessor::UPipeCompositeProcessor()
+	: GroupName(TEXT("None"))
+	, bRunInSeparateThread(false)
+	, bHasOffThreadSubGroups(false)
+{
+	// not auto-registering composite processors since the idea of the global processors list is to indicate all 
+	// the processors doing the work while composite processors are just containers. Having said that subclasses 
+	// can change this behavior if need be.
+	bAutoRegisterWithProcessingPhases = false;
+}
+
+void UPipeCompositeProcessor::SetChildProcessors(TArray<UPipeProcessor*>&& InProcessors)
+{
+	ChildPipeline.SetProcessors(MoveTemp(InProcessors));
+}
+
+void UPipeCompositeProcessor::ConfigureQueries()
+{
+	// nothing to do here since ConfigureQueries will get automatically called for all the processors in ChildPipeline
+	// via their individual PostInitProperties call
+}
+
+FGraphEventRef UPipeCompositeProcessor::DispatchProcessorTasks(UEntitySubsystem& EntitySubsystem, FLWComponentSystemExecutionContext& ExecutionContext, const FGraphEventArray& InPrerequisites)
+{
+	FGraphEventArray Events;
+	Events.Reserve(ProcessingFlatGraph.Num());
+		
+	for (FDependencyNode& ProcessingNode : ProcessingFlatGraph)
+	{
+		FGraphEventArray Prerequisites;
+		for (const int32 DependencyIndex : ProcessingNode.Dependencies)
+		{
+			Prerequisites.Add(Events[DependencyIndex]);
+		}
+
+		if (ProcessingNode.Processor)
+		{
+			Events.Add(ProcessingNode.Processor->DispatchProcessorTasks(EntitySubsystem, ExecutionContext, Prerequisites));
+		}
+		else
+		{
+			Events.Add(FFunctionGraphTask::CreateAndDispatchWhenReady([=](){}
+				, GET_STATID(Pipe_GroupCompletedTask), &Prerequisites, ENamedThreads::AnyHiPriThreadHiPriTask));
+		}
+	}
+
+
+#if WITH_PIPE_DEBUG
+	if (UE::Mass::Debug::bLogProcessingGraph)
+	{
+		for (int i = 0; i < ProcessingFlatGraph.Num(); ++i)
+		{
+			FDependencyNode& ProcessingNode = ProcessingFlatGraph[i];
+			FString DependenciesDesc;
+			for (const int32 DependencyIndex : ProcessingNode.Dependencies)
+			{
+				DependenciesDesc += FString::Printf(TEXT("%s, "), *ProcessingFlatGraph[DependencyIndex].Name.ToString());
+			}
+
+			if (ProcessingNode.Processor)
+			{
+				PROCESSOR_LOG(TEXT("Task %u %s%s%s"), Events[i]->GetTraceId(), *ProcessingNode.Processor->GetProcessorName()
+					, DependenciesDesc.Len() > 0 ? TEXT(" depends on ") : TEXT(""), *DependenciesDesc);
+			}
+			else
+			{
+				PROCESSOR_LOG(TEXT("Group %u %s%s%s"), Events[i]->GetTraceId(), *ProcessingNode.Name.ToString()
+					, DependenciesDesc.Len() > 0 ? TEXT(" depends on ") : TEXT(""), *DependenciesDesc);
+			}
+		}
+	}
+#endif // WITH_PIPE_DEBUG
+
+	FGraphEventRef CompletionEvent = FFunctionGraphTask::CreateAndDispatchWhenReady([this](){}
+		, GET_STATID(Pipe_GroupCompletedTask), &Events, ENamedThreads::AnyHiPriThreadHiPriTask);
+
+	return CompletionEvent;
+}
+
+void UPipeCompositeProcessor::Execute(UEntitySubsystem& EntitySubsystem, FLWComponentSystemExecutionContext& Context)
+{
+#if PARALLELIZED_TRAFFIC_HACK
+	if (FPipeTweakables::bParallelGroups == false
+		&& UE::MassTraffic::bParallelizeTraffic && GetProcessingPhase() == EPipeProcessingPhase::PrePhysics)
+	{
+		static FName TrafficGroup(TEXT("Traffic"));
+		FLWComponentSystemExecutionContext TrafficExecutionContext;
+		FGraphEventRef TrafficCompletionEvent;
+
+		for (UPipeProcessor* Proc : ChildPipeline.Processors)
+		{
+			check(Proc);
+				if (GetProcessingPhase() == EPipeProcessingPhase::PrePhysics)
+			{
+				if (const UPipeCompositeProcessor* CompProcessor = Cast<UPipeCompositeProcessor>(Proc))
+				{
+					if (CompProcessor->GetGroupName() == TrafficGroup)
+					{
+						TrafficExecutionContext = Context;
+						// Needs its own command buffer as it will be in it own thread, will be merged after this loop
+						TrafficExecutionContext.SetDeferredCommandBuffer(MakeShareable(new FLWCCommandBuffer()));
+						TrafficCompletionEvent = TGraphTask<FPipeProcessorTask>::CreateTask().ConstructAndDispatchWhenReady(EntitySubsystem, TrafficExecutionContext, *Proc, /*bManageCommandBuffer=*/false);
+						continue;
+					}
+				}
+			}
+			Proc->CallExecute(EntitySubsystem, Context);
+		}
+
+		if (TrafficCompletionEvent)
+		{
+			// synchronize with the traffic "thread"
+			TrafficCompletionEvent->Wait();
+			Context.Defer().MoveAppend(TrafficExecutionContext.Defer());
+		}
+
+		return;
+	}
+#endif // PARALLELIZED_TRAFFIC_HACK
+
+	if (FPipeTweakables::bParallelGroups && bHasOffThreadSubGroups)
+	{
+		CompletionStatus.Reset();
+		CompletionStatus.AddDefaulted(ChildPipeline.Processors.Num());
+		TArray<int32> PostponedProcessors;
+
+		FLWComponentSystemExecutionContext SingleThreadContext = Context;
+		SingleThreadContext.SetDeferredCommandBuffer(MakeShareable(new FLWCCommandBuffer()));
+
+		for (int32 NodeIndex = 0; NodeIndex < ChildPipeline.Processors.Num(); ++NodeIndex)
+		{
+			UPipeProcessor* Proc = ChildPipeline.Processors[NodeIndex];
+			check(Proc);
+			UPipeCompositeProcessor* CompositeProc = Cast<UPipeCompositeProcessor>(Proc);
+			if (CompositeProc == nullptr || CompositeProc->bRunInSeparateThread == false)
+			{
+				// check if all dependencies have been processed already
+				Proc->TransientDependencyIndices.Reset();
+				for (const int32 DependencyIndex : Proc->DependencyIndices)
+				{
+					if (CompletionStatus[DependencyIndex].IsDone() == false)
+					{
+						Proc->TransientDependencyIndices.Add(DependencyIndex);
+					}
+				}
+
+				if (Proc->TransientDependencyIndices.Num() == 0)
+				{
+					PROCESSOR_LOG(TEXT("+--+ Instant Execution: %s.%s in %u")
+						, CompositeProc ? TEXT("") : *Proc->GetExecutionOrder().ExecuteInGroup.ToString()
+						, CompositeProc ? *CompositeProc->GetGroupName().ToString() : *Proc->GetProcessorName()
+						, FPlatformTLS::GetCurrentThreadId());
+
+					Proc->CallExecute(EntitySubsystem, SingleThreadContext);
+					CompletionStatus[NodeIndex].Status = EProcessorCompletionStatus::Done;
+				}
+				else
+				{
+					PROCESSOR_LOG(TEXT("+--+ POSTPONED Execution: %s.%s in %u")
+						, CompositeProc ? TEXT("") : *Proc->GetExecutionOrder().ExecuteInGroup.ToString()
+						, CompositeProc ? *CompositeProc->GetGroupName().ToString() : *Proc->GetProcessorName()
+						, FPlatformTLS::GetCurrentThreadId());
+
+					// postpone
+					// the graph event is needed in case we have off-thread processors depending on this one
+					CompletionStatus[NodeIndex].CompletionEvent = FGraphEvent::CreateGraphEvent();
+					CompletionStatus[NodeIndex].Status = EProcessorCompletionStatus::Postponed;
+					PostponedProcessors.Add(NodeIndex);
+				}
+			}
+			else
+			{
+				// gather prerequisites
+				FString DependenciesDesc;
+				FGraphEventArray Prerequisites;
+				for (const int32 Index : Proc->GetPrerequisiteIndices())
+				{
+					if (Index != INDEX_NONE && CompletionStatus[Index].IsDone() == false)
+					{
+						Prerequisites.Add(CompletionStatus[Index].CompletionEvent);
+						DependenciesDesc += FString::Printf(TEXT("%s, "), *ChildPipeline.Processors[Index]->GetProcessorName());
+					}
+				}
+
+				PROCESSOR_LOG(TEXT("+--+ Task %s created. %s%s"), *Proc->GetProcessorName()
+					, DependenciesDesc.Len() > 0 ? TEXT(" Dependencies: ") : TEXT(""), *DependenciesDesc);
+
+				// send off to another thread
+				CompletionStatus[NodeIndex].CompletionEvent = TGraphTask<FPipeProcessorTask>::CreateTask(&Prerequisites).ConstructAndDispatchWhenReady(EntitySubsystem, Context, *Proc);
+				CompletionStatus[NodeIndex].Status = EProcessorCompletionStatus::Threaded;
+			}
+		}
+
+		{
+			TRACE_CPUPROFILER_EVENT_SCOPE_TEXT("Postponed PipeProcessors");
+
+			for (int32 i = 0; i < PostponedProcessors.Num(); ++i)
+			{
+				const int32 PostponedIndex = PostponedProcessors[i];
+				UPipeProcessor* Proc = ChildPipeline.Processors[PostponedIndex];
+				check(Proc);
+				for (int32 j = Proc->TransientDependencyIndices.Num() - 1; j >= 0; --j)
+				{
+					int32 DependencyIndex = Proc->TransientDependencyIndices[j];
+					ensureAlways(CompletionStatus[DependencyIndex].IsDone() || CompletionStatus[DependencyIndex].Status == EProcessorCompletionStatus::Threaded);
+						
+					CompletionStatus[DependencyIndex].Wait();
+					Proc->TransientDependencyIndices.RemoveAtSwap(j, 1, /*bAllowShrinking=*/false);
+				}
+
+				if (Proc->TransientDependencyIndices.Num() == 0)
+				{
+					Proc->CallExecute(EntitySubsystem, SingleThreadContext);
+					CompletionStatus[PostponedIndex].Status = EProcessorCompletionStatus::Done;
+					CompletionStatus[PostponedIndex].CompletionEvent->DispatchSubsequents();
+					PostponedProcessors.RemoveAt(i--, 1, /*bAllowShrinking=*/false);
+				}
+			}
+			ensureMsgf(PostponedProcessors.Num() == 0
+				, TEXT("Failed to execute all processors in one intermittent sequence - this indicates an issue with depdendency ordering"));
+		}
+
+		// wait for all events to complete
+		for (auto& Event : CompletionStatus)
+		{
+			if (Event.CompletionEvent.IsValid())
+			{
+				Event.CompletionEvent->Wait();
+			}
+		}
+		Context.Defer().MoveAppend(SingleThreadContext.Defer());
+	}
+	else
+	{
+		for (UPipeProcessor* Proc : ChildPipeline.Processors)
+		{
+			check(Proc);
+			Proc->CallExecute(EntitySubsystem, Context);
+		}
+	}
+}
+
+void UPipeCompositeProcessor::Initialize(UObject& Owner)
+{
+	// remove all nulls
+	ChildPipeline.Processors.RemoveAll([](const UPipeProcessor* Proc) { return Proc == nullptr; });
+
+	// from this point on we don't expect to have nulls in ChildPipeline.Processors
+	for (UPipeProcessor* Proc : ChildPipeline.Processors)
+	{
+		REDIRECT_OBJECT_TO_VLOG(Proc, this);
+		Proc->Initialize(Owner);
+	}
+}
+
+void UPipeCompositeProcessor::CopyAndSort(const FPipeProcessingPhaseConfig& PhaseConfig, const FString& DependencyGraphFileName)
+{
+	check(GetOuter());
+	
+	// create processors via a temporary pipeline. This ensures consistency with other places we use FRuntimePipelines
+	FRuntimePipeline TmpPipeline;
+	TmpPipeline.CreateFromArray(PhaseConfig.ProcessorCDOs, *GetOuter());
+
+	// figure out dependencies
+	FProcessorDependencySolver Solver(TmpPipeline.Processors, GroupName, DependencyGraphFileName);
+	TArray<FProcessorDependencySolver::FOrderInfo> SortedProcessorsAndGroups;
+	Solver.ResolveDependencies(SortedProcessorsAndGroups, PhaseConfig.OffGameThreadGroupNames);
+
+	Populate(SortedProcessorsAndGroups);
+
+	// this part is creating an ordered, flat list of processors that can be executed in sequence
+	// with subsequent task only depending on the elements prior on the list
+	TMap<FName, int32> NameToDependencyIndex;
+	NameToDependencyIndex.Reserve(SortedProcessorsAndGroups.Num());
+	for (FProcessorDependencySolver::FOrderInfo& Element : SortedProcessorsAndGroups)
+	{
+		NameToDependencyIndex.Add(Element.Name, NameToDependencyIndex.Num());
+
+		FDependencyNode& Node = ProcessingFlatGraph.Add_GetRef({ Element.Name, Element.Processor });
+		Node.Dependencies.Reserve(Element.Dependencies.Num());
+		for (FName DependencyName : Element.Dependencies)
+		{
+			Node.Dependencies.Add(NameToDependencyIndex.FindChecked(DependencyName));
+		}
+	}
+}
+
+int32 UPipeCompositeProcessor::Populate(TArray<FProcessorDependencySolver::FOrderInfo>& ProcessorsAndGroups, const int32 StartIndex)
+{
+	ChildPipeline.Processors.Reset();
+
+	// if processor -> Add
+	// if not processor 
+	//    if "start" -> Create group and pass the InOutProcessorsAndGroups in
+	//    else return (pop out)
+	int32 Index = StartIndex;
+	TMap<FName, int32> NameToIndexMap;
+	bool bOffThreadGroupsFound = false;
+	const FPipeProcessingPhaseConfig& PhaseConfig = GET_PIPE_CONFIG_VALUE(GetProcessingPhaseConfig(ProcessingPhase));
+
+	while (ProcessorsAndGroups.IsValidIndex(Index))
+	{
+		FProcessorDependencySolver::FOrderInfo& Element = ProcessorsAndGroups[Index];
+		ensure(Element.NodeType != EDependencyNodeType::Invalid);
+		
+		if (Element.NodeType == EDependencyNodeType::Processor)
+		{
+			check(Element.Processor);
+			check(Element.Processor->GetExecutionOrder().ExecuteInGroup == GroupName || Element.Processor->GetExecutionOrder().ExecuteInGroup.IsNone());
+
+			Element.Processor->DependencyIndices.Reset();
+			for (const FName DependencyName : Element.Dependencies)
+			{
+				Element.Processor->DependencyIndices.Add(NameToIndexMap.FindChecked(DependencyName));
+			}
+
+			NameToIndexMap.Add(Element.Name, ChildPipeline.Processors.Num());
+			ChildPipeline.AppendProcessor(*Element.Processor);
+			++Index;
+		}
+		else if (Element.NodeType == EDependencyNodeType::GroupStart)
+		{
+			UPipeCompositeProcessor* GroupProcessor = NewObject<UPipeCompositeProcessor>(GetOuter());
+			GroupProcessor->SetGroupName(Element.Name);
+			GroupProcessor->SetProcessingPhase(ProcessingPhase);
+			if (PhaseConfig.OffGameThreadGroupNames.Find(Element.Name) != INDEX_NONE)
+			{
+				GroupProcessor->bRunInSeparateThread = true;
+			}
+
+			GroupProcessor->DependencyIndices.Reset();
+			for (const FName DependencyName : Element.Dependencies)
+			{
+				const int32 DependencyIndex = NameToIndexMap.FindChecked(DependencyName);
+				GroupProcessor->DependencyIndices.Add(DependencyIndex);
+			}
+
+			NameToIndexMap.Add(Element.Name, ChildPipeline.Processors.Num());
+			ChildPipeline.AppendProcessor(*GroupProcessor);
+
+			Index = GroupProcessor->Populate(ProcessorsAndGroups, Index + 1) + 1;
+			bOffThreadGroupsFound = (bOffThreadGroupsFound || GroupProcessor->bRunInSeparateThread);
+		}
+		else 
+		{
+			check(Element.NodeType == EDependencyNodeType::GroupEnd);
+			check(Element.Name == GroupName);
+			// done
+			break;
+		}
+	}
+
+	bHasOffThreadSubGroups = bOffThreadGroupsFound;
+
+	return Index;
+}
+
+void UPipeCompositeProcessor::DebugOutputDescription(FOutputDevice& Ar, int32 Indent) const
+{
+#if WITH_PIPE_DEBUG
+	if (ChildPipeline.Processors.Num() == 0)
+	{
+		Ar.Logf(TEXT("%*sGroup %s: []"), Indent, TEXT(""), *GroupName.ToString());
+	}
+	else
+	{
+		Ar.Logf(TEXT("%*sGroup %s:"), Indent, TEXT(""), *GroupName.ToString());
+		for (UPipeProcessor* Proc : ChildPipeline.Processors)
+		{
+			check(Proc);
+			Ar.Logf(TEXT("\n"));
+			Proc->DebugOutputDescription(Ar, Indent + 3);
+		}
+	}
+#endif // WITH_PIPE_DEBUG
+}
+
+void UPipeCompositeProcessor::SetProcessingPhase(EPipeProcessingPhase Phase)
+{
+	Super::SetProcessingPhase(Phase);
+	for (UPipeProcessor* Proc : ChildPipeline.Processors)
+	{
+		Proc->SetProcessingPhase(Phase);
+	}
+}
+
+void UPipeCompositeProcessor::SetGroupName(FName NewName)
+{
+	GroupName = NewName;
+#if CPUPROFILERTRACE_ENABLED
+	StatId = GroupName.ToString();
+#endif
+}
+
+void UPipeCompositeProcessor::AddGroupedProcessor(FName RequestedGroupName, UPipeProcessor& Processor)
+{
+	if (RequestedGroupName.IsNone() || RequestedGroupName == GroupName)
+	{
+		ChildPipeline.AppendProcessor(Processor);
+	}
+	else
+	{
+		FString RemainingGroupName;
+		UPipeCompositeProcessor* GroupProcessor = FindOrAddGroupProcessor(RequestedGroupName, &RemainingGroupName);
+		check(GroupProcessor);
+		GroupProcessor->AddGroupedProcessor(FName(*RemainingGroupName), Processor);
+	}
+}
+
+UPipeCompositeProcessor* UPipeCompositeProcessor::FindOrAddGroupProcessor(FName RequestedGroupName, FString* OutRemainingGroupName)
+{
+	UPipeCompositeProcessor* GroupProcessor = nullptr;
+	const FString NameAsString = RequestedGroupName.ToString();
+	FString TopGroupName;
+	if (NameAsString.Split(TEXT("."), &TopGroupName, OutRemainingGroupName))
+	{
+		RequestedGroupName = FName(*TopGroupName);
+	}
+	GroupProcessor = ChildPipeline.FindTopLevelGroupByName(RequestedGroupName);
+
+	if (GroupProcessor == nullptr)
+	{
+		check(GetOuter());
+		GroupProcessor = NewObject<UPipeCompositeProcessor>(GetOuter());
+		GroupProcessor->SetGroupName(RequestedGroupName);
+		ChildPipeline.AppendProcessor(*GroupProcessor);
+	}
+
+	return GroupProcessor;
+}
diff --git a/Engine/Plugins/Runtime/MassEntity/Source/MassEntity/Private/MassProcessorDependencySolver.cpp b/Engine/Plugins/Runtime/MassEntity/Source/MassEntity/Private/MassProcessorDependencySolver.cpp
new file mode 100644
--- /dev/null
+++ b/Engine/Plugins/Runtime/MassEntity/Source/MassEntity/Private/MassProcessorDependencySolver.cpp
@@ -0,0 +1,901 @@
+// Copyright Epic Games, Inc. All Rights Reserved.
+
+#include "MassProcessorDependencySolver.h"
+#include "MassEntityTypes.h"
+#include "MassProcessor.h"
+#include "Logging/MessageLog.h"
+#include "HAL/FileManager.h"
+#include "Misc/Paths.h"
+
+#define LOCTEXT_NAMESPACE "Pipe"
+
+//----------------------------------------------------------------------//
+//  FProcessorDependencySolver::FNode
+//----------------------------------------------------------------------//
+int32 FProcessorDependencySolver::FNode::FindOrAddGroupNodeIndex(const FString& GroupName)
+{
+	const FName GroupNodeName = FName(*GroupName);
+	const int32* GroupNodeIndex = Indices.Find(GroupNodeName);
+	if (GroupNodeIndex)
+	{
+		return *GroupNodeIndex;
+	}
+
+	const int32 Index = SubNodes.Add({ GroupNodeName, nullptr });
+	Indices.Add(GroupNodeName, Index);
+	return Index;
+}
+
+int32 FProcessorDependencySolver::FNode::FindNodeIndex(FName InNodeName) const
+{
+	const int32* GroupNodeIndex = Indices.Find(InNodeName);
+	return GroupNodeIndex ? *GroupNodeIndex : INDEX_NONE;
+}
+
+bool FProcessorDependencySolver::FNode::HasDependencies() const
+{
+	return ExecuteBefore.Num() > 0 || ExecuteAfter.Num() > 0;
+}
+
+//----------------------------------------------------------------------//
+//  FProcessorDependencySolver
+//----------------------------------------------------------------------//
+FProcessorDependencySolver::FProcessorDependencySolver(TArrayView<UPipeProcessor*> InProcessors, const FName Name, const FString& InDependencyGraphFileName /*= FString()*/)
+	: Processors(InProcessors)
+	, GroupRootNode({Name, nullptr})
+	, DependencyGraphFileName(InDependencyGraphFileName)
+{}
+
+FString FProcessorDependencySolver::NameViewToString(TConstArrayView<FName> View)
+{
+	if (View.Num() == 0)
+	{
+		return TEXT("[]");
+	}
+	FString ReturnVal = FString::Printf(TEXT("[%s"), *View[0].ToString());
+	for (int i = 1; i < View.Num(); ++i)
+	{
+		ReturnVal += FString::Printf(TEXT(", %s"), *View[i].ToString());
+	}
+	return ReturnVal + TEXT("]");
+}
+
+int32 FProcessorDependencySolver::PerformSolverStep(FNode& RootNode, TArray<int32>& InOutIndicesRemaining, TArray<int32>& OutNodeIndices)
+{
+	const int32 StartingCount = OutNodeIndices.Num();
+
+	for (int32 i = InOutIndicesRemaining.Num() - 1; i >= 0; --i)
+	{
+		const int32 NodeIndex = InOutIndicesRemaining[i];
+		if (RootNode.SubNodes[NodeIndex].TransientDependencies.Num() == 0)
+		{
+			InOutIndicesRemaining.RemoveAtSwap(i, 1, /*bAllowShrinking=*/false);
+			OutNodeIndices.Add(NodeIndex);
+		}
+	}
+
+	for (const int32 NodeIndex : OutNodeIndices)
+	{
+		for (const int32 RemainingNodeIndex : InOutIndicesRemaining)
+		{
+			RootNode.SubNodes[RemainingNodeIndex].TransientDependencies.RemoveSingleSwap(NodeIndex, /*bAllowShrinking=*/false);
+		}
+	}
+
+	return OutNodeIndices.Num() - StartingCount;
+}
+
+bool FProcessorDependencySolver::PerformPrioritySolverStep(FNode& RootNode, TArray<int32>& InOutPriorityIndicesRemaining, TArray<int32>& OutNodeIndices)
+{
+	for (int32 i = 0; i < InOutPriorityIndicesRemaining.Num(); ++i)
+	{
+		const int32 NodeIndex = InOutPriorityIndicesRemaining[i];
+		if (RootNode.SubNodes[NodeIndex].TransientDependencies.Num() == 0)
+		{
+			InOutPriorityIndicesRemaining.RemoveAt(i, 1, /*bAllowShrinking=*/false);
+			OutNodeIndices.Add(NodeIndex);
+
+			for (const int32 RemainingNodeIndex : InOutPriorityIndicesRemaining)
+			{
+				RootNode.SubNodes[RemainingNodeIndex].TransientDependencies.RemoveSingleSwap(NodeIndex, /*bAllowShrinking=*/false);
+			}
+
+			return true;
+		}
+	}
+
+	return false;
+}
+
+void FProcessorDependencySolver::CreateSubGroupNames(FName InGroupName, TArray<FString>& SubGroupNames)
+{
+	// the function will convert composite group name into a series of progressively more precise group names
+	// so "A.B.C" will result in ["A", "A.B", "A.B.C"]
+
+	SubGroupNames.Reset();
+	FString GroupNameAsString = InGroupName.ToString();
+	FString TopGroupName;
+
+	while (GroupNameAsString.Split(TEXT("."), &TopGroupName, &GroupNameAsString))
+	{
+		SubGroupNames.Add(TopGroupName);
+	}
+	SubGroupNames.Add(GroupNameAsString);
+	
+	for (int i = 1; i < SubGroupNames.Num(); ++i)
+	{
+		SubGroupNames[i] = FString::Printf(TEXT("%s.%s"), *SubGroupNames[i - 1], *SubGroupNames[i]);
+	}
+}
+
+void FProcessorDependencySolver::AddNode(FName InGroupName, UPipeProcessor& Processor)
+{
+	check(Processor.GetClass());
+	const FName ProcName = Processor.GetClass()->GetFName();
+
+	if (InGroupName.IsNone())
+	{
+		if (GroupRootNode.Indices.Find(ProcName) == nullptr)
+		{
+			GroupRootNode.Indices.Add(ProcName, GroupRootNode.SubNodes.Num());
+			GroupRootNode.SubNodes.Add({ProcName, &Processor});
+		}
+		else
+		{
+			UE_LOG(LogPipe, Warning, TEXT("%s Processor %s already registered. Duplicates are not supported.")
+				, ANSI_TO_TCHAR(__FUNCTION__), *ProcName.ToString());
+		}
+	}
+	else
+	{
+		// The idea here is that we create a tree-like structure to place the Processor as a leaf.
+		// The levels of InGroupName are used to create tree successive branch nodes, so a group name A.B.C
+		// will result in creating node A, then that node will have node A.B created for it and that one in 
+		// turn will get A.B.C; Every node created gets information about the subnodes it hosts for faster 
+		// lookup during dependency resolving
+
+		// first we need to figure out all the subgroups indicated by InGroupName
+		TArray<FString> AllGroupNames;
+		CreateSubGroupNames(InGroupName, AllGroupNames);
+
+		// now drill down into the successive levels of the free and find or create the necessary nodes. 
+		FNode* CurrentNode = &GroupRootNode;
+		for (int GroupNameIndex = 0; GroupNameIndex < AllGroupNames.Num(); ++GroupNameIndex)
+		{
+			const int32 NodeIndex = CurrentNode->FindOrAddGroupNodeIndex(AllGroupNames[GroupNameIndex]);
+			CurrentNode->Indices.Add(ProcName, NodeIndex);
+			// we make every node aware of all the subgroups it contains
+			for (int SubgroupIndex = GroupNameIndex + 1; SubgroupIndex < AllGroupNames.Num(); ++SubgroupIndex)
+			{
+				CurrentNode->Indices.Add(*AllGroupNames[SubgroupIndex], NodeIndex);
+			}
+
+			CurrentNode = &CurrentNode->SubNodes[NodeIndex];
+		}
+
+		// at this point CurrentNode points at the leaf group where we need to add the processor itself
+		if (CurrentNode->Indices.Find(ProcName) == nullptr)
+		{
+			const int32 ProcNodeIndex = CurrentNode->SubNodes.Add({ProcName, &Processor});
+			CurrentNode->Indices.Add(ProcName, ProcNodeIndex);
+		}
+		else
+		{
+			UE_LOG(LogPipe, Warning, TEXT("%s Processor %s already registered. Duplicates are not supported.")
+				, ANSI_TO_TCHAR(__FUNCTION__), *ProcName.ToString());
+		}
+	}
+}
+
+void FProcessorDependencySolver::BuildDependencies(FNode& RootNode)
+{
+	// The main idea here is that for every subnode we convert its dependencies (expressed via ExecuteBefore and ExecuteAfter)
+	// into a single dependency list by reversing the order of ExecuteBefore relationship - so if A.ExecuteBefore contains B then
+	// we turn it into B.ExecuteAfter += A. The dependency names get converted to node indices. If given dependency name 
+	// is unknown at the RootNode's level (i.e. it's not present in its subtree) we make that name a dependency of RootNode. 
+	// Handling of this dependency will then be taken over by RootNode's parent.
+
+	for (int32 NodeIndex = 0; NodeIndex < RootNode.SubNodes.Num(); ++NodeIndex)
+	{
+		FNode& Node = RootNode.SubNodes[NodeIndex];
+		TConstArrayView<FName> ExecuteBefore;
+		TConstArrayView<FName> ExecuteAfter;
+
+		if (Node.Processor)
+		{
+			ExecuteBefore = Node.Processor->GetExecutionOrder().ExecuteBefore;
+			ExecuteAfter = Node.Processor->GetExecutionOrder().ExecuteAfter;
+		}
+		else
+		{
+			BuildDependencies(Node);
+			ExecuteBefore = Node.ExecuteBefore;
+			ExecuteAfter = Node.ExecuteAfter;
+		}
+		
+		for (const FName& DependencyName : ExecuteBefore)
+		{
+			const int32 DependencyIndex = RootNode.FindNodeIndex(DependencyName);
+			// RootNode doesn't know about DependencyName then add it to RootNode's dependencies
+			if (DependencyIndex == INDEX_NONE)
+			{
+				RootNode.ExecuteBefore.AddUnique(DependencyName);
+			}
+			else
+			{
+				if (NodeIndex != DependencyIndex)
+				{
+					// just make DependencyIndex depend on Node
+					RootNode.SubNodes[DependencyIndex].OriginalDependencies.AddUnique(NodeIndex);
+				}
+				else
+				{
+					UE_LOG(LogPipe, Warning, TEXT("%s a node (%s) tried to add a circular dependency on itself")
+						, ANSI_TO_TCHAR(__FUNCTION__), *Node.Name.ToString());
+				}
+			}
+		}
+		for (const FName& DependencyName : ExecuteAfter)
+		{
+			const int32 DependencyIndex = RootNode.FindNodeIndex(DependencyName);
+			// RootNode doesn't know about DependencyName then add it to RootNode's dependencies
+			if (DependencyIndex == INDEX_NONE)
+			{
+				RootNode.ExecuteAfter.AddUnique(DependencyName);
+			}
+			else 
+			{
+				if (NodeIndex != DependencyIndex)
+				{
+					// just make Node depend on DependencyIndex
+					Node.OriginalDependencies.AddUnique(DependencyIndex);
+				}
+				else
+				{
+					UE_LOG(LogPipe, Warning, TEXT("%s a node (%s) tried to add a circular dependency on itself")
+						, ANSI_TO_TCHAR(__FUNCTION__), *Node.Name.ToString());
+				}
+			}
+		}
+	}
+}
+
+void FProcessorDependencySolver::LogNode(const FNode& RootNode, const FNode* ParentNode, int Indent)
+{
+	FString Dependencies;
+	for (const int32 DependencyIndex : RootNode.OriginalDependencies)
+	{
+		if (ensure(ParentNode && ParentNode->SubNodes.IsValidIndex(DependencyIndex)))
+		{
+			Dependencies += ParentNode->SubNodes[DependencyIndex].Name.ToString() + TEXT(", ");
+		}
+	}
+
+	UE_LOG(LogPipe, Log, TEXT("%*s%s %s%s"), Indent, TEXT(""), *RootNode.Name.ToString()
+		, Dependencies.Len() > 0 ? TEXT("after: ") : TEXT(""), *Dependencies);
+
+	if (RootNode.Processor == nullptr)
+	{
+		UE_LOG(LogPipe, Log, TEXT("%*s%s before:%s after:%s"), Indent, TEXT(""), *RootNode.Name.ToString()
+			, *NameViewToString(RootNode.ExecuteBefore)
+			, *NameViewToString(RootNode.ExecuteAfter));
+
+		for (const FNode& Node : RootNode.SubNodes)
+		{
+			LogNode(Node, &RootNode, Indent + 4);
+		}
+	}
+	else
+	{
+		UE_LOG(LogPipe, Log, TEXT("%*s%s before:%s after:%s"), Indent, TEXT(""), *RootNode.Name.ToString()
+			, *NameViewToString(RootNode.Processor->GetExecutionOrder().ExecuteBefore)
+			, *NameViewToString(RootNode.Processor->GetExecutionOrder().ExecuteAfter));
+	}
+}
+
+struct FDumpGraphDependencyUtils
+{
+	static void DumpGraphNode(FArchive& LogFile, const FProcessorDependencySolver::FNode& Node, int Indent, TSet<const FProcessorDependencySolver::FNode*>& AllNodes, const bool bRoot)
+	{
+		const FString NodeName = Node.Name.ToString();
+		if (!Node.Processor)
+		{
+			const FString ClusterNodeName = NodeName.Replace(TEXT("."), TEXT("_"));
+			const FString GraphNodeName = NodeName.Replace(TEXT("."), TEXT(" "));
+
+			int32 Index = -1;
+			Node.Name.ToString().FindLastChar(TEXT('.'), Index);
+			const FString GroupName = NodeName.Mid(Index+1);
+
+			LogFile.Logf(TEXT("%*ssubgraph cluster_%s"), Indent, TEXT(""), *ClusterNodeName);
+			LogFile.Logf(TEXT("%*s{"), Indent, TEXT(""));
+			LogFile.Logf(TEXT("%*slabel =\"%s\";"), Indent + 4, TEXT(""), *GroupName);
+			LogFile.Logf(TEXT("%*s\"%s Start\"%s;"), Indent + 4, TEXT(""), *GraphNodeName, bRoot ? TEXT("") : TEXT("[shape=point style=invis]"));
+			LogFile.Logf(TEXT("%*s\"%s End\"%s;"), Indent + 4, TEXT(""), *GraphNodeName, bRoot ? TEXT("") : TEXT("[shape=point style=invis]"));
+			for (const FProcessorDependencySolver::FNode& SubNode : Node.SubNodes)
+			{
+				DumpGraphNode(LogFile, SubNode, Indent + 4, AllNodes, false);
+			}
+			LogFile.Logf(TEXT("%*s}"), Indent, TEXT(""));
+		}
+		else
+		{
+			LogFile.Logf(TEXT("%*s\"%s\""), Indent, TEXT(""), *NodeName);
+			AllNodes.Add(&Node);
+		}
+	}
+
+	static const FProcessorDependencySolver::FNode* FindDepNodeInParents(const TArray<const FProcessorDependencySolver::FNode*>& Parents, FName DependencyName)
+	{
+		for (int32 i = Parents.Num() - 1; i >= 0; --i)
+		{
+			const FProcessorDependencySolver::FNode* CurNode = Parents[i];
+			int32 DependencyIndex = CurNode->FindNodeIndex(DependencyName);
+			while(DependencyIndex != INDEX_NONE)
+			{
+				if(CurNode->SubNodes[DependencyIndex].Name == DependencyName)
+				{
+					return &CurNode->SubNodes[DependencyIndex];
+				}
+				else
+				{
+					// Dig down the chain
+					CurNode = &CurNode->SubNodes[DependencyIndex];
+					DependencyIndex = CurNode->FindNodeIndex(DependencyName);
+				}
+			}
+		}
+		return nullptr;
+	};
+
+	static bool DoAllSubNodeHasDependency(const FProcessorDependencySolver::FNode& Node, const FName DependencyName, bool bBeforeDep)
+	{
+		bool bDepExistInAllSibling = true;
+		for (const FProcessorDependencySolver::FNode& SiblingSubNode : Node.SubNodes)
+		{
+			if (SiblingSubNode.Processor)
+			{
+				const TArray<FName>& ExecuteDep = bBeforeDep ? SiblingSubNode.Processor->GetExecutionOrder().ExecuteBefore : SiblingSubNode.Processor->GetExecutionOrder().ExecuteAfter;
+				if (ExecuteDep.Find(DependencyName) == INDEX_NONE)
+				{
+					bDepExistInAllSibling = false;
+					break;
+				}
+			}
+			else if (!DoAllSubNodeHasDependency(SiblingSubNode, DependencyName, bBeforeDep))
+			{
+				bDepExistInAllSibling = false;
+				break;
+			}
+		}
+		return bDepExistInAllSibling;
+	}
+
+	static bool DoesDependencyExistIndirectly(const TArray<const FProcessorDependencySolver::FNode*>& Parents, const FProcessorDependencySolver::FNode& Node, const FName DepNameToFind, bool bBeforeDep)
+	{
+		check(Node.Processor);
+		const TArray<FName>& ExecuteDep = bBeforeDep ? Node.Processor->GetExecutionOrder().ExecuteBefore : Node.Processor->GetExecutionOrder().ExecuteAfter;
+		for (const FName& DependencyName : ExecuteDep)
+		{
+			if (DependencyName == DepNameToFind)
+			{
+				continue;
+			}
+			if (const FProcessorDependencySolver::FNode* DepNode = FindDepNodeInParents(Parents, DependencyName))
+			{
+				if(DepNode->Processor)
+				{
+					const TArray<FName>& ExecuteDepDep = bBeforeDep ? DepNode->Processor->GetExecutionOrder().ExecuteBefore : DepNode->Processor->GetExecutionOrder().ExecuteAfter;
+					if(ExecuteDepDep.Find(DepNameToFind) != INDEX_NONE)
+					{
+						return true;
+					}
+				}
+				else if(DoAllSubNodeHasDependency(*DepNode, DepNameToFind, bBeforeDep))
+				{
+					return true;
+				}
+			}
+		}
+		return false;
+	}
+
+	static void RemoveAllProcessorFromSet(const FProcessorDependencySolver::FNode& DepNode, TSet<const FProcessorDependencySolver::FNode*>& Set)
+	{
+		if (DepNode.Processor)
+		{
+			Set.Remove(&DepNode);
+		}
+		else
+		{
+			for (const FProcessorDependencySolver::FNode& SubNode : DepNode.SubNodes)
+			{
+				RemoveAllProcessorFromSet(SubNode, Set);
+			}
+		}
+	}
+
+	static bool AreAllProcessorInSet(const FProcessorDependencySolver::FNode& DepNode, TSet<const FProcessorDependencySolver::FNode*>& Set)
+	{
+		if (DepNode.Processor)
+		{
+			if (!Set.Contains(&DepNode))
+			{
+				return false;
+			}
+		}
+		else
+		{
+			for (const FProcessorDependencySolver::FNode& SubNode : DepNode.SubNodes)
+			{
+				if (!AreAllProcessorInSet(SubNode, Set))
+				{
+					return false;
+				}
+			}
+		}
+		return true;
+	}
+
+	static void DumpGraphDependencies(FArchive& LogFile, const FProcessorDependencySolver::FNode& Node, TArray<const FProcessorDependencySolver::FNode*> Parents, TArray<const FProcessorDependencySolver::FNode*> CommonBeforeDep, TArray<const FProcessorDependencySolver::FNode*> CommonAfterDep, TSet<const FProcessorDependencySolver::FNode*>& DependsOnStart, TSet<const FProcessorDependencySolver::FNode*>& LinkToEnd)
+	{
+		const FString NodeName = Node.Name.ToString();
+		const FProcessorDependencySolver::FNode* ParentNode = Parents.Num() > 0 ? Parents.Last() : nullptr;
+		const FString ParentNodeName = ParentNode ? ParentNode->Name.ToString() : TEXT("");
+		const FString ParentGraphNodeName = ParentNodeName.Replace(TEXT("."), TEXT(" "));
+		const FString ParentClusterNodeName = ParentNodeName.Replace(TEXT("."), TEXT("_"));
+		if (Node.Processor)
+		{
+			for (const FName& DependencyName : Node.Processor->GetExecutionOrder().ExecuteBefore)
+			{
+				if (const FProcessorDependencySolver::FNode* DepNode = FindDepNodeInParents(Parents, DependencyName))
+				{
+					LinkToEnd.Remove(&Node);
+					RemoveAllProcessorFromSet(*DepNode, DependsOnStart);
+
+					if (CommonBeforeDep.Find(DepNode) != INDEX_NONE)
+					{
+						continue;
+					}
+					if (DoesDependencyExistIndirectly(Parents, Node, DependencyName, true))
+					{
+						continue;
+					}
+
+					const FString DepNodeName = DepNode->Name.ToString();
+					if (DepNode->Processor)
+					{
+						LogFile.Logf(TEXT("    \"%s\" -> \"%s\";"), *NodeName, *DepNodeName);
+					}
+					else
+					{
+						const FString DepClusterNodeName = DepNodeName.Replace(TEXT("."), TEXT("_"));
+						const FString DepGraphNodeName = DepNodeName.Replace(TEXT("."), TEXT(" "));
+						LogFile.Logf(TEXT("    \"%s\" -> \"%s Start\"[lhead=cluster_%s];"), *NodeName, *DepGraphNodeName, *DepClusterNodeName);
+					}
+
+				}
+			}
+			for (const FName& DependencyName : Node.Processor->GetExecutionOrder().ExecuteAfter)
+			{
+				if (const FProcessorDependencySolver::FNode* DepNode = FindDepNodeInParents(Parents, DependencyName))
+				{
+					DependsOnStart.Remove(&Node);
+					RemoveAllProcessorFromSet(*DepNode, LinkToEnd);
+
+					if (CommonAfterDep.Find(DepNode) != INDEX_NONE)
+					{
+						continue;
+					}
+					if (DoesDependencyExistIndirectly(Parents, Node, DependencyName, false))
+					{
+						continue;
+					}
+
+					const FString DepNodeName = DepNode->Name.ToString();
+					if (DepNode->Processor)
+					{
+						LogFile.Logf(TEXT("    \"%s\" -> \"%s\";"), *DepNodeName, *NodeName);
+					}
+					else
+					{
+						const FString DepClusterNodeName = DepNodeName.Replace(TEXT("."), TEXT("_"));
+						const FString DepGraphNodeName = DepNodeName.Replace(TEXT("."), TEXT(" "));
+						LogFile.Logf(TEXT("    \"%s End\" -> \"%s\"[ltail=cluster_%s];"), *DepGraphNodeName, *NodeName, *DepClusterNodeName);
+					}
+				}
+			}
+
+			// Layouting
+			LogFile.Logf(TEXT("    \"%s Start\" -> \"%s\" -> \"%s End\"[style=invis];"), *ParentGraphNodeName, *NodeName, *ParentGraphNodeName);
+		}
+		else if (ParentNode)
+		{
+
+			const FString ClusterNodeName = NodeName.Replace(TEXT("."), TEXT("_"));
+			const FString GraphNodeName = NodeName.Replace(TEXT("."), TEXT(" "));
+
+			// Find common dependency through out all sub nodes
+ 			TSet<FName> BeforeDepOuputed;
+ 			TSet<FName> AfterDepOutputed;
+			for (const FProcessorDependencySolver::FNode& SubNode : Node.SubNodes)
+			{
+				if(SubNode.Processor)
+				{
+					for (const FName& DependencyName : SubNode.Processor->GetExecutionOrder().ExecuteBefore)
+					{
+						if (const FProcessorDependencySolver::FNode* DepNode = FindDepNodeInParents(Parents, DependencyName))
+						{
+							if (DoAllSubNodeHasDependency(Node, DependencyName, true /*bBeforeDep*/))
+							{
+								CommonBeforeDep.Add(DepNode);
+
+								if (!DoesDependencyExistIndirectly(Parents, SubNode, DependencyName, true))
+								{
+									if (BeforeDepOuputed.Contains(DependencyName))
+									{
+										continue;
+									}
+									BeforeDepOuputed.Add(DependencyName);
+
+									const FString DepNodeName = DepNode->Name.ToString();
+									if (DepNode->Processor)
+									{
+										LogFile.Logf(TEXT("    \"%s End\" -> \"%s\"[ltail=cluster_%s];"), *GraphNodeName, *DepNodeName, *ClusterNodeName);
+									}
+									else
+									{
+										const FString DepClusterNodeName = DepNodeName.Replace(TEXT("."), TEXT("_"));
+										const FString DepGraphNodeName = DepNodeName.Replace(TEXT("."), TEXT(" "));
+										LogFile.Logf(TEXT("    \"%s End\" -> \"%s Start\"[ltail=cluster_%s, lhead=cluster_%s];"), *GraphNodeName, *DepGraphNodeName, *ClusterNodeName, *DepClusterNodeName);
+									}
+								}
+							}
+						}
+					}
+
+					for (const FName& DependencyName : SubNode.Processor->GetExecutionOrder().ExecuteAfter)
+					{
+						if (const FProcessorDependencySolver::FNode* DepNode = FindDepNodeInParents(Parents, DependencyName))
+						{
+							if (DoAllSubNodeHasDependency(Node, DependencyName, false /*bBeforeDep*/))
+							{
+								CommonAfterDep.Add(DepNode);
+
+								if (!DoesDependencyExistIndirectly(Parents, SubNode, DependencyName, false))
+								{
+									if (AfterDepOutputed.Contains(DependencyName))
+									{
+										continue;
+									}
+									AfterDepOutputed.Add(DependencyName);
+
+									const FString DepNodeName = DepNode->Name.ToString();
+									if (DepNode->Processor)
+									{
+										LogFile.Logf(TEXT("    \"%s\" -> \"%s Start\"[lhead=cluster_%s];"), *DepNodeName, *GraphNodeName, *ClusterNodeName);
+									}
+									else
+									{
+										const FString DepClusterNodeName = DepNodeName.Replace(TEXT("."), TEXT("_"));
+										const FString DepGraphNodeName = DepNodeName.Replace(TEXT("."), TEXT(" "));
+										LogFile.Logf(TEXT("    \"%s End\" -> \"%s Start\"[ltail=cluster_%s, lhead=cluster_%s];"), *DepGraphNodeName,*GraphNodeName, *DepClusterNodeName, *ClusterNodeName);
+									}
+								}
+							}
+						}
+					}
+				}
+			}
+
+ 			// Layouting
+			LogFile.Logf(TEXT("    \"%s Start\" -> \"%s Start\" -> \"%s End\" -> \"%s End\"[style=invis];"), *ParentGraphNodeName, *GraphNodeName, *GraphNodeName, *ParentGraphNodeName);
+		}
+
+		Parents.Push(&Node);
+		for (const FProcessorDependencySolver::FNode& SubNode : Node.SubNodes)
+		{
+			DumpGraphDependencies(LogFile, SubNode, Parents, CommonBeforeDep, CommonAfterDep, DependsOnStart, LinkToEnd);
+		}
+	}
+
+	static void PromoteStartAndEndDependency(FArchive& LogFile, const FString& GraphName, const FProcessorDependencySolver::FNode& Node, TSet<const FProcessorDependencySolver::FNode*>& DependsOnStart, TSet<const FProcessorDependencySolver::FNode*>& LinkToEnd)
+	{
+		if (Node.Processor)
+		{
+			return;
+		}
+
+		bool bAllDependsOnStart = true;
+		bool bAllLinkedToEnd = true;
+		for (const FProcessorDependencySolver::FNode& SubNode : Node.SubNodes)
+		{
+			if (!AreAllProcessorInSet(SubNode, DependsOnStart))
+			{
+				bAllDependsOnStart = false;
+			}
+
+			if (!AreAllProcessorInSet(SubNode, LinkToEnd))
+			{
+				bAllLinkedToEnd = false;
+			}
+
+			if (!bAllDependsOnStart && !bAllLinkedToEnd)
+			{
+				break;
+			}
+		}
+
+		const FString NodeName = Node.Name.ToString();
+		const FString ClusterNodeName = NodeName.Replace(TEXT("."), TEXT("_"));
+		const FString GraphNodeName = NodeName.Replace(TEXT("."), TEXT(" "));
+		if (bAllDependsOnStart)
+		{
+			RemoveAllProcessorFromSet(Node, DependsOnStart);
+			LogFile.Logf(TEXT("    \"%s Start\" -> \"%s Start\"[lhead=cluster_%s, weight=0];"), *GraphName, *GraphNodeName, *ClusterNodeName);
+		}
+
+		if (bAllLinkedToEnd)
+		{
+			RemoveAllProcessorFromSet(Node, LinkToEnd);
+			LogFile.Logf(TEXT("    \"%s End\" -> \"%s End\"[ltail=cluster_%s, weight=0];"), *GraphNodeName, *GraphName, *ClusterNodeName);
+		}
+
+		if(!bAllDependsOnStart || !bAllLinkedToEnd)
+		{
+			for (const FProcessorDependencySolver::FNode& SubNode : Node.SubNodes)
+			{
+				PromoteStartAndEndDependency(LogFile, GraphName, SubNode, DependsOnStart, LinkToEnd);
+			}
+		}
+	}
+
+	static void DumpAllGraphDependencies(FArchive& LogFile, const FProcessorDependencySolver::FNode& GroupRootNode, const TSet<const FProcessorDependencySolver::FNode*>& AllNodes)
+	{
+		TSet<const FProcessorDependencySolver::FNode*> DependsOnStart = AllNodes;
+		TSet<const FProcessorDependencySolver::FNode*> LinkToEnd = AllNodes;
+		const FString GraphName = GroupRootNode.Name.ToString();
+		FDumpGraphDependencyUtils::DumpGraphDependencies(LogFile, GroupRootNode, TArray<const FProcessorDependencySolver::FNode*>(), TArray<const FProcessorDependencySolver::FNode*>(), TArray<const FProcessorDependencySolver::FNode*>(), DependsOnStart, LinkToEnd);
+		FDumpGraphDependencyUtils::PromoteStartAndEndDependency(LogFile, GraphName, GroupRootNode, DependsOnStart, LinkToEnd);
+		for (const FProcessorDependencySolver::FNode* Node : DependsOnStart)
+		{
+			LogFile.Logf(TEXT("    \"%s Start\" -> \"%s\"[weight=0];"), *GraphName, *Node->Name.ToString());
+		}
+		for (const FProcessorDependencySolver::FNode* Node : LinkToEnd)
+		{
+			LogFile.Logf(TEXT("    \"%s\" -> \"%s End\"[weight=0];"), *Node->Name.ToString(), *GraphName);
+		}
+	}
+};
+
+void FProcessorDependencySolver::DumpGraph(FArchive& LogFile) const
+{
+	TSet<const FNode*> AllNodes;
+	LogFile.Logf(TEXT("digraph PipeProcessorGraph"));
+	LogFile.Logf(TEXT("{"));
+	LogFile.Logf(TEXT("    compound = true;"));
+	LogFile.Logf(TEXT("    newrank = true;"));
+	FDumpGraphDependencyUtils::DumpGraphNode(LogFile, GroupRootNode, 4/* Indent */, AllNodes, true/* bRoot */);
+	FDumpGraphDependencyUtils::DumpAllGraphDependencies(LogFile, GroupRootNode, AllNodes);
+	LogFile.Logf(TEXT("}"));
+}
+
+void FProcessorDependencySolver::Solve(FNode& RootNode, TConstArrayView<const FName> PriorityNodes, TArray<FProcessorDependencySolver::FOrderInfo>& OutResult, int LoggingIndent)
+{
+	if (RootNode.SubNodes.Num() == 0)
+	{
+		return;
+	}
+
+	for (FNode& SubNode : RootNode.SubNodes)
+	{
+		SubNode.TransientDependencies = SubNode.OriginalDependencies;
+	}
+
+	TArray<int32> IndicesRemaining;
+	// populate with priority nodes and their dependencies first
+	
+	if (PriorityNodes.Num())
+	{
+		IndicesRemaining.Reserve(RootNode.SubNodes.Num());
+		for (const FName PriorityNodeName : PriorityNodes)
+		{
+			// @todo complex names
+			const int32* PriorityNodeIndex = RootNode.Indices.Find(PriorityNodeName);
+			if (PriorityNodeIndex && RootNode.SubNodes.IsValidIndex(*PriorityNodeIndex)
+				// remember that RootNode.Indices points to a whole group's index for elements from a given group
+				&& RootNode.SubNodes[*PriorityNodeIndex].Name == PriorityNodeName)
+			{ 
+				const int32 StartIndex = IndicesRemaining.Add(*PriorityNodeIndex);
+				// in this loop we'll be potentially adding elements to IndicesRemaining and the idea is to process
+				// the values freshly added as well, so that we populate IndicesRemaining with all super-dependencies of 
+				// the priority node currently processed
+				for (int i = StartIndex; i < IndicesRemaining.Num(); ++i)
+				{
+					const int32 SubNodeIndex = IndicesRemaining[i];
+					for (const int32 DependencyIndex : RootNode.SubNodes[SubNodeIndex].OriginalDependencies)
+					{
+						// some dependencies may have already been added to IndicesRemaining by higher-priority nodes
+						// thus we're adding only unique items (IndicesRemaining is a 1-1 mapping of all RootNode's child nodes).
+						IndicesRemaining.AddUnique(DependencyIndex);
+					}
+				}
+			}
+		}
+	}
+
+	if (IndicesRemaining.Num())
+	{
+		// some node indices already added, fill up with missing ones
+		for (int32 i = 0; i < RootNode.SubNodes.Num(); ++i)
+		{
+			IndicesRemaining.AddUnique(i);
+		}
+	}
+	else
+	{
+		IndicesRemaining.AddZeroed(RootNode.SubNodes.Num());
+		// no priority nodes, just use add all indices in sequence
+		// start from 1 since IndicesRemaining[0] = 0 already
+		for (int32 i = 1; i < RootNode.SubNodes.Num(); ++i)
+		{
+			IndicesRemaining[i] = i;
+		}
+	}
+	
+	TArray<int32> SortedNodeIndices;
+	SortedNodeIndices.Reserve(RootNode.SubNodes.Num());
+
+	bool bCyclesDetected = false;
+	int32 Steps = 0;
+	while (IndicesRemaining.Num())
+	{
+		bCyclesDetected = (PerformPrioritySolverStep(RootNode, IndicesRemaining, SortedNodeIndices) == false);
+		++Steps;
+
+		if (bCyclesDetected)
+		{
+			bAnyCyclesDetected = true;
+
+			UE_LOG(LogPipe, Error, TEXT("Detected processing dependency cycle:"));
+			for (const int32 Index : IndicesRemaining)
+			{
+				UPipeProcessor* Processor = RootNode.SubNodes[Index].Processor;
+				if (Processor)
+				{
+					UE_LOG(LogPipe, Error, TEXT("\t%s, group: %s, before: %s, after %s")
+						, *Processor->GetName()
+						, *Processor->GetExecutionOrder().ExecuteInGroup.ToString()
+						, *NameViewToString(Processor->GetExecutionOrder().ExecuteBefore)
+						, *NameViewToString(Processor->GetExecutionOrder().ExecuteAfter));
+				}
+				else
+				{
+					// group
+					UE_LOG(LogPipe, Error, TEXT("\tGroup %s"), *RootNode.SubNodes[Index].Name.ToString());
+				}
+			}
+			UE_LOG(LogPipe, Error, TEXT("Cutting the chain at an arbitrary location."));
+
+			// remove first dependency
+			// note that if we're in a cycle handling scenario every node does have some dependencies left
+			RootNode.SubNodes[IndicesRemaining[0]].TransientDependencies.Pop(/*bAllowShrinking=*/false);
+		}
+	}
+
+	// now we have the desired order in SortedNodeIndices. We have to traverse it recursively to add to OutResult
+	for (int i = 0; i < SortedNodeIndices.Num(); ++i)
+	{
+		const int32 NodeIndex = SortedNodeIndices[i];
+
+		TArray<FName> DependencyNames;
+		for (const int32 DependencyIndex : RootNode.SubNodes[NodeIndex].OriginalDependencies)
+		{
+			DependencyNames.Add(RootNode.SubNodes[DependencyIndex].Name);
+		}
+
+		if (RootNode.SubNodes[NodeIndex].Processor != nullptr)
+		{
+			OutResult.Add({ RootNode.SubNodes[NodeIndex].Name, RootNode.SubNodes[NodeIndex].Processor, EDependencyNodeType::Processor, DependencyNames });
+		}
+		else if (RootNode.SubNodes[NodeIndex].SubNodes.Num() > 0)
+		{
+			OutResult.Add({ RootNode.SubNodes[NodeIndex].Name, RootNode.SubNodes[NodeIndex].Processor, EDependencyNodeType::GroupStart, DependencyNames });
+			Solve(RootNode.SubNodes[NodeIndex], PriorityNodes, OutResult, LoggingIndent);
+			DependencyNames.Reset(RootNode.SubNodes[NodeIndex].SubNodes.Num());
+			for (const FNode& ChildNode : RootNode.SubNodes[NodeIndex].SubNodes)
+			{
+				DependencyNames.Add(ChildNode.Name);
+			}
+			OutResult.Add({ RootNode.SubNodes[NodeIndex].Name, RootNode.SubNodes[NodeIndex].Processor, EDependencyNodeType::GroupEnd, DependencyNames });
+		}
+	}
+}
+
+void FProcessorDependencySolver::ResolveDependencies(TArray<FProcessorDependencySolver::FOrderInfo>& OutResult, TConstArrayView<const FName> PriorityNodes)
+{
+	FScopedCategoryAndVerbosityOverride LogOverride(TEXT("LogPipe"), ELogVerbosity::Log);
+
+	bAnyCyclesDetected = false;
+
+	UE_LOG(LogPipe, Log, TEXT("Gathering dependencies data:"));
+
+	// gather the processors information first
+	for (UPipeProcessor* Processor : Processors)
+	{
+		if (Processor == nullptr)
+		{
+			UE_LOG(LogPipe, Warning, TEXT("%s nullptr found in Processors collection being processed"), ANSI_TO_TCHAR(__FUNCTION__));
+			continue;
+		}
+		AddNode(Processor->GetExecutionOrder().ExecuteInGroup, *Processor);
+	}
+
+	BuildDependencies(GroupRootNode);
+	// @todo anything in GroupRootNode.Dependencies is an undefined symbol
+
+	// Any dependencies that are promoted to the root node means they are unresolved.
+    const bool bAnyUnresolvedDependencies = GroupRootNode.HasDependencies();
+	if (bAnyUnresolvedDependencies)
+	{
+		const UWorld* World = Processors.IsEmpty() ? nullptr : Processors[0]->GetWorld();
+	    for (const FName& DependencyName : GroupRootNode.ExecuteBefore)
+	    {
+		    UE_LOG(LogPipe, Warning, TEXT("(%s) %s found an unresolved execute before dependency (%s)"),
+		    	World != nullptr ? *ToString(World->GetNetMode()) : TEXT("unknown"), ANSI_TO_TCHAR(__FUNCTION__), *DependencyName.ToString());
+	    }
+    
+	    for (const FName& DependencyName : GroupRootNode.ExecuteAfter)
+	    {
+		    UE_LOG(LogPipe, Warning, TEXT("(%s) %s found an unresolved execute after dependency (%s)"),
+		    	World != nullptr ? *ToString(World->GetNetMode()) : TEXT("unknown"), ANSI_TO_TCHAR(__FUNCTION__), *DependencyName.ToString());
+	    }
+	}
+
+	LogNode(GroupRootNode);
+	
+	if (!DependencyGraphFileName.IsEmpty())
+	{
+		const FString FileName = FString::Printf(TEXT("%s%s-%s.dot"), *FPaths::ProjectLogDir(), *DependencyGraphFileName, *FDateTime::Now().ToString());
+		if (FArchive* LogFile = IFileManager::Get().CreateFileWriter(*FileName))
+		{
+			DumpGraph(*LogFile);
+
+			LogFile->Close();
+			delete LogFile;
+		}
+		else
+		{
+			UE_LOG(LogPipe, Error, TEXT("%s Unable to dump dependency graph into filename %s"), ANSI_TO_TCHAR(__FUNCTION__), *DependencyGraphFileName);
+		}
+	}
+
+	GroupRootNode.TransientDependencies = GroupRootNode.OriginalDependencies;
+	Solve(GroupRootNode, PriorityNodes, OutResult);
+
+#if WITH_UNREAL_DEVELOPER_TOOLS
+	if (bAnyCyclesDetected || bAnyUnresolvedDependencies)
+	{
+		FMessageLog EditorErrors("LogPipe");
+		if(bAnyCyclesDetected)
+		{
+		EditorErrors.Error(LOCTEXT("ProcessorDependenciesCycle", "Processor dependencies cycle found!"));
+		    EditorErrors.Notify(LOCTEXT("ProcessorDependenciesCycle", "Processor dependencies cycle found!"));
+		}
+		if(bAnyUnresolvedDependencies)
+		{
+		    EditorErrors.Error(LOCTEXT("ProcessorUnresolvedDependencies", "Unresolved processor dependencies found!"));
+		    EditorErrors.Notify(LOCTEXT("ProcessorUnresolvedDependencies", "Unresolved processor dependencies found!"));
+		}
+		EditorErrors.Info(FText::FromString(TEXT("See the log for details")));
+	}
+#endif // WITH_UNREAL_DEVELOPER_TOOLS
+
+	UE_LOG(LogPipe, Log, TEXT("Dependency order:"));
+	for (const FProcessorDependencySolver::FOrderInfo& Info : OutResult)
+	{
+		UE_LOG(LogPipe, Log, TEXT("\t%s"), *Info.Name.ToString());
+	}
+}
+
+#undef LOCTEXT_NAMESPACE 
diff --git a/Engine/Plugins/Runtime/MassEntity/Source/MassEntity/Private/MassSchematic.cpp b/Engine/Plugins/Runtime/MassEntity/Source/MassEntity/Private/MassSchematic.cpp
new file mode 100644
--- /dev/null
+++ b/Engine/Plugins/Runtime/MassEntity/Source/MassEntity/Private/MassSchematic.cpp
@@ -0,0 +1,4 @@
+// Copyright Epic Games, Inc. All Rights Reserved.
+
+#include "MassSchematic.h"
+#include "MassProcessor.h"
diff --git a/Engine/Plugins/Runtime/MassEntity/Source/MassEntity/Public/ArchetypeTypes.h b/Engine/Plugins/Runtime/MassEntity/Source/MassEntity/Public/ArchetypeTypes.h
new file mode 100644
--- /dev/null
+++ b/Engine/Plugins/Runtime/MassEntity/Source/MassEntity/Public/ArchetypeTypes.h
@@ -0,0 +1,157 @@
+// Copyright Epic Games, Inc. All Rights Reserved.
+
+#pragma once
+
+#include "UObject/Class.h"
+#include "Containers/ArrayView.h"
+#include "Containers/UnrealString.h"
+#include "LWComponentTypes.h"
+
+class UEntitySubsystem;
+struct FArchetypeData;
+struct FLWComponentSystemExecutionContext;
+struct FLWComponentData;
+struct FArchetypeChunkIterator;
+struct FLWComponentQuery;
+struct FArchetypeChunkCollection;
+struct FEntityView;
+
+typedef TFunction< void(FLWComponentSystemExecutionContext& /*ExecutionContext*/) > FLWComponentSystemExecuteFunction;
+typedef TFunction< bool(const FLWComponentSystemExecutionContext& /*ExecutionContext*/) > FLWComponentSystemChunkConditionFunction;
+
+
+//////////////////////////////////////////////////////////////////////
+//
+
+// An opaque handle to an archetype
+struct FArchetypeHandle final
+{
+	FArchetypeHandle() = default;
+	bool IsValid() const { return DataPtr.IsValid(); }
+
+	MASSENTITY_API bool operator==(const FArchetypeData* Other) const;
+	bool operator==(const FArchetypeHandle& Other) const { return DataPtr == Other.DataPtr; }
+	bool operator!=(const FArchetypeHandle& Other) const { return DataPtr != Other.DataPtr; }
+
+	MASSENTITY_API friend uint32 GetTypeHash(const FArchetypeHandle& Instance);
+private:
+	FArchetypeHandle(const TSharedPtr<FArchetypeData>& InDataPtr)
+	: DataPtr(InDataPtr)
+	{}
+	TSharedPtr<FArchetypeData> DataPtr;
+
+	friend UEntitySubsystem;
+	friend FArchetypeChunkCollection;
+	friend FLWComponentQuery;
+	friend FEntityView;
+};
+
+
+//////////////////////////////////////////////////////////////////////
+//
+
+/** A struct that converts an arbitrary array of entities of given Archetype into a sequence of continuous
+ *  entity chunks. The goal is to have the user create an instance of this struct once and run through a bunch of
+ *  systems. The runtime code usually uses FArchetypeChunkIterator to iterate on the chunk collection
+ */
+struct MASSENTITY_API FArchetypeChunkCollection
+{
+public:
+	struct FChunkInfo
+	{
+		int32 ChunkIndex = INDEX_NONE;
+		int32 SubchunkStart = 0;
+		/** negative or 0-length means "all available entities within chunk" */
+		int32 Length = 0;
+
+		FChunkInfo() = default;
+		explicit FChunkInfo(const int32 InChunkIndex, const int32 InSubchunkStart = 0, const int32 InLength = 0) : ChunkIndex(InChunkIndex), SubchunkStart(InSubchunkStart), Length(InLength) {}
+		/** Note that we consider invalid-length chunks valid as long as ChunkIndex and SubchunkStart are valid */
+		bool IsSet() const { return ChunkIndex != INDEX_NONE && SubchunkStart >= 0; }
+	};
+private:
+	TArray<FChunkInfo> Chunks;
+	/** entity indices indicated by SubChunks are only valid with given Archetype */
+	FArchetypeHandle Archetype;
+
+public:
+	FArchetypeChunkCollection() = default;
+	FArchetypeChunkCollection(const FArchetypeHandle& InArchetype, TConstArrayView<FLWEntity> InEntities);
+	explicit FArchetypeChunkCollection(FArchetypeHandle& InArchetypeHandle);
+	explicit FArchetypeChunkCollection(TSharedPtr<FArchetypeData>& InArchetype);
+
+	TArrayView<const FChunkInfo> GetChunks() const { return Chunks; }
+	const FArchetypeHandle& GetArchetype() const { return Archetype; }
+	bool IsEmpty() const { return Chunks.Num() == 0 && Archetype.IsValid() == false; }
+	bool IsSet() const { return Archetype.IsValid(); }
+	void Reset() 
+	{ 
+		Archetype = FArchetypeHandle();
+		Chunks.Reset();
+	}
+
+private:
+	void GatherChunksFromArchetype(TSharedPtr<FArchetypeData>& InArchetype);
+};
+
+//////////////////////////////////////////////////////////////////////
+//
+
+/**
+ *  The type used to iterate over given archetype's chunks, be it full, continuous chunks or sparse subchunks. It hides
+ *  this details from the rest of the system.
+ */
+struct MASSENTITY_API FArchetypeChunkIterator
+{
+private:
+	const FArchetypeChunkCollection& ChunkData;
+	int32 CurrentChunkIndex = 0;
+
+public:
+	explicit FArchetypeChunkIterator(const FArchetypeChunkCollection& InChunkData) : ChunkData(InChunkData), CurrentChunkIndex(0) {}
+
+	operator bool() const { return ChunkData.GetChunks().IsValidIndex(CurrentChunkIndex) && ChunkData.GetChunks()[CurrentChunkIndex].IsSet(); }
+	FArchetypeChunkIterator& operator++() { ++CurrentChunkIndex; return *this; }
+
+	const FArchetypeChunkCollection::FChunkInfo* operator->() const { check(bool(*this)); return &ChunkData.GetChunks()[CurrentChunkIndex]; }
+	const FArchetypeChunkCollection::FChunkInfo& operator*() const { check(bool(*this)); return ChunkData.GetChunks()[CurrentChunkIndex]; }
+};
+
+//////////////////////////////////////////////////////////////////////
+//
+struct FInternalEntityHandle 
+{
+	FInternalEntityHandle() = default;
+	FInternalEntityHandle(uint8* InChunkRawMemory, const int32 InIndexWithinChunk)
+        : ChunkRawMemory(InChunkRawMemory), IndexWithinChunk(InIndexWithinChunk)
+	{}
+	bool IsValid() const { return ChunkRawMemory != nullptr && IndexWithinChunk != INDEX_NONE; }
+	bool operator==(const FInternalEntityHandle & Other) const { return ChunkRawMemory == Other.ChunkRawMemory && IndexWithinChunk == Other.IndexWithinChunk; }
+
+	uint8* ChunkRawMemory = nullptr;
+	int32 IndexWithinChunk = INDEX_NONE;
+};
+
+typedef TArray<int32, TInlineAllocator<16>> FLWComponentIndicesMapping;
+typedef TConstArrayView<int32> FLWComponentIndicesMappingView;
+struct FLWRequirementIndicesMapping
+{
+	FLWRequirementIndicesMapping() = default;
+
+	FLWComponentIndicesMapping EntityComponents;
+	FLWComponentIndicesMapping ChunkComponents;
+	FORCEINLINE bool IsEmpty() const
+	{
+		return EntityComponents.Num() == 0 || ChunkComponents.Num() == 0;
+	}
+};
+
+template<typename T>
+struct FLWComponentSorterOperator
+{
+	bool operator()(const T& A, const T& B) const
+	{
+		return (A.GetStructureSize() > B.GetStructureSize())
+			|| (A.GetStructureSize() == B.GetStructureSize() && B.GetFName().FastLess(A.GetFName()));
+	}
+};
diff --git a/Engine/Plugins/Runtime/MassEntity/Source/MassEntity/Public/EntityQuery.h b/Engine/Plugins/Runtime/MassEntity/Source/MassEntity/Public/EntityQuery.h
new file mode 100644
--- /dev/null
+++ b/Engine/Plugins/Runtime/MassEntity/Source/MassEntity/Public/EntityQuery.h
@@ -0,0 +1,375 @@
+// Copyright Epic Games, Inc. All Rights Reserved.
+
+#pragma once
+
+#include "UObject/Class.h"
+#include "Containers/ArrayView.h"
+#include "Containers/UnrealString.h"
+#include "LWComponentTypes.h"
+#include "ArchetypeTypes.h"
+#include "EntityQuery.generated.h"
+
+
+class UEntitySubsystem;
+struct FArchetypeData;
+struct FLWComponentSystemExecutionContext;
+struct FLWComponentData;
+struct FArchetypeHandle;
+
+enum class ELWComponentAccess : uint8
+{
+	// no binding required
+	None, 
+
+	// We want to read the data for the component
+	ReadOnly,
+
+	// We want to read and write the data for the component
+	ReadWrite,
+
+	MAX
+};
+
+enum class ELWComponentPresence : uint8
+{
+	// All of the required components must be present
+	All,
+
+	// One of the required components must be present
+	Any,
+
+	// None of the required components can be present
+	None,
+
+	// If component is present we'll use it, but it missing stop processing of a given archetype
+	Optional,
+
+	MAX
+};
+
+
+struct MASSENTITY_API FLWComponentRequirement
+{
+	const UScriptStruct* StructType = nullptr;
+	ELWComponentAccess AccessMode = ELWComponentAccess::None;
+	ELWComponentPresence Presence = ELWComponentPresence::Optional;
+
+public:
+	FLWComponentRequirement(){}
+	FLWComponentRequirement(const UScriptStruct* InStruct, const ELWComponentAccess InAccessMode, const ELWComponentPresence InPresence)
+		: StructType(InStruct)
+		, AccessMode(InAccessMode)
+		, Presence(InPresence)
+	{
+		check(InStruct);
+		checkf((Presence != ELWComponentPresence::Any && Presence != ELWComponentPresence::Optional)
+			|| AccessMode == ELWComponentAccess::ReadOnly || AccessMode == ELWComponentAccess::ReadWrite, TEXT("Only ReadOnly and ReadWrite modes are suppored for optional requirements"));
+	}
+
+	bool RequiresBinding() const { return (AccessMode != ELWComponentAccess::None); }
+	bool IsOptional() const { return (Presence == ELWComponentPresence::Optional || Presence == ELWComponentPresence::Any); }
+
+	FString DebugGetDescription() const;
+
+	// these functions are used for sorting. See FLWComponentSorterOperator
+	int32 GetStructureSize() const
+	{
+		return StructType->GetStructureSize();
+	}
+
+	FName GetFName() const
+	{
+		return StructType->GetFName();
+	}
+};
+
+
+/** 
+ *  FLWComponentQuery is a structure that serves two main purposes:
+ *  1. Describe properties required of an archetype that's a subject of calculations
+ *  2. Trigger calculations on cached set of valid archetypes as described by Requirements
+ * 
+ *  A query to be considered valid needs declared at least one ELWComponentPresence::All, ELWComponentPresence::Any 
+ *  ELWComponentPresence::Optional component requirement.
+ */
+USTRUCT()
+struct MASSENTITY_API FLWComponentQuery
+{
+	GENERATED_BODY()
+
+public:
+	FLWComponentQuery();
+	FLWComponentQuery(std::initializer_list<UScriptStruct*> InitList);
+	FLWComponentQuery(TConstArrayView<const UScriptStruct*> InitList);
+
+	/** Runs ExecuteFunction on all entities matching Requirements */
+	void ForEachEntityChunk(UEntitySubsystem& EntitySubsystem, FLWComponentSystemExecutionContext& ExecutionContext, const FLWComponentSystemExecuteFunction& ExecuteFunction);
+	
+	/** Will first verify that the archetype given with Chunks matches the query's requirements, and if so will run the other, more generic ForEachEntityChunk implementation */
+	void ForEachEntityChunk(const FArchetypeChunkCollection& Chunks, UEntitySubsystem& EntitySubsystem, FLWComponentSystemExecutionContext& ExecutionContext, const FLWComponentSystemExecuteFunction& ExecuteFunction);
+
+	/**
+	 * Attempts to process every chunk of every affected archetype in parallel.
+	 */
+	void ParallelForEachEntityChunk(UEntitySubsystem& EntitySubsystem, FLWComponentSystemExecutionContext& ExecutionContext, const FLWComponentSystemExecuteFunction& ExecuteFunction);
+
+	/** Will gather all archetypes from InEntitySubsystem matching this->Requirements.
+	 *  Note that no work will be done if the cached data is up to date (as tracked by EntitySubsystemHash and 
+	 *	ArchetypeDataVersion properties). */
+	void CacheArchetypes(UEntitySubsystem& InEntitySubsystem);
+
+	FLWComponentQuery& AddRequirement(const UScriptStruct* ComponentType, const ELWComponentAccess AccessMode, const ELWComponentPresence Presence = ELWComponentPresence::All)
+	{
+		checkf(Requirements.FindByPredicate([ComponentType](const FLWComponentRequirement& Item){ return Item.StructType == ComponentType; }) == nullptr
+			, TEXT("Duplicated requirements are not supported. %s already present"), *GetNameSafe(ComponentType));
+		
+		if (Presence != ELWComponentPresence::None)
+		{
+			Requirements.Emplace(ComponentType, AccessMode, Presence);
+		}
+
+		switch (Presence)
+		{
+		case ELWComponentPresence::All:
+			RequiredAllComponents.Add(*ComponentType);
+			break;
+		case ELWComponentPresence::Any:
+			RequiredAnyComponents.Add(*ComponentType);
+			break;
+		case ELWComponentPresence::Optional:
+			RequiredOptionalComponents.Add(*ComponentType);
+			break;
+		case ELWComponentPresence::None:
+			RequiredNoneComponents.Add(*ComponentType);
+			break;
+		}
+		// force recaching the next time this query is used or the following CacheArchetypes call.
+		DirtyCachedData();
+		return *this;
+	}
+
+	/** FLWComponentQuery ref returned for chaining */
+	template<typename T>
+	FLWComponentQuery& AddRequirement(const ELWComponentAccess AccessMode, const ELWComponentPresence Presence = ELWComponentPresence::All)
+	{
+		checkf(Requirements.FindByPredicate([](const FLWComponentRequirement& Item) { return Item.StructType == T::StaticStruct(); }) == nullptr
+			, TEXT("Duplicated requirements are not supported. %s already present"), *T::StaticStruct()->GetName());
+
+		static_assert(TIsDerivedFrom<T, FLWComponentData>::IsDerived, "Given struct doesn't represent a valid fragment type. Make sure to inherit from FLWComponentData or one of its child-types.");
+		
+		if (Presence != ELWComponentPresence::None)
+		{
+			Requirements.Emplace(T::StaticStruct(), AccessMode, Presence);
+		}
+		
+		switch (Presence)
+		{
+		case ELWComponentPresence::All:
+			RequiredAllComponents.Add<T>();
+			break;
+		case ELWComponentPresence::Any:
+			RequiredAnyComponents.Add<T>();
+			break;
+		case ELWComponentPresence::Optional:
+			RequiredOptionalComponents.Add<T>();
+			break;
+		case ELWComponentPresence::None:
+			RequiredNoneComponents.Add<T>();
+			break;
+		}
+		// force recaching the next time this query is used or the following CacheArchetypes call.
+		DirtyCachedData();
+		return *this;
+	}
+
+	void AddTagRequirement(const UScriptStruct& ComponentType, const ELWComponentPresence Presence)
+	{
+		checkfSlow(int(Presence) < int(ELWComponentPresence::Optional), TEXT("Optional and MAX presence are not valid calues for AddTagRequirement"));
+		switch (Presence)
+		{
+		case ELWComponentPresence::All:
+			RequiredAllTags.Add(ComponentType);
+			break;
+		case ELWComponentPresence::Any:
+			RequiredAnyTags.Add(ComponentType);
+			break;
+		case ELWComponentPresence::None:
+			RequiredNoneTags.Add(ComponentType);
+			break;
+		}
+	}
+
+	template<typename T>
+	FLWComponentQuery& AddTagRequirement(const ELWComponentPresence Presence)
+	{
+		checkfSlow(int(Presence) < int(ELWComponentPresence::Optional), TEXT("Optional and MAX presence are not valid calues for AddTagRequirement"));
+		static_assert(TIsDerivedFrom<T, FComponentTag>::IsDerived, "Given struct doesn't represent a valid tag type. Make sure to inherit from FLWComponentData or one of its child-types.");
+		switch (Presence)
+		{
+			case ELWComponentPresence::All:
+				RequiredAllTags.Add<T>();
+				break;
+			case ELWComponentPresence::Any:
+				RequiredAnyTags.Add<T>();
+				break;
+			case ELWComponentPresence::None:
+				RequiredNoneTags.Add<T>();
+				break;
+		}
+
+		return *this;
+	}
+
+	// actual implementation in specializations
+	template<ELWComponentPresence Presence> 
+	FLWComponentQuery& AddTagRequirements(const FLWTagBitSet& TagBitSet)
+	{
+		static_assert(Presence == ELWComponentPresence::None || Presence == ELWComponentPresence::All || Presence == ELWComponentPresence::Any
+			, "The only valid values for AddTagRequirements are All, Any and None");
+		return *this;
+	}
+
+	template<typename T>
+	FLWComponentQuery& AddChunkRequirement(const ELWComponentAccess AccessMode, const ELWComponentPresence Presence = ELWComponentPresence::All)
+	{
+		static_assert(TIsDerivedFrom<T, FLWChunkComponent>::IsDerived, "Given struct doesn't represent a valid chunk fragment type. Make sure to inherit from FLWChunkComponent or one of its child-types.");
+		checkf(ChunkRequirements.FindByPredicate([](const FLWComponentRequirement& Item) { return Item.StructType == T::StaticStruct(); }) == nullptr
+			, TEXT("Duplicated requirements are not supported. %s already present"), *T::StaticStruct()->GetName());
+		checkfSlow(Presence == ELWComponentPresence::None || Presence == ELWComponentPresence::All
+			, TEXT("The only valid Presence values for AddChunkRequirement are All and None"));
+		
+		if (Presence == ELWComponentPresence::All)
+		{
+			RequiredAllChunkComponents.Add<T>();
+			ChunkRequirements.Emplace(T::StaticStruct(), AccessMode, Presence);
+		}
+		else
+		{
+			RequiredNoneChunkComponents.Add<T>();
+		}
+
+		return *this;
+	}
+
+	void Clear()
+	{
+		Requirements.Empty();
+		DirtyCachedData();
+	}
+
+	FORCEINLINE void DirtyCachedData()
+	{
+		EntitySubsystemHash = 0;
+		ArchetypeDataVersion = 0;
+	}
+	
+	bool DoesArchetypeMatchRequirements(const FArchetypeHandle& ArchetypeHandle) const;
+
+	/** The function validates requirements we make for queries. See the FLWComponentQuery struct description for details.
+	 *  Note that this function is non-trivial and end users are not expected to need to use it. 
+	 *  @return whether this query's requirements follow the rules. */
+	bool CheckValidity() const;
+
+	FString DebugGetDescription() const;
+
+	TConstArrayView<FLWComponentRequirement> GetRequirements() const { return Requirements; }
+	const FLWComponentBitSet& GetRequiredAllComponents() const { return RequiredAllComponents; }
+	const FLWComponentBitSet& GetRequiredAnyComponents() const { return RequiredAnyComponents; }
+	const FLWComponentBitSet& GetRequiredOptionalComponents() const { return RequiredOptionalComponents; }
+	const FLWComponentBitSet& GetRequiredNoneComponents() const { return RequiredNoneComponents; }
+	const FLWTagBitSet& GetRequiredAllTags() const { return RequiredAllTags; }
+	const FLWTagBitSet& GetRequiredAnyTags() const { return RequiredAnyTags; }
+	const FLWTagBitSet& GetRequiredNoneTags() const { return RequiredNoneTags; }
+	const FLWChunkComponentBitSet& GetRequiredAllChunkComponents() const { return RequiredAllChunkComponents; }
+	const FLWChunkComponentBitSet& GetRequiredNoneChunkComponents() const { return RequiredNoneChunkComponents; }
+
+	const TArray<FArchetypeHandle>& GetArchetypes() const
+	{ 
+		return ValidArchetypes; 
+	}
+
+	/** 
+	 * Goes through ValidArchetypes and sums up the number of entities contained in them.
+	 * Note that the function is not const because calling it can result in re-caching of ValidArchetypes 
+	 * @return the number of entities this given query would process if called "now"
+	 */
+	int32 GetMatchingEntitiesNum(UEntitySubsystem& InEntitySubsystem);
+
+	/**
+	 * Checks if any of ValidArchetypes has any entities.
+	 * Note that the function is not const because calling it can result in re-caching of ValidArchetypes
+	 * @return "true" if any of the ValidArchetypes has any entities, "false" otherwise
+	 */
+	bool HasMatchingEntities(UEntitySubsystem& InEntitySubsystem);
+
+	/** 
+	 * Sets a chunk filter condition that will applied to each chunk of all valid archetypes. Note 
+	 * that this condition won't be applied when a specific collection of chunks is used (via FArchetypeChunkCollection)
+	 * The value returned by InFunction controls whether to allow execution (true) or block it (false).
+	 */
+	void SetChunkFilter(const FLWComponentSystemChunkConditionFunction& InFunction) { ChunkCondition = InFunction; }
+
+	void ClearChunkFilter() { ChunkCondition.Reset(); }
+
+	bool HasChunkFilter() const { return bool(ChunkCondition); }
+
+protected:
+	void SortRequirements();
+	void ReadCommandlineParams();
+
+protected:
+	TArray<FLWComponentRequirement> Requirements;
+	TArray<FLWComponentRequirement> ChunkRequirements;
+	FLWTagBitSet RequiredAllTags;
+	FLWTagBitSet RequiredAnyTags;
+	FLWTagBitSet RequiredNoneTags;
+	FLWComponentBitSet RequiredAllComponents;
+	FLWComponentBitSet RequiredAnyComponents;
+	FLWComponentBitSet RequiredOptionalComponents;
+	FLWComponentBitSet RequiredNoneComponents;
+	FLWChunkComponentBitSet RequiredAllChunkComponents;
+	FLWChunkComponentBitSet RequiredNoneChunkComponents;
+
+private:
+	/** 
+	 * This function represents a condition that will be called for every chunk to be processed before the actual 
+	 * execution function is called. The chunk component requirements are already bound and ready to be used by the time 
+	 * ChunkCondition is executed.
+	 */
+	FLWComponentSystemChunkConditionFunction ChunkCondition;
+
+	uint32 EntitySubsystemHash = 0;
+	uint32 ArchetypeDataVersion = 0;
+
+	TArray<FArchetypeHandle> ValidArchetypes;
+	TArray<FLWRequirementIndicesMapping> ArchetypeComponentMapping;
+
+	bool bAllowParallelExecution = false;
+};
+
+template<>
+FORCEINLINE FLWComponentQuery& FLWComponentQuery::AddTagRequirements<ELWComponentPresence::All>(const FLWTagBitSet& TagBitSet)
+{
+	RequiredAllTags += TagBitSet;
+	// force recaching the next time this query is used or the following CacheArchetypes call.
+	DirtyCachedData();
+	return *this;
+}
+
+template<>
+FORCEINLINE FLWComponentQuery& FLWComponentQuery::AddTagRequirements<ELWComponentPresence::Any>(const FLWTagBitSet& TagBitSet)
+{
+	RequiredAnyTags += TagBitSet;
+	// force recaching the next time this query is used or the following CacheArchetypes call.
+	DirtyCachedData();
+	return *this;
+}
+
+template<>
+FORCEINLINE FLWComponentQuery& FLWComponentQuery::AddTagRequirements<ELWComponentPresence::None>(const FLWTagBitSet& TagBitSet)
+{
+	RequiredNoneTags += TagBitSet;
+	// force recaching the next time this query is used or the following CacheArchetypes call.
+	DirtyCachedData();
+	return *this;
+}
