AggregateTicking and Pipe merged into Engine/Plugins/Runtime/MassEntity

diff --git a/Engine/Plugins/Runtime/MassEntity/Source/MassEntity/Public/EntitySubsystem.h b/Engine/Plugins/Runtime/MassEntity/Source/MassEntity/Public/EntitySubsystem.h
new file mode 100644
--- /dev/null
+++ b/Engine/Plugins/Runtime/MassEntity/Source/MassEntity/Public/EntitySubsystem.h
@@ -0,0 +1,521 @@
+// Copyright Epic Games, Inc. All Rights Reserved.
+
+#pragma once
+
+#include "Subsystems/WorldSubsystem.h"
+#include "LWComponentTypes.h"
+#include "InstancedStruct.h"
+#include "EntityQuery.h"
+#include "EntitySubsystem.generated.h"
+
+MASSENTITY_API DECLARE_LOG_CATEGORY_EXTERN(LogAggregateTicking, Warning, All);
+
+class UEntitySubsystem;
+struct FLWComponentQuery;
+struct FLWComponentSystemExecutionContext;
+struct FArchetypeData;
+struct FLWCCommandBuffer;
+struct FArchetypeChunkCollection;
+struct FArchetypeChunk;
+class FOutputDevice;
+enum class ELWComponentAccess : uint8;
+
+namespace UE { namespace AggregateTicking {
+FString DebugGetComponentAccessString(ELWComponentAccess Access);
+}}
+
+//@TODO: Comment this guy
+UCLASS()
+class MASSENTITY_API UEntitySubsystem : public UWorldSubsystem
+{
+	GENERATED_BODY()
+
+	friend struct FLWComponentQuery;
+private:
+	// Index 0 is reserved so we can treat that index as an invalid entity handle
+	constexpr static int32 NumReservedEntities = 1;
+
+	struct FEntityData
+	{
+		TSharedPtr<FArchetypeData> CurrentArchetype;
+		int32 SerialNumber = 0;
+
+		void Reset()
+		{
+			CurrentArchetype.Reset();
+			SerialNumber = 0;
+		}
+
+		bool IsValid() const
+		{
+			return SerialNumber != 0 && CurrentArchetype.IsValid();
+		}
+	};
+	
+public:
+	struct FScopedProcessing
+	{
+		explicit FScopedProcessing(std::atomic<int32>& InProcessingScopeCount) : ScopedProcessingCount(InProcessingScopeCount)
+		{
+			++ScopedProcessingCount;
+		}
+		~FScopedProcessing()
+		{
+			--ScopedProcessingCount;
+		}
+	private:
+		std::atomic<int32>& ScopedProcessingCount;
+	};
+
+	const static FLWEntity InvalidEntity;
+
+	UEntitySubsystem();
+
+	//~UObject interface
+	virtual void GetResourceSizeEx(FResourceSizeEx& CumulativeResourceSize) override;
+	//~End of UObject interface
+
+	//~USubsystem interface
+	virtual void Initialize(FSubsystemCollectionBase& Collection) override;
+	virtual void Deinitialize() override;
+	//~End of USubsystem interface
+
+	/** 
+	 * A special, relaxed but slower version of CreateArchetype functions that allows ComponentAngTagsList to contain 
+	 * both components and tags. 
+	 */
+	FArchetypeHandle CreateArchetype(TConstArrayView<const UScriptStruct*> ComponentsAngTagsList);
+
+	FArchetypeHandle CreateArchetype(TConstArrayView<const UScriptStruct*> ComponentList, const FLWTagBitSet& Tags, const FLWChunkComponentBitSet& ChunkComponents);
+
+	/** 
+	 *  Creates an archetype like SourceArchetype + NewComponentList. 
+	 *  @param SourceArchetype the archetype used to initially populate the list of components of the archetype being created. 
+	 *  @param NewComponentList list of unique components to add to components fetched from SourceArchetype. Note that 
+	 *   adding an empty list is not supported and doing so will result in failing a `check`
+	 *  @return a handle of a new archetype
+	 *  @note it's caller's responsibility to ensure that NewComponentList is not empty and contains only component
+	 *   types that SourceArchetype doesn't already have. If the caller cannot guarantee it use of AddComponent functions
+	 *   family is recommended.
+	 */
+	FArchetypeHandle CreateArchetype(const TSharedPtr<FArchetypeData>& SourceArchetype, const FLWComponentBitSet& NewComponentList);
+
+	FArchetypeHandle CreateArchetype(const FLWCompositionDescriptor& Descriptor);
+
+
+	FArchetypeHandle GetArchetypeForEntity(FLWEntity Entity) const;
+	/** Method to iterate on all the component types of an archetype */
+	static void ForEachArchetypeComponentType(const FArchetypeHandle Archetype, TFunction< void(const UScriptStruct* /*ComponentType*/)> Function);
+
+	void SetDefaultChunkComponentValue(const FArchetypeHandle Archetype, FConstStructView InstancedStruct);
+
+	/**
+	 * Go through all archetypes and compact entities
+	 * @param TimeAllowed to do entity compaction, once it reach that time it will stop and return
+	 */
+	void DoEntityCompaction(const double TimeAllowed);
+
+	/**
+	 * Creates fully built entity ready to be used by the subsystem
+	 * @param Archetype you want this entity to be
+	 * @return FLWEntity id of the newly created entity */
+	FLWEntity CreateEntity(const FArchetypeHandle Archetype);
+
+	/**
+	 * Creates fully built entity ready to be used by the subsystem
+	 * @param ComponentInstanceList is the components to create the entity from and initialize values
+	 * @return FLWEntity id of the newly created entity */
+	FLWEntity CreateEntity(TConstArrayView<FInstancedStruct> ComponentInstanceList);
+
+	/** A version of CreateEntity that's creating a number of entities (Count) at one go
+	 *  @param Archetype you want this entity to be
+	 *  @param Count number of entities to create
+	 *  @param OutEntities the newly created entities are appended to given array, i.e. the pre-existing content of OutEntities won't be affected by the call
+	 *  @return number of entities created */
+	int32 BatchCreateEntities(const FArchetypeHandle Archetype, const int32 Count, TArray<FLWEntity>& OutEntities);
+
+	/**
+	 * Destroys a fully built entity, use ReleaseReservedEntity if entity was not yet built.
+	 * @param Entity to destroy */
+	void DestroyEntity(FLWEntity Entity);
+
+	/**
+	 * Reserves an entity in the subsystem, the entity is still not ready to be used by the subsystem, need to call BuildEntity()
+	 * @return FLWEntity id of the reserved entity */
+	FLWEntity ReserveEntity();
+
+	/**
+	 * Builds an entity for it to be ready to be used by the subsystem
+	 * @param Entity to build which was retrieved with ReserveEntity() method
+	 * @param Archetype you want this entity to be*/
+	void BuildEntity(FLWEntity Entity, const FArchetypeHandle Archetype);
+
+	/**
+	 * Builds an entity for it to be ready to be used by the subsystem
+	 * @param Entity to build which was retrieved with ReserveEntity() method
+	 * @param ComponentInstanceList is the components to create the entity from and initialize values*/
+	void BuildEntity(FLWEntity Entity, TConstArrayView<FInstancedStruct> ComponentInstanceList);
+
+	/*
+	 * Releases a previously reserved entity that was not yet built, otherwise call DestroyEntity
+	 * @param Entity to release */
+	void ReleaseReservedEntity(FLWEntity Entity);
+
+	/**
+	 * Destroys all the entity in the provided array of entities
+	 * @param InEntities to destroy
+	 */
+	void BatchDestroyEntities(TConstArrayView<FLWEntity> InEntities);
+
+	void BatchDestroyEntityChunks(const FArchetypeChunkCollection& Chunks);
+
+	void AddComponentToEntity(FLWEntity Entity, const UScriptStruct* ComponentType);
+
+	/** 
+	 *  Ensures that only unique components are added. 
+	 *  @note It's caller's responsibility to ensure Entity's and ComponentList's validity. 
+	 */
+	void AddComponentListToEntity(FLWEntity Entity, TConstArrayView<const UScriptStruct*> ComponentList);
+
+	void AddComponentInstanceListToEntity(FLWEntity Entity, TConstArrayView<FInstancedStruct> ComponentInstanceList);
+	void RemoveComponentFromEntity(FLWEntity Entity, const UScriptStruct* ComponentType);
+	void RemoveComponentListFromEntity(FLWEntity Entity, TConstArrayView<const UScriptStruct*> ComponentList);
+
+	void AddTagToEntity(FLWEntity Entity, const UScriptStruct* TagType);
+	void RemoveTagFromEntity(FLWEntity Entity, const UScriptStruct* TagType);
+	void SwapTagsForEntity(FLWEntity Entity, const UScriptStruct* FromComponentType, const UScriptStruct* ToComponentType);
+
+
+	/**
+	 * Adds fragments and tags indicated by InOutDescriptor to the Entity. The function also figures out which elements
+	 * in InOutDescriptor are missing from the current composition of the given entity and then returns the resulting 
+	 * delta via InOutDescriptor.
+	 */
+	void AddCompositionToEntity_GetDelta(FLWEntity Entity, FLWCompositionDescriptor& InOutDescriptor);
+	void RemoveCompositionFromEntity(FLWEntity Entity, const FLWCompositionDescriptor& InDescriptor);
+
+	/** 
+	 * Moves an entity over to a new archetype by copying over components common to both archetypes
+	 * @param Entity is the entity to move 
+	 * @param NewArchetypeHandle the handle to the new archetype
+	 */
+	void MoveEntityToAnotherArchetype(FLWEntity Entity, FArchetypeHandle NewArchetypeHandle);
+
+	/** Copies values from ComponentInstanceList over to Entity's component. Caller is responsible for ensuring that 
+	 *  the given entity does have given components. Failing this assumption will cause a check-fail.*/
+	void SetEntityComponentsValues(FLWEntity Entity, TArrayView<const FInstancedStruct> ComponentInstanceList);
+
+	/** Copies values from ComponentInstanceList over to components of given entities collection. The caller is responsible 
+	 *  for ensuring that the given entity archetype (FArchetypeChunkCollection.Archetype) does have given components. 
+	 *  Failing this assumption will cause a check-fail. */
+	static void BatchSetEntityComponentsValues(const FArchetypeChunkCollection& SparseEntities, TArrayView<const FInstancedStruct> ComponentInstanceList);
+
+	// Return true if it is an valid built entity
+	bool IsEntityActive(FLWEntity Entity) const 
+	{
+		return IsEntityValid(Entity) && IsEntityBuilt(Entity);
+	}
+
+	// Returns true if Entity is valid
+	bool IsEntityValid(FLWEntity Entity) const;
+
+	// Returns true if Entity is has been fully built (expecting a valid Entity)
+	bool IsEntityBuilt(FLWEntity Entity) const;
+
+	// Asserts that IsEntityValid
+	void CheckIfEntityIsValid(FLWEntity Entity) const;
+
+	// Asserts that IsEntityBuilt
+	void CheckIfEntityIsActive(FLWEntity Entity) const;
+
+	template <typename ComponentType>
+	ComponentType& GetComponentDataChecked(FLWEntity Entity) const
+	{
+		return *((ComponentType*)InternalGetComponentDataChecked(Entity, ComponentType::StaticStruct()));
+	}
+
+	template <typename ComponentType>
+	ComponentType* GetComponentDataPtr(FLWEntity Entity) const
+	{
+		return (ComponentType*)InternalGetComponentDataPtr(Entity, ComponentType::StaticStruct());
+	}
+
+	FStructView GetComponentDataStruct(FLWEntity Entity, const UScriptStruct* ComponentType) const
+	{
+		return FStructView(ComponentType, static_cast<uint8*>(InternalGetComponentDataPtr(Entity, ComponentType)));
+	}
+
+	uint32 GetArchetypeDataVersion() const { return ArchetypeDataVersion; }
+
+	/**
+	 * Creates and initializes a FLWComponentSystemExecutionContext instance.
+	 */
+	FLWComponentSystemExecutionContext CreateExecutionContext(const float DeltaSeconds) const;
+
+	FScopedProcessing NewProcessingScope() { return FScopedProcessing(ProcessingScopeCount); }
+	bool IsProcessing() const { return ProcessingScopeCount > 0; }
+	FLWCCommandBuffer& Defer() { return *DeferredCommandBuffer.Get(); }
+
+#if WITH_AGGREGATETICKING_DEBUG
+	void DebugPrintEntity(int32 Index, FOutputDevice& Ar, const TCHAR* InPrefix = TEXT("")) const;
+	void DebugPrintEntity(FLWEntity Entity, FOutputDevice& Ar, const TCHAR* InPrefix = TEXT("")) const;
+	void DebugPrintArchetypes(FOutputDevice& Ar) const;
+	static void DebugGetStringDesc(const FArchetypeHandle& Archetype, FOutputDevice& Ar);
+	void DebugGetArchetypesStringDetails(FOutputDevice& Ar, const bool bIncludeEmpty = true);
+	void DebugGetArchetypeComponentTypes(const FArchetypeHandle& Archetype, TArray<const UScriptStruct*>& InOutComponentList) const;
+	int32 DebugGetArchetypeEntitiesCount(const FArchetypeHandle& Archetype) const;	
+	int32 DebugGetEntityCount() const { return Entities.Num() - NumReservedEntities - EntityFreeIndexList.Num(); }
+	int32 DebugGetArchetypesCount() const { return ComponentHashToArchetypeMap.Num(); }
+	void DebugRemoveAllEntities();
+	void DebugForceArchetypeDataVersionBump() { ++ArchetypeDataVersion; }
+	void DebugGetArchetypeStrings(const FArchetypeHandle& Archetype, TArray<FName>& OutComponentNames, TArray<FName>& OutTagNames);
+#endif // WITH_AGGREGATETICKING_DEBUG
+
+protected:
+	void GetValidArchetypes(const FLWComponentQuery& Query, TArray<FArchetypeHandle>& OutValidArchetypes);
+	
+	FArchetypeHandle InternalCreateSiblingArchetype(const TSharedPtr<FArchetypeData>& SourceArchetype, const FLWTagBitSet& OverrideTags);
+
+private:
+	void InternalBuildEntity(FLWEntity Entity, const FArchetypeHandle Archetype);
+	void InternalReleaseEntity(FLWEntity Entity);
+
+	/** 
+	 *  Adds components in ComponentList to Entity. Only the unique components will be added.
+	 */
+	void InternalAddComponentListToEntityChecked(FLWEntity Entity, const FLWComponentBitSet& InComponents);
+
+	/** 
+	 *  Similar to InternalAddComponentListToEntity but expects NewComponentList not overlapping with current entity's
+	 *  component list. It's callers responsibility to ensure that's true. Failing this will cause a `check` fail.
+	 */
+	void InternalAddComponentListToEntity(FLWEntity Entity, const FLWComponentBitSet& NewComponents);
+	void* InternalGetComponentDataChecked(FLWEntity Entity, const UScriptStruct* ComponentType) const;
+	void* InternalGetComponentDataPtr(FLWEntity Entity, const UScriptStruct* ComponentType) const;
+
+	//@TODO: Not optimized or a great API, super WIP
+	struct MASSENTITY_API FChunkIteratorByComponent
+	{
+		FChunkIteratorByComponent(UEntitySubsystem* InSystem, const UScriptStruct* InComponentType);
+
+		void GetCurrentChunkBase(void*& OutChunkBasePtr, int32& OutChunkEntityCount) const;
+
+		void Advance();
+		bool IsDone() const;
+		FArchetypeData* GetCurrentArchetype() const;
+
+		UEntitySubsystem* System;
+		const UScriptStruct* ComponentType;
+
+		TArray<TSharedPtr<FArchetypeData>>* pPendingArchetypes;
+		int32 ArchetypeIndex = 0;
+		int32 ChunkIndex = 0;
+	};
+
+	friend FChunkIteratorByComponent;
+
+private:
+	TChunkedArray<FEntityData> Entities;
+	TArray<int32> EntityFreeIndexList;
+
+	TSharedPtr<FLWCCommandBuffer> DeferredCommandBuffer;
+	std::atomic<int32> SerialNumberGenerator;
+	std::atomic<int32> ProcessingScopeCount;
+
+	// the "version" number increased every time an archetype gets added
+	uint32 ArchetypeDataVersion = 0;
+
+	// Map of hash of sorted component list to archetypes with that hash
+	TMap<uint32, TArray<TSharedPtr<FArchetypeData>>> ComponentHashToArchetypeMap;
+
+	// Map to list of archetypes that contain the specified component type
+	TMap<const UScriptStruct*, TArray<TSharedPtr<FArchetypeData>>> ComponentTypeToArchetypeMap;
+};
+
+
+//////////////////////////////////////////////////////////////////////
+//
+
+struct MASSENTITY_API FLWComponentSystemExecutionContext
+{
+private:
+	struct FComponentView 
+	{
+		FLWComponentRequirement Requirement;
+		TArrayView<FLWComponentData> ComponentView;
+
+		FComponentView() {}
+		explicit FComponentView(const FLWComponentRequirement& InRequirement) : Requirement(InRequirement) {}
+
+		bool operator==(const UScriptStruct* ComponentType) const { return Requirement.StructType == ComponentType; }
+	};
+	TArray<FComponentView, TInlineAllocator<8>> ComponentViews;
+
+	struct FChunkComponentView
+	{
+		FLWComponentRequirement Requirement;
+		FStructView ChunkComponentView;
+
+		FChunkComponentView() {}
+		explicit FChunkComponentView(const FLWComponentRequirement& InRequirement) : Requirement(InRequirement)	{}
+
+		bool operator==(const UScriptStruct* ComponentType) const { return Requirement.StructType == ComponentType; }
+	};
+	TArray<FChunkComponentView, TInlineAllocator<4>> ChunkComponents;
+
+	// mz@todo make this shared ptr thread-safe and never auto-flush in MT environment. 
+	TSharedPtr<FLWCCommandBuffer> DeferredCommandBuffer;
+	TArrayView<FLWEntity> EntityListView;
+	
+	/** If set this indicates the exact archetype and its chunks to be processed. 
+	 *  @todo this data should live somewhere else, preferably be just a parameter to Query.ForEachEntityChunk function */
+	FArchetypeChunkCollection ChunkCollection;
+	
+	FInstancedStruct AuxData;
+	float DeltaTimeSeconds = 0.0f;
+	int32 ChunkSerialModificationNumber = -1;
+	FLWTagBitSet CurrentArchetypesTagBitSet;
+
+#if WITH_AGGREGATETICKING_DEBUG
+	FString DebugExecutionDescription;
+#endif
+	
+	/** @todo update comment, and look if we If true the EntitySystem will flush the deferred commands stored in DeferredCommandBuffer just after executing 
+	 *  the given system. If False then the party calling UEntitySubsystem::ExecuteSystem is responsible for manually
+	 *  calling FLWCCommandBuffer.ReplayBufferAgainstSystem() */
+	bool bFlushDeferredCommands = true;
+
+	TArrayView<FComponentView> GetMutableRequirements() { return ComponentViews; }
+	TArrayView<FChunkComponentView> GetMutableChunkRequirements() { return ChunkComponents; }
+	
+	friend FArchetypeData;
+	friend FLWComponentQuery;
+
+public:
+	FLWComponentSystemExecutionContext() = default;
+	explicit FLWComponentSystemExecutionContext(const float InDeltaTimeSeconds, const bool bInFlushDeferredCommands = true)
+		: DeltaTimeSeconds(InDeltaTimeSeconds)
+		, bFlushDeferredCommands(bInFlushDeferredCommands)
+	{}
+
+#if WITH_AGGREGATETICKING_DEBUG
+	const FString& DebugGetExecutionDesc() const { return DebugExecutionDescription; }
+	void DebugSetExecutionDesc(const FString& Description) { DebugExecutionDescription = Description; }
+#endif
+
+	/** Sets bFlushDeferredCommands. Note that setting to True while the system is being executed doesn't result in
+	 *  immediate commands flushing */
+	void SetFlushDeferredCommands(const bool bNewFlushDeferredCommands) { bFlushDeferredCommands = bNewFlushDeferredCommands; } 
+	void SetDeferredCommandBuffer(const TSharedPtr<FLWCCommandBuffer>& InDeferredCommandBuffer) { DeferredCommandBuffer = InDeferredCommandBuffer; }
+	void SetChunkCollection(const FArchetypeChunkCollection& InChunkCollection);
+	void SetChunkCollection(FArchetypeChunkCollection&& InChunkCollection);
+	void ClearChunkCollection() { ChunkCollection.Reset(); }
+	void SetAuxData(const FInstancedStruct& InAuxData) { AuxData = InAuxData; }
+
+	float GetDeltaTimeSeconds() const
+	{
+		return DeltaTimeSeconds;
+	}
+
+	TSharedPtr<FLWCCommandBuffer> GetSharedDeferredCommandBuffer() const { return DeferredCommandBuffer; }
+	FLWCCommandBuffer& Defer() { checkSlow(DeferredCommandBuffer.IsValid()); return *DeferredCommandBuffer.Get(); }
+
+	TConstArrayView<FLWEntity> GetEntities() const { return EntityListView; }
+	int32 GetEntitiesNum() const { return EntityListView.Num(); }
+
+	FLWEntity GetEntity(const int32 Index) const
+	{
+		return EntityListView[Index];
+	}
+
+	template<typename T>
+	bool DoesArchetypeHaveTag() const
+	{
+		static_assert(TIsDerivedFrom<T, FComponentTag>::IsDerived, "Given struct is not of a valid fragment type.");
+		return CurrentArchetypesTagBitSet.Contains<T>();
+	}
+
+	/** Chunk related operation */
+	void SetCurrentChunkSerialModificationNumber(const int32 SerialModificationNumber) { ChunkSerialModificationNumber = SerialModificationNumber; }
+	int32 GetChunkSerialModificationNumber() const { return ChunkSerialModificationNumber; }
+
+	template<typename T>
+	T& GetMutableChunkComponent()
+	{
+		static_assert(TIsDerivedFrom<T, FLWChunkComponent>::IsDerived, "Given struct doesn't represent a valid chunk fragment type. Make sure to inherit from FLWChunkComponent or one of its child-types.");
+
+		const UScriptStruct* Type = T::StaticStruct();
+		const FChunkComponentView* FoundChunkComponentData = ChunkComponents.FindByPredicate([Type](const FChunkComponentView& Element) { return Element.Requirement.StructType == Type; });
+		checkf(FoundChunkComponentData, TEXT("Chunk Component requirement not found: %s"), *T::StaticStruct()->GetName());
+		return FoundChunkComponentData->ChunkComponentView.template GetMutable<T>();
+	}
+
+	template<typename T>
+	const T& GetChunkComponent() const
+	{
+		static_assert(TIsDerivedFrom<T, FLWChunkComponent>::IsDerived, "Given struct doesn't represent a valid chunk fragment type. Make sure to inherit from FLWChunkComponent or one of its child-types.");
+
+		const UScriptStruct* Type = T::StaticStruct();
+		const FChunkComponentView* FoundChunkComponentData = ChunkComponents.FindByPredicate([Type](const FChunkComponentView& Element) { return Element.Requirement.StructType == Type; } );
+		checkf(FoundChunkComponentData, TEXT("Chunk Component requirement not found: %s"), *T::StaticStruct()->GetName());
+		return FoundChunkComponentData->ChunkComponentView.template Get<T>();
+	}
+
+	template<typename TLWComponent>
+	TArrayView<TLWComponent> GetMutableComponentView()
+	{
+		const UScriptStruct* ComponentType = TLWComponent::StaticStruct();
+		const FComponentView* View = ComponentViews.FindByPredicate([ComponentType](const FComponentView& Element) { return Element.Requirement.StructType == ComponentType; });
+		//checkfSlow(View != nullptr, TEXT("Requested component type not bound"));
+		//checkfSlow(View->Requirement.AccessMode == ELWComponentAccess::ReadWrite, TEXT("Requested component has not been bound for writing"));
+		return MakeArrayView<TLWComponent>((TLWComponent*)View->ComponentView.GetData(), View->ComponentView.Num());
+	}
+
+	template<typename TLWComponent>
+	TConstArrayView<TLWComponent> GetComponentView() const
+	{
+		const UScriptStruct* ComponentType = TLWComponent::StaticStruct();
+		const FComponentView* View = ComponentViews.FindByPredicate([ComponentType](const FComponentView& Element) { return Element.Requirement.StructType == ComponentType; });
+		//checkfSlow(View != nullptr, TEXT("Requested component type not bound"));
+		return TConstArrayView<TLWComponent>((const TLWComponent*)View->ComponentView.GetData(), View->ComponentView.Num());
+	}
+
+	TConstArrayView<FLWComponentData> GetComponentComponentView(const UScriptStruct* ComponentType) const
+	{
+		const FComponentView* View = ComponentViews.FindByPredicate([ComponentType](const FComponentView& Element) { return Element.Requirement.StructType == ComponentType; });
+		checkSlow(View);
+		return TConstArrayView<FLWComponentData>((const FLWComponentData*)View->ComponentView.GetData(), View->ComponentView.Num());;
+	}
+
+	TArrayView<FLWComponentData> GetMutableComponentView(const UScriptStruct* ComponentType) 
+	{
+		const FComponentView* View = ComponentViews.FindByPredicate([ComponentType](const FComponentView& Element) { return Element.Requirement.StructType == ComponentType; });
+		checkSlow(View);
+		return View->ComponentView;
+	}
+
+	/** Sparse chunk related operation */
+	const FArchetypeChunkCollection& GetChunkCollection() const { return ChunkCollection; }
+
+	const FInstancedStruct& GetAuxData() const { return AuxData; }
+	FInstancedStruct& GetMutableAuxData() { return AuxData; }
+
+	void FlushDeferred(UEntitySubsystem& EntitySystem) const;
+
+	void ClearExecutionData();
+	void SetCurrentArchetypeData(FArchetypeData& ArchetypeData);
+
+protected:
+	void SetRequirements(TConstArrayView<FLWComponentRequirement> InRequirements, TConstArrayView<FLWComponentRequirement> InChunkRequirements);
+	void ClearComponentViews()
+	{
+		for (FComponentView& View : ComponentViews)
+		{
+			View.ComponentView = TArrayView<FLWComponentData>();
+		}
+		for (FChunkComponentView& View : ChunkComponents)
+		{
+			View.ChunkComponentView.Reset();
+		}
+	}
+};
diff --git a/Engine/Plugins/Runtime/MassEntity/Source/MassEntity/Public/EntityView.h b/Engine/Plugins/Runtime/MassEntity/Source/MassEntity/Public/EntityView.h
new file mode 100644
--- /dev/null
+++ b/Engine/Plugins/Runtime/MassEntity/Source/MassEntity/Public/EntityView.h
@@ -0,0 +1,91 @@
+// Copyright Epic Games, Inc. All Rights Reserved.
+
+#pragma once
+
+#include "ArchetypeTypes.h"
+#include "InstancedStruct.h"
+#include "EntityView.generated.h"
+
+
+class UEntitySubsystem;
+struct FArchetypeData;
+struct FArchetypeHandle;
+struct FArchetypeComponentConfig;
+
+/** 
+ * The type representing a single entity in a single archetype. It's of a very transient nature so we guarantee it's 
+ * validity only within the scope it has been created in. Don't store it. 
+ */
+USTRUCT()
+struct MASSENTITY_API FEntityView
+{
+	GENERATED_BODY()
+
+	FEntityView() = default;
+
+	/** 
+	 *  Resolves Entity against ArchetypeHandle. Note that this approach requires the caller to ensure that Entity
+	 *  indeed belongs to ArchetypeHandle. If not the call will fail a check. As a remedy calling the 
+	 *  UEntitySubsystem-flavored constructor is recommended since it will first find the appropriate archetype for
+	 *  Entity. 
+	 */
+	FEntityView(const FArchetypeHandle& ArchetypeHandle, FLWEntity Entity);
+
+	/** 
+	 *  Finds the archetype Entity belongs to and then resolves against it. The caller is responsible for ensuring
+	 *  that the given Entity is in fact a valid ID tied to any of the archetypes 
+	 */
+	FEntityView(const UEntitySubsystem& EntitySubsystem, FLWEntity Entity);
+
+	FLWEntity GetEntity() const	{ return Entity; }
+
+	/** will fail a check if the viewed entity doesn't have the given component */	
+	template<typename T>
+	T& GetComponentData() const
+	{
+		static_assert(!TIsDerivedFrom<T, FComponentTag>::IsDerived,
+			"Given struct doesn't represent a valid fragment type but a tag. Use HasTag instead.");
+		static_assert(TIsDerivedFrom<T, FComponentTag>::IsDerived || TIsDerivedFrom<T, FLWComponentData>::IsDerived,
+			"Given struct doesn't represent a valid fragment type. Make sure to inherit from FLWComponentData or one of its child-types.");
+
+		return *((T*)GetComponentPtrChecked(*T::StaticStruct()));
+	}
+		
+	/** if the viewed entity doesn't have the given component the function will return null */
+	template<typename T>
+	T* GetComponentDataPtr() const
+	{
+		static_assert(!TIsDerivedFrom<T, FComponentTag>::IsDerived,
+			"Given struct doesn't represent a valid fragment type but a tag. Use HasTag instead.");
+		static_assert(TIsDerivedFrom<T, FComponentTag>::IsDerived || TIsDerivedFrom<T, FLWComponentData>::IsDerived,
+			"Given struct doesn't represent a valid fragment type. Make sure to inherit from FLWComponentData or one of its child-types.");
+
+		return (T*)GetComponentPtr(*T::StaticStruct());
+	}
+
+	FStructView GetComponentDataStruct(const UScriptStruct* ComponentType) const
+	{
+		check(ComponentType);
+		return FStructView(ComponentType, static_cast<uint8*>(GetComponentPtr(*ComponentType)));
+	}
+
+	template<typename T>
+	bool HasTag() const
+	{
+		static_assert(TIsDerivedFrom<T, FComponentTag>::IsDerived, "Given struct doesn't represent a valid tag type. Make sure to inherit from FComponentTag or one of its child-types.");
+		return HasTag(*T::StaticStruct());
+	}
+
+	bool IsSet() const { return Archetype != nullptr && EntityHandle.IsValid(); }
+	bool operator==(const FEntityView& Other) const { return Archetype == Other.Archetype && EntityHandle == Other.EntityHandle; }
+
+protected:
+	void* GetComponentPtr(const UScriptStruct& ComponentType) const;
+	void* GetComponentPtrChecked(const UScriptStruct& ComponentType) const;
+	bool HasTag(const UScriptStruct& TagType) const;
+
+private:
+	FLWEntity Entity;
+	FInternalEntityHandle EntityHandle;
+	FArchetypeData* Archetype = nullptr;
+};
diff --git a/Engine/Plugins/Runtime/MassEntity/Source/MassEntity/Public/IMassEntityModule.h b/Engine/Plugins/Runtime/MassEntity/Source/MassEntity/Public/IMassEntityModule.h
new file mode 100644
--- /dev/null
+++ b/Engine/Plugins/Runtime/MassEntity/Source/MassEntity/Public/IMassEntityModule.h
@@ -0,0 +1,40 @@
+// Copyright Epic Games, Inc. All Rights Reserved.
+
+#pragma once
+
+#include "CoreMinimal.h"
+#include "Modules/ModuleInterface.h"
+#include "Modules/ModuleManager.h"
+
+
+/**
+ * The public interface to this module.  In most cases, this interface is only public to sibling modules 
+ * within this plugin.
+ */
+class IPipeModule : public IModuleInterface
+{
+
+public:
+
+	/**
+	 * Singleton-like access to this module's interface.  This is just for convenience!
+	 * Beware of calling this during the shutdown phase, though.  Your module might have been unloaded already.
+	 *
+	 * @return Returns singleton instance, loading the module on demand if needed
+	 */
+	static inline IPipeModule& Get()
+	{
+		return FModuleManager::LoadModuleChecked<IPipeModule>( "Pipe" );
+	}
+
+	/**
+	 * Checks to see if this module is loaded and ready.  It is only valid to call Get() if IsAvailable() returns true.
+	 *
+	 * @return True if the module is loaded and ready to use
+	 */
+	static inline bool IsAvailable()
+	{
+		return FModuleManager::Get().IsModuleLoaded( "Pipe" );
+	}
+};
+
diff --git a/Engine/Plugins/Runtime/MassEntity/Source/MassEntity/Public/LWCCommandBuffer.h b/Engine/Plugins/Runtime/MassEntity/Source/MassEntity/Public/LWCCommandBuffer.h
new file mode 100644
--- /dev/null
+++ b/Engine/Plugins/Runtime/MassEntity/Source/MassEntity/Public/LWCCommandBuffer.h
@@ -0,0 +1,368 @@
+// Copyright Epic Games, Inc. All Rights Reserved.
+
+#pragma once
+
+#include "LWComponentTypes.h"
+#include "EntitySubsystem.h"
+#include "InstancedStructStream.h"
+#include "Misc/MTAccessDetector.h"
+
+#include "LWCCommandBuffer.generated.h"
+
+//@TODO: Consider debug information in case there is an assert when replaying the command buffer
+// (e.g., which system added the command, or even file/line number in development builds for the specific call via a macro)
+//@TODO: Support more commands
+
+USTRUCT()
+struct MASSENTITY_API FCommandBufferEntryBase
+{
+	GENERATED_BODY()
+
+	FLWEntity TargetEntity;
+
+	FCommandBufferEntryBase() = default;
+	virtual ~FCommandBufferEntryBase() = default;
+
+	FCommandBufferEntryBase(FLWEntity InEntity)
+		: TargetEntity(InEntity)
+	{}
+
+	virtual void Execute(UEntitySubsystem& System) const PURE_VIRTUAL(FCommandBufferEntryBase::Execute, );
+};
+template<> struct TStructOpsTypeTraits<FCommandBufferEntryBase> : public TStructOpsTypeTraitsBase2<FCommandBufferEntryBase> { enum { WithPureVirtual = true, }; };
+
+/**
+ * Command dedicated to execute a captured lambda 
+ */
+USTRUCT()
+struct MASSENTITY_API FDeferredCommand : public FCommandBufferEntryBase
+{
+	GENERATED_BODY()
+
+	FDeferredCommand() = default;
+
+	FDeferredCommand(const TFunction<void(UEntitySubsystem& System)>& InDeferredFunction)
+		: DeferredFunction(InDeferredFunction)
+	{}
+
+	virtual void Execute(UEntitySubsystem& System) const override
+	{
+		DeferredFunction(System);
+	}
+
+protected:
+	TFunction<void(UEntitySubsystem& System)> DeferredFunction;
+};
+
+/**
+ * Command dedicated to building an entity from an existing component instance
+ */
+USTRUCT()
+struct MASSENTITY_API FBuildEntityFromComponentInstance : public FCommandBufferEntryBase
+{
+	GENERATED_BODY()
+
+	FBuildEntityFromComponentInstance() = default;
+	FBuildEntityFromComponentInstance(const FLWEntity Entity, FStructView InStruct)
+		: FCommandBufferEntryBase(Entity)
+		, Struct(InStruct)
+	{}
+
+protected:
+	virtual void Execute(UEntitySubsystem& System) const override;
+
+	FInstancedStruct Struct;
+};
+
+/**
+* Command dedicated to building an entity from a list of existing component instances
+*/
+USTRUCT()
+struct MASSENTITY_API FBuildEntityFromComponentInstances : public FCommandBufferEntryBase
+{
+	GENERATED_BODY()
+
+	FBuildEntityFromComponentInstances() = default;
+	FBuildEntityFromComponentInstances(const FLWEntity Entity, TConstArrayView<FInstancedStruct> InInstances)
+		: FCommandBufferEntryBase(Entity)
+		, Instances(InInstances)
+	{}
+
+protected:
+	virtual void Execute(UEntitySubsystem& System) const override;
+
+	TArray<FInstancedStruct> Instances;
+};
+
+/**
+ * Command dedicated to add a new component to an existing entity
+ */
+USTRUCT()
+struct MASSENTITY_API FCommandAddComponent : public FCommandBufferEntryBase
+{
+	GENERATED_BODY()
+
+	FCommandAddComponent() = default;
+	FCommandAddComponent(FLWEntity InEntity, UScriptStruct* InStruct)
+		: FCommandBufferEntryBase(InEntity)
+		, StructParam(InStruct)
+	{}
+
+protected:
+	virtual void Execute(UEntitySubsystem& System) const override;
+
+	UScriptStruct* StructParam = nullptr;
+};
+
+/**
+ * Command dedicated to add a new component from an existing instance
+ */
+USTRUCT()
+struct MASSENTITY_API FCommandAddComponentInstance : public FCommandBufferEntryBase
+{
+	GENERATED_BODY()
+
+	FCommandAddComponentInstance() = default;
+	FCommandAddComponentInstance(const FLWEntity InEntity, FStructView InStruct)
+        : FCommandBufferEntryBase(InEntity)
+        , Struct(InStruct)
+	{}
+
+protected:
+	virtual void Execute(UEntitySubsystem& System) const override;
+
+	FInstancedStruct Struct;
+};
+
+/**
+ * Command dedicated to remove a component from an existing entity
+ */
+USTRUCT()
+struct MASSENTITY_API FCommandRemoveComponent : public FCommandBufferEntryBase
+{
+	GENERATED_BODY()
+
+	FCommandRemoveComponent() = default;
+	FCommandRemoveComponent(FLWEntity InEntity, UScriptStruct* InStruct)
+		: FCommandBufferEntryBase(InEntity)
+		, StructParam(InStruct)
+	{}
+
+protected:
+	virtual void Execute(UEntitySubsystem& System) const override;
+
+	UScriptStruct* StructParam = nullptr;
+};
+
+/**
+ * Command performing addition of a list of components to an existing entity
+ */
+USTRUCT()
+struct MASSENTITY_API FCommandAddComponentList : public FCommandBufferEntryBase
+{
+	GENERATED_BODY()
+
+	FCommandAddComponentList() = default;
+	FCommandAddComponentList(FLWEntity InEntity, TConstArrayView<const UScriptStruct*> InComponentList)
+		: FCommandBufferEntryBase(InEntity)
+		, ComponentList(InComponentList)
+	{}
+
+protected:
+	virtual void Execute(UEntitySubsystem& System) const override;
+
+	TArray<const UScriptStruct*> ComponentList;
+};
+
+/**
+ * Command performing a removal of a list of components from an existing entity
+ */
+USTRUCT()
+struct MASSENTITY_API FCommandRemoveComponentList : public FCommandBufferEntryBase
+{
+	GENERATED_BODY()
+
+	FCommandRemoveComponentList() = default;
+	FCommandRemoveComponentList(FLWEntity InEntity, TConstArrayView<const UScriptStruct*> InComponentList)
+		: FCommandBufferEntryBase(InEntity)
+		, ComponentList(InComponentList)
+	{}
+
+protected:
+	virtual void Execute(UEntitySubsystem& System) const override;
+
+	TArray<const UScriptStruct*> ComponentList;
+};
+
+/**
+ * Command dedicated to add a new tag to an existing entity
+ */
+USTRUCT()
+struct MASSENTITY_API FCommandAddTag: public FCommandBufferEntryBase
+{
+	GENERATED_BODY()
+
+	FCommandAddTag() = default;
+	FCommandAddTag(FLWEntity InEntity, UScriptStruct* InStruct)
+		: FCommandBufferEntryBase(InEntity)
+		, StructParam(InStruct)
+	{}
+
+protected:
+	virtual void Execute(UEntitySubsystem& System) const override;
+
+	UScriptStruct* StructParam = nullptr;
+};
+
+/**
+ * Command dedicated to remove a tag from an existing entity
+ */
+USTRUCT()
+struct MASSENTITY_API FCommandRemoveTag : public FCommandBufferEntryBase
+{
+	GENERATED_BODY()
+
+	FCommandRemoveTag() = default;
+	FCommandRemoveTag(FLWEntity InEntity, UScriptStruct* InStruct)
+		: FCommandBufferEntryBase(InEntity)
+		, StructParam(InStruct)
+	{}
+
+protected:
+	virtual void Execute(UEntitySubsystem& System) const override;
+
+	UScriptStruct* StructParam = nullptr;
+};
+
+/**
+ * Command for swapping given tags for a given entity. Note that the entity doesn't need to own the "old tag", the 
+ * "new tag" will be added regardless.
+ */
+USTRUCT()
+struct MASSENTITY_API FCommandSwapTags : public FCommandBufferEntryBase
+{
+	GENERATED_BODY()
+
+	FCommandSwapTags() = default;
+
+	FCommandSwapTags(const FLWEntity InEntity, const UScriptStruct* InOldTagType, const UScriptStruct* InNewTagType)
+		: FCommandBufferEntryBase(InEntity)
+		, OldTagType(InOldTagType)
+		, NewTagType(InNewTagType)
+	{
+		checkf((InOldTagType == nullptr) || InOldTagType->IsChildOf(FComponentTag::StaticStruct()), TEXT("FCommandSwapTags works only with tags while '%s' is not one."), *GetPathNameSafe(InOldTagType));
+		checkf((InNewTagType == nullptr) || InNewTagType->IsChildOf(FComponentTag::StaticStruct()), TEXT("FCommandSwapTags works only with tags while '%s' is not one."), *GetPathNameSafe(InNewTagType));
+	}
+
+protected:
+	virtual void Execute(UEntitySubsystem& System) const override;
+
+	const UScriptStruct* OldTagType = nullptr;
+	const UScriptStruct* NewTagType = nullptr;
+};
+
+/**
+ * Command performing a removal of a collection of components and tags
+ */
+USTRUCT()
+struct MASSENTITY_API FCommandRemoveComposition : public FCommandBufferEntryBase
+{
+	GENERATED_BODY()
+
+	FCommandRemoveComposition() = default;
+	FCommandRemoveComposition(FLWEntity InEntity, const FLWCompositionDescriptor& InDescriptor)
+		: FCommandBufferEntryBase(InEntity)
+		, Descriptor(InDescriptor)
+	{}
+
+protected:
+	virtual void Execute(UEntitySubsystem& System) const override;
+
+	FLWCompositionDescriptor Descriptor;
+};
+
+struct MASSENTITY_API FLWCCommandBuffer
+{
+public:
+
+	/** Push any command, requires to derive from FCommandBufferEntryBase */
+	template< typename T, typename = typename TEnableIf<TIsDerivedFrom<typename TRemoveReference<T>::Type, FCommandBufferEntryBase>::IsDerived, void>::Type >
+	void PushCommand(const T& Command)
+	{
+		EmplaceCommand<T>(Command);
+	}
+
+	/** Emplace any command, requires to derive from FCommandBufferEntryBase */
+	template< typename T, typename = typename TEnableIf<TIsDerivedFrom<typename TRemoveReference<T>::Type, FCommandBufferEntryBase>::IsDerived, void>::Type, typename... TArgs >
+	void EmplaceCommand(TArgs&&... InArgs)
+	{
+		UE_MT_SCOPED_WRITE_ACCESS(PendingCommandsDetector);
+		return PendingCommands.Emplace<T>(Forward<TArgs>(InArgs)...);
+	}
+
+	/** Emplace any command and return its ref, requires to derive from FCommandBufferEntryBase */
+	template< typename T, typename = typename TEnableIf<TIsDerivedFrom<typename TRemoveReference<T>::Type, FCommandBufferEntryBase>::IsDerived, void>::Type, typename... TArgs >
+	T& EmplaceCommand_GetRef(TArgs&&... InArgs)
+	{
+		UE_MT_SCOPED_WRITE_ACCESS(PendingCommandsDetector);
+		return PendingCommands.Emplace_GetRef<T>(Forward<TArgs>(InArgs)...);
+	}
+
+	template<typename T>
+	void AddComponent(FLWEntity Entity)
+	{
+		static_assert(TIsDerivedFrom<T, FLWComponentData>::IsDerived, "Given struct type is not a valid component type.");
+		EmplaceCommand<FCommandAddComponent>(Entity, T::StaticStruct());
+	}
+
+	template<typename T>
+	void RemoveComponent(FLWEntity Entity)
+	{
+		static_assert(TIsDerivedFrom<T, FLWComponentData>::IsDerived, "Given struct type is not a valid component type.");
+		EmplaceCommand<FCommandRemoveComponent>(Entity, T::StaticStruct());
+	}
+
+	template<typename T>
+	void AddTag(FLWEntity Entity)
+	{
+		static_assert(TIsDerivedFrom<T, FComponentTag>::IsDerived, "Given struct type is not a valid tag type.");
+		EmplaceCommand<FCommandAddTag>(Entity, T::StaticStruct());
+	}
+
+	template<typename T>
+	void RemoveTag(FLWEntity Entity)
+	{
+		static_assert(TIsDerivedFrom<T, FComponentTag>::IsDerived, "Given struct type is not a valid tag type.");
+		EmplaceCommand<FCommandRemoveTag>(Entity, T::StaticStruct());
+	}
+
+	void DestroyEntity(FLWEntity Entity)
+	{
+		EntitiesToDestroy.Add(Entity);
+	}
+
+	void BatchDestroyEntities(const TArray<FLWEntity>& InEntitiesToDestroy)
+	{
+		EntitiesToDestroy.Append(InEntitiesToDestroy);
+	}
+
+	void ReplayBufferAgainstSystem(UEntitySubsystem* System);
+
+	SIZE_T GetAllocatedSize() const { return PendingCommands.GetAllocatedSize(); }
+
+	/** 
+	 * Appends the commands from the passed buffer into this one
+	 * @param InOutOther the source buffer to copy the commands from. Note that after the call the InOutOther will be 
+	 *	emptied due to the function using Move semantics
+	 */
+	void MoveAppend(FLWCCommandBuffer& InOutOther);
+
+	int32 GetPendingCommandsCount() const { return PendingCommands.Num(); }
+
+private:
+	FInstancedStructStream PendingCommands;
+	UE_MT_DECLARE_RW_ACCESS_DETECTOR(PendingCommandsDetector);
+	FCriticalSection AppendingCommandsCS;
+
+	TArray<FLWEntity> EntitiesToDestroy;
+};
diff --git a/Engine/Plugins/Runtime/MassEntity/Source/MassEntity/Public/LWComponentTypes.h b/Engine/Plugins/Runtime/MassEntity/Source/MassEntity/Public/LWComponentTypes.h
new file mode 100644
--- /dev/null
+++ b/Engine/Plugins/Runtime/MassEntity/Source/MassEntity/Public/LWComponentTypes.h
@@ -0,0 +1,148 @@
+// Copyright Epic Games, Inc. All Rights Reserved.
+
+#pragma once
+
+#include "ScriptStructTypeBitSet.h"
+#include "LWComponentTypes.generated.h"
+
+// the name is temporary. To be changed when we arrive at the final name for the plugin
+#define WITH_AGGREGATETICKING_DEBUG (!(UE_BUILD_SHIPPING || UE_BUILD_SHIPPING_WITH_EDITOR || UE_BUILD_TEST) && WITH_STRUCTUTILS_DEBUG && 1)
+
+// This is the base class for all lightweight components
+USTRUCT()
+struct FLWComponentData
+{
+	GENERATED_BODY()
+
+	FLWComponentData() {}
+};
+
+// This is the base class for types that will only be tested for presence/absence, i.e. Tags.
+// Subclasses should never contain any member properties.
+USTRUCT()
+struct FComponentTag
+{
+	GENERATED_BODY()
+
+	FComponentTag() {}
+};
+
+USTRUCT()
+struct FLWChunkComponent
+{
+	GENERATED_BODY()
+
+	FLWChunkComponent() {}
+};
+
+// A handle to a lightweight entity.  An entity is used in conjunction with the UEntitySubsystem
+// for the current world and can contain lightweight components.
+USTRUCT()
+struct FLWEntity
+{
+	GENERATED_BODY()
+
+	FLWEntity() {}
+	
+	UPROPERTY(VisibleAnywhere, Category = "AggregateTicking|Debug", Transient)
+	int32 Index = 0;
+	
+	UPROPERTY(VisibleAnywhere, Category = "AggregateTicking|Debug", Transient)
+	int32 SerialNumber = 0;
+
+	bool operator==(const FLWEntity& Other) const
+	{
+		return Index == Other.Index && SerialNumber == Other.SerialNumber;
+	}
+
+	bool operator!=(const FLWEntity& Other) const
+	{
+		return !operator==(Other);
+	}
+
+	/** Note that this function is merely checking if Index and SerialNumber are set. There's no way to validate if 
+	 *  these indicate a valid entity in an EntitySubsystem without asking the system. */
+	bool IsSet() const
+	{
+		return Index != 0 && SerialNumber != 0;
+	}
+
+	void Reset()
+	{
+		Index = SerialNumber = 0;
+	}
+
+	friend uint32 GetTypeHash(const FLWEntity& Entity)
+	{
+		return HashCombine(Entity.Index, Entity.SerialNumber);
+	}
+
+	FString DebugGetDescription() const
+	{
+		return FString::Printf(TEXT("i: %d sn: %d"), Index, SerialNumber);
+	}
+};
+
+template struct MASSENTITY_API TScriptStructTypeBitSet<FLWComponentData>;
+using FLWComponentBitSet = TScriptStructTypeBitSet<FLWComponentData>;
+template struct MASSENTITY_API TScriptStructTypeBitSet<FComponentTag>;
+using FLWTagBitSet = TScriptStructTypeBitSet<FComponentTag>;
+template struct MASSENTITY_API TScriptStructTypeBitSet<FLWChunkComponent>;
+using FLWChunkComponentBitSet = TScriptStructTypeBitSet<FLWChunkComponent>;
+
+/** The type summarily describing a composition of an entity or an archetype. It contains information on both the
+ *  components as well as tags */
+struct FLWCompositionDescriptor
+{
+	FLWCompositionDescriptor() = default;
+	FLWCompositionDescriptor(const FLWComponentBitSet& InComponents, const FLWTagBitSet& InTags, const FLWChunkComponentBitSet& InChunkComponents)
+		: Components(InComponents), Tags(InTags), ChunkComponents(InChunkComponents)
+	{}
+
+	FLWCompositionDescriptor(TConstArrayView<const UScriptStruct*> InComponents, const FLWTagBitSet& InTags, const FLWChunkComponentBitSet& InChunkComponents)
+		: FLWCompositionDescriptor(FLWComponentBitSet(InComponents), InTags, InChunkComponents)
+	{}
+
+	FLWCompositionDescriptor(TConstArrayView<FInstancedStruct> InComponentInstances, const FLWTagBitSet& InTags, const FLWChunkComponentBitSet& InChunkComponents)
+		: FLWCompositionDescriptor(FLWComponentBitSet(InComponentInstances), InTags, InChunkComponents)
+	{}
+
+	FLWCompositionDescriptor(FLWComponentBitSet&& InComponents, FLWTagBitSet&& InTags, FLWChunkComponentBitSet&& InChunkComponents)
+		: Components(MoveTemp(InComponents)), Tags(MoveTemp(InTags)), ChunkComponents(MoveTemp(InChunkComponents))
+	{}
+
+	void Reset()
+	{
+		Components.Reset();
+		Tags.Reset();
+	}
+
+	bool IsEquivalent(const FLWComponentBitSet& InComponentBitSet, const FLWTagBitSet& InTagBitSet, const FLWChunkComponentBitSet& InChunkComponentsBitSet) const
+	{
+		return Components.IsEquivalent(InComponentBitSet) && Tags.IsEquivalent(InTagBitSet) && ChunkComponents.IsEquivalent(InChunkComponentsBitSet);
+	}
+
+	bool IsEquivalent(const FLWCompositionDescriptor& OtherDescriptor) const
+	{
+		return IsEquivalent(OtherDescriptor.Components, OtherDescriptor.Tags, OtherDescriptor.ChunkComponents);
+	}
+
+	bool IsEmpty() const { return Components.IsEmpty() && Tags.IsEmpty() && ChunkComponents.IsEmpty(); }
+
+	static uint32 CalculateHash(const FLWComponentBitSet& InComponents, const FLWTagBitSet& InTags, const FLWChunkComponentBitSet& InChunkComponents)
+	{
+		const uint32 ComponentsHash = GetTypeHash(InComponents);
+		const uint32 TagsHash = GetTypeHash(InTags);
+		const uint32 ChunkComponentsHash = GetTypeHash(InChunkComponents);
+		return HashCombine(HashCombine(ComponentsHash, TagsHash), ChunkComponentsHash);
+	}	
+
+	uint32 CalculateHash() const 
+	{
+		return CalculateHash(Components, Tags, ChunkComponents);
+	}
+
+	FLWComponentBitSet Components;
+	FLWTagBitSet Tags;
+	FLWChunkComponentBitSet ChunkComponents;
+};
diff --git a/Engine/Plugins/Runtime/MassEntity/Source/MassEntity/Public/MassEntityDebug.h b/Engine/Plugins/Runtime/MassEntity/Source/MassEntity/Public/MassEntityDebug.h
new file mode 100644
--- /dev/null
+++ b/Engine/Plugins/Runtime/MassEntity/Source/MassEntity/Public/MassEntityDebug.h
@@ -0,0 +1,17 @@
+// Copyright Epic Games, Inc. All Rights Reserved.
+
+#pragma once
+#include "VisualLogger/VisualLoggerTypes.h"
+#include "MassEntityTypes.h"
+
+class FOutputDevice;
+class UPipeProcessor;
+
+MASSENTITY_API DECLARE_ENUM_TO_STRING(EPipeProcessingPhase);
+
+namespace UE { namespace Pipe { namespace Debug {
+
+MASSENTITY_API extern void DebugOutputDescription(TConstArrayView<UPipeProcessor*> Processors, FOutputDevice& Ar);
+
+}}} // namespace UE::Pipe::Debug
+
diff --git a/Engine/Plugins/Runtime/MassEntity/Source/MassEntity/Public/MassEntitySettings.h b/Engine/Plugins/Runtime/MassEntity/Source/MassEntity/Public/MassEntitySettings.h
new file mode 100644
--- /dev/null
+++ b/Engine/Plugins/Runtime/MassEntity/Source/MassEntity/Public/MassEntitySettings.h
@@ -0,0 +1,101 @@
+// Copyright Epic Games, Inc. All Rights Reserved.
+
+#pragma once
+
+#include "Engine/DeveloperSettings.h"
+#include "MassProcessingPhase.h"
+#include "MassProcessor.h"
+#include "InstancedStruct.h"
+#include "MassEntitySettings.generated.h"
+
+#define GET_PIPE_CONFIG_VALUE(a) (GetMutableDefault<UPipeSettings>()->a)
+
+class UPipeProcessingPhaseManager;
+struct FPipeProcessingPhaseConfig;
+struct FPropertyChangedEvent;
+
+
+USTRUCT()
+struct FPipeProcessingPhaseConfig
+{
+	GENERATED_BODY()
+
+	UPROPERTY(EditAnywhere, Category = Pipe, config)
+	FName PhaseName;
+
+	UPROPERTY(EditAnywhere, Category = Pipe, config, NoClear)
+	TSubclassOf<UPipeCompositeProcessor> PhaseGroupClass = UPipeCompositeProcessor::StaticClass();
+
+	UPROPERTY(EditAnywhere, Category = Pipe, config, NoClear)
+	TArray<FName> OffGameThreadGroupNames;
+
+	UPROPERTY(Transient)
+	TArray<UPipeProcessor*> ProcessorCDOs;
+
+#if WITH_EDITORONLY_DATA
+	// this processor is available only in editor since it's used to present the user the order in which processors
+	// will be executed when given processing phase gets triggered
+	UPROPERTY(Transient)
+	UPipeCompositeProcessor* PhaseProcessor = nullptr;
+
+	UPROPERTY(VisibleAnywhere, Category = Pipe, Transient)
+	FString Description;
+#endif //WITH_EDITORONLY_DATA
+};
+
+/**
+ * Implements the settings for MassSimulation
+ */
+UCLASS(config = Game, defaultconfig, DisplayName = "Pipe")
+class MASSENTITY_API UPipeSettings : public UDeveloperSettings
+{
+	GENERATED_BODY()
+public:
+#if WITH_EDITORONLY_DATA
+	DECLARE_MULTICAST_DELEGATE_OneParam(FOnSettingsChange, const FPropertyChangedEvent& /*PropertyChangedEvent*/);
+#endif // WITH_EDITORONLY_DATA
+
+	UPipeSettings(const FObjectInitializer& ObjectInitializer = FObjectInitializer::Get());
+
+	void BuildProcessorListAndPhases();
+	void AddToActiveProcessorsList(TSubclassOf<UPipeProcessor> ProcessorClass);
+
+	const FPipeProcessingPhaseConfig* GetProcessingPhasesConfig();
+	const FPipeProcessingPhaseConfig& GetProcessingPhaseConfig(const EPipeProcessingPhase ProcessingPhase) const { check(ProcessingPhase != EPipeProcessingPhase::MAX); return ProcessingPhasesConfig[int(ProcessingPhase)]; }
+
+#if WITH_EDITOR
+	FOnSettingsChange& GetOnSettingsChange() { return OnSettingsChange; }	
+
+protected:
+	virtual void PostEditChangeProperty(struct FPropertyChangedEvent& PropertyChangedEvent) override;
+	virtual void PostEditChangeChainProperty(struct FPropertyChangedChainEvent& PropertyChangedEvent) override;
+#endif // WITH_EDITOR
+
+	virtual void BeginDestroy() override;
+
+	void BuildPhases();
+	void BuildProcessorList();
+
+public:
+	/** 
+	 * The name of the file to dump the processor dependency graph. T
+	 * The dot file will be put in the project log folder.
+	 * To generate a svg out of that file, simply run dot executable with following parameters: -Tsvg -O filename.dot 
+	 */
+	UPROPERTY(EditDefaultsOnly, Category = Pipe, Transient)
+	FString DumpDependencyGraphFileName;
+
+	/** Lets users configure processing phases including the composite processor class to be used as a container for the phases' processors. */
+	UPROPERTY(EditDefaultsOnly, Category = Pipe, config)
+	FPipeProcessingPhaseConfig ProcessingPhasesConfig[(uint8)EPipeProcessingPhase::MAX];
+
+	/** This list contains all the processors available in the given binary (including plugins). The contents are sorted by display name.*/
+	UPROPERTY(VisibleAnywhere, Category = Pipe, Transient, Instanced)
+	TArray<UPipeProcessor*> ProcessorCDOs;
+
+#if WITH_EDITORONLY_DATA
+protected:
+	FOnSettingsChange OnSettingsChange;
+#endif // WITH_EDITORONLY_DATA
+	bool bInitialized = false;
+};
diff --git a/Engine/Plugins/Runtime/MassEntity/Source/MassEntity/Public/MassEntitySystem.h b/Engine/Plugins/Runtime/MassEntity/Source/MassEntity/Public/MassEntitySystem.h
new file mode 100644
--- /dev/null
+++ b/Engine/Plugins/Runtime/MassEntity/Source/MassEntity/Public/MassEntitySystem.h
@@ -0,0 +1,32 @@
+// Copyright Epic Games, Inc. All Rights Reserved.
+
+#pragma once
+
+#include "EntitySubsystem.h"
+#include "MassEntitySystem.generated.h"
+
+
+class UWorld;
+
+UCLASS()
+class MASSENTITY_API UPipeEntitySubsystem : public UEntitySubsystem
+{
+	GENERATED_BODY()
+public:
+	/** we never create PipeEntitySubsystem as a regular WorldSubsystem. Created on demand.
+	 *  This functionality might need to be moved down to Mass plugin.
+	 *  
+	 *  @todo this is a bit hacky since makes UPipeEntitySubsystem unobtainable via regular World::GetSubsystem which 
+	 *  might be confusing to users. Ideally we'd extract UEntitySubsystem's functionality into a separate class 
+	 *  and make both UEntitySubsystem and UPipeEntitySubsystem utilize it.
+	 */
+	virtual bool ShouldCreateSubsystem(UObject* Outer) const override { return false; }
+
+	static UPipeEntitySubsystem* GetCurrent(UWorld* World) { return InstanceGetter(World); }
+
+	/** This function is meant to be called once per project, make sure you know all code sites that call it to avoid surprises */
+	static void SetGetter_Internal(TFunctionRef<UPipeEntitySubsystem*(UWorld*)> InGetter) { new (&InstanceGetter)TFunctionRef<UPipeEntitySubsystem*(UWorld*)>(InGetter); }
+
+private:
+	static TFunctionRef<UPipeEntitySubsystem*(UWorld*)> InstanceGetter;
+};
\ No newline at end of file
diff --git a/Engine/Plugins/Runtime/MassEntity/Source/MassEntity/Public/MassEntityTypes.h b/Engine/Plugins/Runtime/MassEntity/Source/MassEntity/Public/MassEntityTypes.h
new file mode 100644
--- /dev/null
+++ b/Engine/Plugins/Runtime/MassEntity/Source/MassEntity/Public/MassEntityTypes.h
@@ -0,0 +1,117 @@
+// Copyright Epic Games, Inc. All Rights Reserved.
+
+#pragma once
+
+#include "InstancedStruct.h"
+#include "LWComponentTypes.h"
+#include "MassEntityTypes.generated.h"
+
+#define WITH_PIPE_DEBUG (!(UE_BUILD_SHIPPING || UE_BUILD_SHIPPING_WITH_EDITOR || UE_BUILD_TEST) && WITH_AGGREGATETICKING_DEBUG && 1)
+
+MASSENTITY_API DECLARE_LOG_CATEGORY_EXTERN(LogPipe, Warning, All);
+
+class UEntitySubsystem;
+class UPipeProcessor;
+class UPipeSchematic;
+class UPipeCompositeProcessor;
+
+UENUM(meta = (Bitflags, UseEnumValuesAsMaskValuesInEditor = "true"))
+enum class EProcessorExecutionFlags : uint8
+{
+	None = 0 UMETA(Hidden),
+	Standalone = 1 << 0,
+	Server = 1 << 1,
+	Client = 1 << 2,
+	All = Standalone | Server | Client UMETA(Hidden)
+};
+ENUM_CLASS_FLAGS(EProcessorExecutionFlags);
+
+USTRUCT()
+struct FProcessorAuxDataBase
+{
+	GENERATED_BODY()
+};
+
+USTRUCT(BlueprintType)
+struct MASSENTITY_API FPipeContext
+{
+	GENERATED_BODY()
+
+	UPROPERTY()
+	UEntitySubsystem* EntitySubsystem = nullptr;
+	
+	UPROPERTY()
+	float DeltaSeconds = 0.f;
+
+	UPROPERTY()
+	FInstancedStruct AuxData;
+
+	FPipeContext() = default;
+	FPipeContext(UEntitySubsystem& InEntities, const float InDeltaSeconds);
+};
+
+/** 
+ *  A helper type that converts a set of UPipeSchematics into a runtime-usable array of PipeProcessor copies
+ */
+USTRUCT()
+struct MASSENTITY_API FRuntimePipeline
+{
+	GENERATED_BODY()
+
+	UPROPERTY()
+	TArray<UPipeProcessor*> Processors;
+
+	void Reset();
+	void Initialize(UObject& Owner);
+	
+	/** Creates runtime copies of the given UPipeProcessors collection. */
+	void SetProcessors(TArray<UPipeProcessor*>&& InProcessors);
+
+	/** Creates runtime copies of UPipeProcessors declared in Schematics using InOwner as new UPipeProcessors' outer. */
+	void InitializeFromSchematics(TConstArrayView<TSoftObjectPtr<UPipeSchematic>> Schematics, UObject& InOwner);
+
+	/** Creates runtime copies of UPipeProcessors given in InProcessors input parameter, using InOwner as new UPipeProcessors' outer. */
+	void CreateFromArray(TConstArrayView<const UPipeProcessor*> InProcessors, UObject& InOwner);
+
+	/** Calls CreateFromArray and calls Initialize on all processors afterwards. */
+	void InitializeFromArray(TConstArrayView<const UPipeProcessor*> InProcessors, UObject& InOwner);
+	
+	/** Creates runtime instances of UPipeProcessors for each processor class given via InProcessorClasses. 
+	 *  The instances will be created with InOwner as outer. */
+	void InitializeFromClassArray(TConstArrayView<TSubclassOf<UPipeProcessor>> InProcessorClasses, UObject& InOwner);
+
+	/** Creates a runtime instance of every processors in the given InProcessors array. If a processor of that class
+	 *  already exists in Processors array it gets overridden. Otherwise it gets added to the end of the collection.*/
+	void AppendOrOverrideRuntimeProcessorCopies(TConstArrayView<const UPipeProcessor*> InProcessors, UObject& InOwner);
+
+	/** Creates a runtime instance of every processors in the given array if there's no processor of that class in Processors already.
+	 *  Call this function when adding processors to an already configured FRuntimePipeline instance. If you're creating 
+	 *  one from scratch calling any of the InitializeFrom* methods will be more efficient (and will produce same results)
+	 *  or call AppendOrOverrideRuntimeProcessorCopies.*/
+	void AppendUniqueRuntimeProcessorCopies(TConstArrayView<const UPipeProcessor*> InProcessors, UObject& InOwner);
+
+	/** Adds InProcessor to Processors without any additional checks */
+	void AppendProcessor(UPipeProcessor& Processor);
+
+	/** goes through Processor looking for a UPipeCompositeProcessor instance which GroupName matches the one given as the parameter */
+	UPipeCompositeProcessor* FindTopLevelGroupByName(const FName GroupName);
+
+	void DebugOutputDescription(FOutputDevice& Ar) const;
+
+	bool HasProcessorOfExactClass(TSubclassOf<UPipeProcessor> InClass) const;
+	bool IsEmpty() const { return Processors.IsEmpty();}
+
+	MASSENTITY_API friend uint32 GetTypeHash(const FRuntimePipeline& Instance);
+};
+
+UENUM()
+enum class EPipeProcessingPhase : uint8
+{
+	PrePhysics,
+	StartPhysics,
+	DuringPhysics,
+	EndPhysics,
+	PostPhysics,
+	FrameEnd,
+	MAX,
+};
diff --git a/Engine/Plugins/Runtime/MassEntity/Source/MassEntity/Public/MassEntityUtils.h b/Engine/Plugins/Runtime/MassEntity/Source/MassEntity/Public/MassEntityUtils.h
new file mode 100644
--- /dev/null
+++ b/Engine/Plugins/Runtime/MassEntity/Source/MassEntity/Public/MassEntityUtils.h
@@ -0,0 +1,13 @@
+// Copyright Epic Games, Inc. All Rights Reserved.
+
+#pragma once
+
+#include "MassEntityTypes.h"
+
+namespace UE { namespace Pipe { namespace Utils
+{
+
+/** returns the current execution mode for the processors calculated from the world network mode */
+MASSENTITY_API extern EProcessorExecutionFlags GetProcessorExecutionFlagsForWold(const UWorld& World);
+
+} } }// namespace UE::Pipe::Utils
\ No newline at end of file
diff --git a/Engine/Plugins/Runtime/MassEntity/Source/MassEntity/Public/MassExecutor.h b/Engine/Plugins/Runtime/MassEntity/Source/MassEntity/Public/MassExecutor.h
new file mode 100644
--- /dev/null
+++ b/Engine/Plugins/Runtime/MassEntity/Source/MassEntity/Public/MassExecutor.h
@@ -0,0 +1,48 @@
+// Copyright Epic Games, Inc. All Rights Reserved.
+
+#pragma once
+
+#include "Subsystems/WorldSubsystem.h"
+#include "EntitySubsystem.h"
+#include "Async/TaskGraphInterfaces.h"
+
+
+struct FRuntimePipeline;
+struct FPipeContext;
+struct FLWEntity;
+struct FArchetypeChunkCollection;
+class UPipeProcessor;
+
+
+namespace UE::Pipe::Executor
+{
+	/** Executes processors in a given RuntimePipeline */
+	MASSENTITY_API void Run(FRuntimePipeline& RuntimePipeline, FPipeContext& PipeContext);
+
+	/** Executes given Processor. Used mainly for triggering calculations via PipeCompositeProcessors, e.g processing phases */
+	MASSENTITY_API void Run(UPipeProcessor& Processor, FPipeContext& PipeContext);
+
+	/** Similar to the Run function, but instead of using all the entities hosted by PipeContext.EntitySubsystem 
+	 *  it is processing only the entities given by EntityID via the Entities input parameter. 
+	 *  Note that all the entities need to be of Archetype archetype. 
+	 *  Under the hood the function converts Archetype-Entities pair to FArchetypeChunkCollection and calls the other flavor of RunSparse
+	 */
+	MASSENTITY_API void RunSparse(FRuntimePipeline& RuntimePipeline, FPipeContext& PipeContext, FArchetypeHandle Archetype, TConstArrayView<FLWEntity> Entities);
+
+	/** Similar to the Run function, but instead of using all the entities hosted by PipeContext.EntitySubsystem 
+	 *  it is processing only the entities given by SparseEntities input parameter.
+	 *  @todo rename
+	 */
+	MASSENTITY_API void RunSparse(FRuntimePipeline& RuntimePipeline, FPipeContext& PipeContext, const FArchetypeChunkCollection& ChunkCollection);
+
+	/** Executes given Processors array view. This function gets called under the hood by the rest of Run* functions */
+	MASSENTITY_API void RunProcessorsView(TArrayView<UPipeProcessor*> Processors, FPipeContext& PipeContext, const FArchetypeChunkCollection* ChunkCollection = nullptr);
+
+	/** 
+	 *  Triggers tasks executing Processor (and potentially it's children) and returns the task graph event representing 
+	 *  the task (the event will be "completed" once all the processors finish running). 
+	 *  @param OnDoneNotification will be called after all the processors are done, just after flushing the command buffer.
+	 *    Note that OnDoneNotification will be executed on GameThread.
+	 */
+	MASSENTITY_API FGraphEventRef TriggerParallelTasks(UPipeProcessor& Processor, FPipeContext& PipeContext, TFunction<void()> OnDoneNotification);
+};
diff --git a/Engine/Plugins/Runtime/MassEntity/Source/MassEntity/Public/MassProcessingPhase.h b/Engine/Plugins/Runtime/MassEntity/Source/MassEntity/Public/MassProcessingPhase.h
new file mode 100644
--- /dev/null
+++ b/Engine/Plugins/Runtime/MassEntity/Source/MassEntity/Public/MassProcessingPhase.h
@@ -0,0 +1,151 @@
+// Copyright Epic Games, Inc. All Rights Reserved.
+
+#pragma once
+
+#include "CoreMinimal.h"
+#include "UObject/UObjectGlobals.h"
+#include "UObject/Object.h"
+#include "Engine/EngineBaseTypes.h"
+#include "MassEntityTypes.h"
+#include "MassProcessingPhase.generated.h"
+
+
+class UPipeProcessingPhaseManager;
+class UPipeProcessor;
+class UPipeCompositeProcessor;
+class UPipeEntitySubsystem;
+struct FLWCCommandBuffer;
+
+USTRUCT()
+struct FPipeProcessingPhase : public FTickFunction
+{
+	GENERATED_BODY()
+
+	DECLARE_MULTICAST_DELEGATE_OneParam(FOnPhaseEvent, const float /*DeltaSeconds*/);
+
+	FPipeProcessingPhase();
+
+protected:
+	// FTickFunction interface
+	virtual void ExecuteTick(float DeltaTime, ELevelTick TickType, ENamedThreads::Type CurrentThread, const FGraphEventRef& MyCompletionGraphEvent) override;
+	virtual FString DiagnosticMessage() override;
+	virtual FName DiagnosticContext(bool bDetailed) override;
+	// End of FTickFunction interface
+
+	void OnParallelExecutionDone(const float DeltaTime);
+
+	bool IsConfiguredForParallelMode() const { return bRunInParallelMode; }
+	void ConfigureForParallelMode() { bRunInParallelMode = true; }
+	void ConfigureForSingleThreadMode() { bRunInParallelMode = false; }
+
+public:
+	bool IsDuringPipeProcessing() const { return bIsDuringPipeProcessing; }
+
+protected:
+	friend UPipeProcessingPhaseManager;
+
+	UPROPERTY(EditAnywhere, Category=Pipe)
+	UPipeCompositeProcessor* PhaseProcessor = nullptr;
+
+	UPROPERTY()
+	UPipeProcessingPhaseManager* Manager = nullptr;
+
+	EPipeProcessingPhase Phase = EPipeProcessingPhase::MAX;
+	FOnPhaseEvent OnPhaseStart;
+	FOnPhaseEvent OnPhaseEnd;
+
+private:
+	bool bRunInParallelMode = false;
+	bool bIsDuringPipeProcessing = false;
+};
+
+// It is unsafe to copy FTickFunctions and any subclasses of FTickFunction should specify the type trait WithCopy = false
+template<>
+struct TStructOpsTypeTraits<FPipeProcessingPhase> : public TStructOpsTypeTraitsBase2<FPipeProcessingPhase>
+{
+	enum
+	{
+		WithCopy = false
+	};
+};
+
+/** PipeProcessingPhaseManager owns separate FPipeProcessingPhase instances for every ETickingGroup. When activated
+ *  via Start function it registers and enables the FPipeProcessingPhase instances which themselves are tick functions 
+ *  that host FRuntimePipeline which they trigger as part of their Tick function. 
+ *  PipeProcessingPhaseManager serves as an interface to said FPipeProcessingPhase instances and allows initialization
+ *  with PipeSchematics (via InitializePhases function) as well as registering arbitrary functions to be called 
+ *  when a particular phase starts of ends (via GetOnPhaseStart and GetOnPhaseEnd functions). */
+UCLASS(Transient, HideCategories = (Tick))
+class MASSENTITY_API UPipeProcessingPhaseManager : public UObject
+{
+	GENERATED_BODY()
+public:
+	UPipeEntitySubsystem& GetEntitySubsystemRef() { check(EntitySubsystem); return *EntitySubsystem; }
+
+	/** Retrieves OnPhaseStart multicast delegate's reference for a given Phase */
+	FPipeProcessingPhase::FOnPhaseEvent& GetOnPhaseStart(const EPipeProcessingPhase Phase) { return ProcessingPhases[uint8(Phase)].OnPhaseStart; }
+	/** Retrieves OnPhaseEnd multicast delegate's reference for a given Phase */
+	FPipeProcessingPhase::FOnPhaseEvent& GetOnPhaseEnd(const EPipeProcessingPhase Phase) { return ProcessingPhases[uint8(Phase)].OnPhaseEnd; }
+
+	/** 
+	 *  Populates hosted FPipeProcessingPhase instances with Processors read from PipeSettings configuration.
+	 *  Calling this function overrides previous configuration of Phases.
+	 */
+	void InitializePhases(UObject& InProcessorOwner);
+
+	/** 
+	 *  Both flavors of Start function boil down to setting EntitySubsystem and Executor. If the callee has these 
+	 *  at hand it's suggested to use that Start version, otherwise call the World-using one. 
+	 */
+	void Start(UWorld& World);
+	void Start(UPipeEntitySubsystem& InEntitySubsystem);
+	void Stop();
+	bool IsRunning() const { return EntitySubsystem != nullptr; }
+
+	/** 
+	 *  returns true when called while any of the ProcessingPhases is actively executing its processors. Used to 
+	 *  determine whether it's safe to do entity-related operations like adding fragments.
+	 *  Note that the function will return false while the OnPhaseStart or OnPhaseEnd are being broadcast,
+	 *  the value returned will be `true` only when the entity subsystem is actively engaged 
+	 */
+	bool IsDuringPipeProcessing() const { return CurrentPhase != EPipeProcessingPhase::MAX && ProcessingPhases[int(CurrentPhase)].IsDuringPipeProcessing(); }
+
+protected:
+	virtual void PostInitProperties() override;
+	virtual void BeginDestroy() override;
+	void EnableTickFunctions(const UWorld& World);
+
+	/** Creates phase processors instances for each declared phase name, based on PipeSettings */
+	virtual void CreatePhases();
+
+#if WITH_EDITOR
+	virtual void OnPipeSettingsChange(const FPropertyChangedEvent& PropertyChangedEvent);
+#endif // WITH_EDITOR
+
+	friend FPipeProcessingPhase;
+
+	/** 
+	 *  Called by the given Phase at the very start of its execution function (the FPipeProcessingPhase::ExecuteTick),
+	 *  even before the FPipeProcessingPhase.OnPhaseStart broadcast delegate
+	 */
+	void OnPhaseStart(const FPipeProcessingPhase& Phase);
+
+	/**
+	 *  Called by the given Phase at the very end of its execution function (the FPipeProcessingPhase::ExecuteTick),
+	 *  after the FPipeProcessingPhase.OnPhaseEnd broadcast delegate
+	 */
+	void OnPhaseEnd(FPipeProcessingPhase& Phase);
+
+protected:	
+	UPROPERTY(VisibleAnywhere, Category=Pipe)
+	FPipeProcessingPhase ProcessingPhases[(uint8)EPipeProcessingPhase::MAX];
+
+	UPROPERTY()
+	UPipeEntitySubsystem* EntitySubsystem;
+
+	EPipeProcessingPhase CurrentPhase = EPipeProcessingPhase::MAX;
+
+#if WITH_EDITOR
+	FDelegateHandle PipeSettingsChangeHandle;
+#endif // WITH_EDITOR
+};
diff --git a/Engine/Plugins/Runtime/MassEntity/Source/MassEntity/Public/MassProcessor.h b/Engine/Plugins/Runtime/MassEntity/Source/MassEntity/Public/MassProcessor.h
new file mode 100644
--- /dev/null
+++ b/Engine/Plugins/Runtime/MassEntity/Source/MassEntity/Public/MassProcessor.h
@@ -0,0 +1,219 @@
+// Copyright Epic Games, Inc. All Rights Reserved.
+
+#pragma once
+
+#include "CoreMinimal.h"
+#include "EntitySubsystem.h"
+#include "MassEntityTypes.h"
+#include "Async/TaskGraphInterfaces.h"
+#include "LWCCommandBuffer.h"
+#include "MassProcessorDependencySolver.h"
+#include "MassProcessor.generated.h"
+
+
+struct FPipeProcessingPhaseConfig;
+
+enum class EProcessorCompletionStatus : uint8
+{
+	Invalid,
+	Threaded,
+	Postponed,
+	Done
+};
+
+USTRUCT()
+struct FPipeProcessorExecutionOrder
+{
+	GENERATED_BODY()
+
+	/** Determines which processing group this processor will be placed in. Leaving it empty ("None") means "top-most group for my ProcessingPhase" */
+	UPROPERTY(EditAnywhere, Category = Processor, config)
+	FName ExecuteInGroup = FName();
+
+	UPROPERTY(EditAnywhere, Category = Processor, config)
+	TArray<FName> ExecuteBefore;
+
+	UPROPERTY(EditAnywhere, Category = Processor, config)
+	TArray<FName> ExecuteAfter;
+};
+
+
+UCLASS(abstract, EditInlineNew, CollapseCategories, config = Game, defaultconfig)
+class MASSENTITY_API UPipeProcessor : public UObject
+{
+	GENERATED_BODY()
+public:
+	UPipeProcessor(const FObjectInitializer& ObjectInitializer = FObjectInitializer::Get());
+
+	virtual void Initialize(UObject& Owner) {}
+	virtual FGraphEventRef DispatchProcessorTasks(UEntitySubsystem& EntitySubsystem, FLWComponentSystemExecutionContext& ExecutionContext, const FGraphEventArray& Prerequisites = FGraphEventArray());
+
+	EProcessorExecutionFlags GetExecutionFlags() const { return (EProcessorExecutionFlags)ExecutionFlags; }
+
+	/** Whether this processor should execute according the CurrentExecutionFlags parameters */
+	bool ShouldExecute(const EProcessorExecutionFlags CurrentExecutionFlags) const { return (GetExecutionFlags() & CurrentExecutionFlags) != EProcessorExecutionFlags::None; }
+	void CallExecute(UEntitySubsystem& EntitySubsystem, FLWComponentSystemExecutionContext& Context);
+	
+	bool AllowDuplicates() const { return bAllowDuplicates; }
+
+	virtual void DebugOutputDescription(FOutputDevice& Ar, int32 Indent = 0) const;
+	virtual FString GetProcessorName() const { return GetName(); }
+
+	//----------------------------------------------------------------------//
+	// Ordering functions 
+	//----------------------------------------------------------------------//
+	virtual EPipeProcessingPhase GetProcessingPhase() const { return ProcessingPhase; }
+	virtual void SetProcessingPhase(EPipeProcessingPhase Phase) { ProcessingPhase = Phase; }
+	bool DoesRequireGameThreadExecution() const { return bRequiresGameThreadExecution; }
+	
+	const FPipeProcessorExecutionOrder& GetExecutionOrder() const { return ExecutionOrder; }
+
+	TConstArrayView<const int32> GetPrerequisiteIndices() const { return DependencyIndices; }
+
+	bool ShouldAutoAddToGlobalList() const { return bAutoRegisterWithProcessingPhases; }
+#if WITH_EDITOR
+	bool ShouldShowUpInSettings() const { return ShouldAutoAddToGlobalList() || bCanShowUpInSettings; }
+#endif // WITH_EDITORONLY_DATA
+
+	/** Sets bAutoRegisterWithProcessingPhases. Setting it to true will result in this processor class being always 
+	 * instantiated to be automatically evaluated every frame. @see UPipeProcessingPhaseManager
+	 * Note that calling this function is only valid on CDOs. Calling it on a regular instance will fail an ensure and 
+	 * have no other effect, i.e. CDO's value won't change */
+	void SetShouldAutoRegisterWithGlobalList(const bool bAutoRegister);
+
+#if CPUPROFILERTRACE_ENABLED
+	FString StatId;
+#endif
+	
+protected:
+	virtual void ConfigureQueries() PURE_VIRTUAL(UPipeProcessor::ConfigureQueries);
+	virtual void PostInitProperties() override;
+	virtual void Execute(UEntitySubsystem& EntitySubsystem, FLWComponentSystemExecutionContext& Context) PURE_VIRTUAL(UPipeProcessor::Execute);
+#if WITH_EDITOR
+	virtual void PostEditChangeProperty(FPropertyChangedEvent& PropertyChangedEvent) override;
+#endif // WITH_EDITOR
+
+protected:
+	/** Whether this processor should be executed on StandAlone or Server or Client */
+	UPROPERTY(EditAnywhere, Category = "Pipeline", meta = (Bitmask, BitmaskEnum = EProcessorExecutionFlags), config)
+	int32 ExecutionFlags;
+
+	/** Processing phase this processor will be automatically run as part of. */
+	UPROPERTY(EditDefaultsOnly, Category = Processor, config)
+	EPipeProcessingPhase ProcessingPhase = EPipeProcessingPhase::PrePhysics;
+
+	/** Configures when this given processor can be executed in relation to other processors and processing groups, within its processing phase. */
+	UPROPERTY(EditDefaultsOnly, Category = Processor, config)
+	FPipeProcessorExecutionOrder ExecutionOrder;
+
+	/** Configures whether this processor should be automatically included in the global list of processors executed every tick (see ProcessingPhase and ExecutionOrder). */
+	UPROPERTY(EditDefaultsOnly, Category = Processor, config)
+	bool bAutoRegisterWithProcessingPhases = true;
+
+	/** meant as a class property, make sure to set it in subclass' constructor. Controls whether there can be multiple 
+	 *  instances of a given class in a single FRuntimePipeline */
+	bool bAllowDuplicates = false;
+
+	UPROPERTY(EditDefaultsOnly, Category = Processor, config)
+	bool bRequiresGameThreadExecution = false;
+
+#if WITH_EDITORONLY_DATA
+	/** Used to permanently remove a given processor class from PipeSetting's listing. Used primarily for test-time 
+	 *  processor classes, but can also be used by project-specific code to prune the processor list. */
+	UPROPERTY(config)
+	bool bCanShowUpInSettings = true;
+#endif // WITH_EDITORONLY_DATA
+
+	friend class UPipeCompositeProcessor;
+	TArray<int32> DependencyIndices;
+	TArray<int32> TransientDependencyIndices;
+};
+
+
+UCLASS()
+class MASSENTITY_API UPipeCompositeProcessor : public UPipeProcessor
+{
+	GENERATED_BODY()
+
+	struct FDependencyNode
+	{
+		FName Name;
+		UPipeProcessor* Processor = nullptr;
+		TArray<int32> Dependencies;
+	};
+
+public:
+	UPipeCompositeProcessor();
+
+	void SetChildProcessors(TArray<UPipeProcessor*>&& InProcessors);
+
+	virtual void Initialize(UObject& Owner) override;
+	virtual void DebugOutputDescription(FOutputDevice& Ar, int32 Indent = 0) const override;
+	virtual void SetProcessingPhase(EPipeProcessingPhase Phase) override;
+
+	void SetGroupName(FName NewName);
+	FName GetGroupName() const { return GroupName; }
+
+	virtual void CopyAndSort(const FPipeProcessingPhaseConfig& PhaseConfig, const FString& DependencyGraphFileName = FString());
+
+	/** adds SubProcessor to an appropriately named group. If RequestedGroupName == None then SubProcessor
+	 *  will be added directly to ChildPipeline. If not then the indicated group will be searched for in ChildPipeline 
+	 *  and if it's missing it will be created and AddGroupedProcessor will be called recursively */
+	void AddGroupedProcessor(FName RequestedGroupName, UPipeProcessor& SubProcessor);
+
+	virtual FGraphEventRef DispatchProcessorTasks(UEntitySubsystem& EntitySubsystem, FLWComponentSystemExecutionContext& ExecutionContext, const FGraphEventArray& Prerequisites = FGraphEventArray()) override;
+
+	bool IsEmpty() const { return ChildPipeline.IsEmpty(); }
+
+	virtual FString GetProcessorName() const override { return GroupName.ToString(); }
+
+protected:
+	virtual void ConfigureQueries() override;
+	virtual void Execute(UEntitySubsystem& EntitySubsystem, FLWComponentSystemExecutionContext& Context) override;
+
+	/**
+	 *  Called recursively to add processors and composite processors to ChildPipeline based on ProcessorsAndGroups
+	 */
+	int32 Populate(TArray<FProcessorDependencySolver::FOrderInfo>& ProcessorsAndGroups, const int32 StartIndex = 0);
+
+	/** RequestedGroupName can indicate a multi-level group name, like so: A.B.C
+	 *  We need to extract the highest-level group name ('A' in the example), and see if it already exists. 
+	 *  If not, create it. 
+	 *  @param RequestedGroupName name of the group for which we want to find or create the processor.
+	 *  @param OutRemainingGroupName contains the group name after cutting the high-level group. In the used example it
+	 *    will contain "B.C". This value is then used to recursively create subgroups */
+	UPipeCompositeProcessor* FindOrAddGroupProcessor(FName RequestedGroupName, FString* OutRemainingGroupName = nullptr);
+
+protected:
+	UPROPERTY(VisibleAnywhere, Category=Pipe)
+	FRuntimePipeline ChildPipeline;
+
+	/** Group name that will be used when resolving processor dependencies and grouping */
+	UPROPERTY()
+	FName GroupName;
+
+	TArray<FDependencyNode> ProcessingFlatGraph;
+
+	struct FProcessorCompletion
+	{
+		FGraphEventRef CompletionEvent;
+		EProcessorCompletionStatus Status = EProcessorCompletionStatus::Invalid;
+
+		bool IsDone() const 
+		{
+			return Status == EProcessorCompletionStatus::Done || (CompletionEvent.IsValid() && CompletionEvent->IsComplete());
+		}
+
+		void Wait()
+		{
+			if (CompletionEvent.IsValid())
+			{
+				CompletionEvent->Wait();
+			}
+		}
+	};
+	TArray<FProcessorCompletion> CompletionStatus;
+
+	bool bRunInSeparateThread;
+	bool bHasOffThreadSubGroups;
+};
diff --git a/Engine/Plugins/Runtime/MassEntity/Source/MassEntity/Public/MassProcessorDependencySolver.h b/Engine/Plugins/Runtime/MassEntity/Source/MassEntity/Public/MassProcessorDependencySolver.h
new file mode 100644
--- /dev/null
+++ b/Engine/Plugins/Runtime/MassEntity/Source/MassEntity/Public/MassProcessorDependencySolver.h
@@ -0,0 +1,89 @@
+// Copyright Epic Games, Inc. All Rights Reserved.
+
+#pragma once
+
+#include "MassEntityTypes.h"
+
+
+class UPipeProcessor;
+
+enum class EDependencyNodeType : uint8
+{
+	Invalid,
+	Processor,
+	GroupStart,
+	GroupEnd
+};
+
+struct MASSENTITY_API FProcessorDependencySolver
+{
+private:
+	struct FNode
+	{
+		FNode(const FName InName, UPipeProcessor* InProcessor) 
+			: Name(InName), Processor(InProcessor)
+		{}
+
+		FName Name = TEXT("");
+		TArray<FNode> SubNodes;
+		TMap<FName, int32> Indices;
+		UPipeProcessor* Processor = nullptr;
+		TArray<int32> OriginalDependencies;
+		TArray<int32> TransientDependencies;
+		TArray<FName> ExecuteBefore;
+		TArray<FName> ExecuteAfter;
+
+		int32 FindOrAddGroupNodeIndex(const FString& GroupName);
+		int32 FindNodeIndex(FName InNodeName) const;
+		bool HasDependencies() const;
+	};
+
+public:
+	struct FOrderInfo
+	{
+		FName Name = TEXT("");
+		UPipeProcessor* Processor = nullptr;
+		EDependencyNodeType NodeType = EDependencyNodeType::Invalid;
+		TArray<FName> Dependencies;
+	};
+
+	FProcessorDependencySolver(TArrayView<UPipeProcessor*> InProcessors, const FName Name, const FString& InDependencyGraphFileName = FString());
+	void ResolveDependencies(TArray<FOrderInfo>& OutResult, TConstArrayView<const FName> PriorityNodes = TConstArrayView<const FName>());
+
+	static void CreateSubGroupNames(FName InGroupName, TArray<FString>& SubGroupNames);
+
+protected:
+	// note that internals are protected rather than private to support unit testing
+
+	/**
+	 * Traverses RootNode's child nodes indicated by InOutIndicesRemaining and appends to OutNodeIndices the ones that 
+	 * have no dependencies. The indices added to OutNodeIndices also get removed from remaining nodes' outstanding 
+	 * dependencies.
+	 * Note that the whole InOutIndicesRemaining gets tested in sequence, which means nodes can get their dependencies 
+	 * emptied and added to OutNodeIndices within one call (as opposed to PerformPrioritySolverStep).
+	 */
+	static int32 PerformSolverStep(FNode& RootNode, TArray<int32>& InOutIndicesRemaining, TArray<int32>& OutNodeIndices);
+	
+	/**
+	 * Traverses InOutIndicesRemaining in search of the first RootNode's node that has no dependencies left. Once found 
+	 * the node's index gets added to OutNodeIndices, removed from dependency lists from all other nodes and the function 
+	 * quits.
+	 * @return 'true' if a dependency-less node has been found and added to OutNodeIndices; 'false' otherwise.
+	 */
+	static bool PerformPrioritySolverStep(FNode& RootNode, TArray<int32>& InOutIndicesRemaining, TArray<int32>& OutNodeIndices);
+	
+	static FString NameViewToString(TConstArrayView<FName> View);
+
+	void AddNode(FName GroupName, UPipeProcessor& Processor);
+	void BuildDependencies(FNode& RootNode);
+	void Solve(FNode& RootNode, TConstArrayView<const FName> PriorityNodes, TArray<FProcessorDependencySolver::FOrderInfo>& OutResult, int LoggingIndent = 0);
+	void LogNode(const FNode& RootNode, const FNode* ParentNode = nullptr, int Indent = 0);
+	void DumpGraph(FArchive& LogFile) const;
+	
+	TArrayView<UPipeProcessor*> Processors;
+	FNode GroupRootNode;
+	bool bAnyCyclesDetected = false;
+	FString DependencyGraphFileName;
+
+	friend struct FDumpGraphDependencyUtils;
+};
\ No newline at end of file
diff --git a/Engine/Plugins/Runtime/MassEntity/Source/MassEntity/Public/MassSchematic.h b/Engine/Plugins/Runtime/MassEntity/Source/MassEntity/Public/MassSchematic.h
new file mode 100644
--- /dev/null
+++ b/Engine/Plugins/Runtime/MassEntity/Source/MassEntity/Public/MassSchematic.h
@@ -0,0 +1,27 @@
+// Copyright Epic Games, Inc. All Rights Reserved.
+
+#pragma once
+
+#include "CoreMinimal.h"
+#include "Engine/DataAsset.h"
+#include "MassEntityTypes.h"
+#include "MassSchematic.generated.h"
+
+
+class UPipeProcessor;
+
+/**
+ * Pipe Schematic asset.
+ */
+UCLASS(BlueprintType)
+class MASSENTITY_API UPipeSchematic : public UDataAsset
+{
+	GENERATED_BODY()
+
+public:
+	TConstArrayView<const UPipeProcessor*> GetProcessors() const { return MakeArrayView(Processors);  }
+
+protected:
+	UPROPERTY(EditAnywhere, BlueprintReadOnly, Category = Pipe, Instanced)
+	TArray<UPipeProcessor*> Processors;
+};
diff --git a/Engine/Plugins/Runtime/MassEntity/Source/MassEntityEditor/Private/AssetTypeActions_MassSchematic.cpp b/Engine/Plugins/Runtime/MassEntity/Source/MassEntityEditor/Private/AssetTypeActions_MassSchematic.cpp
new file mode 100644
--- /dev/null
+++ b/Engine/Plugins/Runtime/MassEntity/Source/MassEntityEditor/Private/AssetTypeActions_MassSchematic.cpp
@@ -0,0 +1,37 @@
+// Copyright Epic Games, Inc. All Rights Reserved.
+
+#include "AssetTypeActions_MassSchematic.h"
+#include "MassSchematic.h"
+#include "IMassEntityEditor.h"
+#include "MassEntityEditorModule.h"
+#include "AIModule.h"
+
+#define LOCTEXT_NAMESPACE "AssetTypeActions"
+
+void FAssetTypeActions_PipeSchematic::OpenAssetEditor(const TArray<UObject*>& InObjects, TSharedPtr<IToolkitHost> EditWithinLevelEditor)
+{
+	EToolkitMode::Type Mode = EditWithinLevelEditor.IsValid() ? EToolkitMode::WorldCentric : EToolkitMode::Standalone;
+
+	for (auto ObjIt = InObjects.CreateConstIterator(); ObjIt; ++ObjIt)
+	{
+		if (UPipeSchematic* Pipe = Cast<UPipeSchematic>(*ObjIt))
+		{
+			FPipeEditorModule& EditorModule = FModuleManager::LoadModuleChecked<FPipeEditorModule>("PipeEditor");
+			TSharedRef<IPipeEditor> NewEditor = EditorModule.CreatePipeEditor(Mode, EditWithinLevelEditor, Pipe);
+		}
+	}
+}
+
+UClass* FAssetTypeActions_PipeSchematic::GetSupportedClass() const
+{
+	return UPipeSchematic::StaticClass();
+}
+
+uint32 FAssetTypeActions_PipeSchematic::GetCategories()
+{ 
+	// @todo will probably need to use a different category, like "gameplay"
+	IAIModule& AIModule = FModuleManager::GetModuleChecked<IAIModule>("AIModule").Get();
+	return AIModule.GetAIAssetCategoryBit();
+}
+
+#undef LOCTEXT_NAMESPACE
diff --git a/Engine/Plugins/Runtime/MassEntity/Source/MassEntityEditor/Private/AssetTypeActions_MassSchematic.h b/Engine/Plugins/Runtime/MassEntity/Source/MassEntityEditor/Private/AssetTypeActions_MassSchematic.h
new file mode 100644
--- /dev/null
+++ b/Engine/Plugins/Runtime/MassEntity/Source/MassEntityEditor/Private/AssetTypeActions_MassSchematic.h
@@ -0,0 +1,21 @@
+// Copyright Epic Games, Inc. All Rights Reserved.
+
+#pragma once
+
+#include "CoreMinimal.h"
+#include "AssetTypeActions_Base.h"
+
+/**
+  * Asset type actions for UPipeSchematic.
+ */
+
+class FAssetTypeActions_PipeSchematic : public FAssetTypeActions_Base
+{
+public:
+	// IAssetTypeActions Implementation
+	virtual FText GetName() const override { return NSLOCTEXT("AssetTypeActions", "FAssetTypeActions_PipeSchematic", "PipeSchematic"); }
+	virtual FColor GetTypeColor() const override { return FColor(201, 185, 29); }
+	virtual UClass* GetSupportedClass() const override;
+	virtual void OpenAssetEditor(const TArray<UObject*>& InObjects, TSharedPtr<class IToolkitHost> EditWithinLevelEditor = TSharedPtr<IToolkitHost>()) override;
+	virtual uint32 GetCategories() override;
+};
diff --git a/Engine/Plugins/Runtime/MassEntity/Source/MassEntityEditor/Private/MassEditorStyle.cpp b/Engine/Plugins/Runtime/MassEntity/Source/MassEntityEditor/Private/MassEditorStyle.cpp
new file mode 100644
--- /dev/null
+++ b/Engine/Plugins/Runtime/MassEntity/Source/MassEntityEditor/Private/MassEditorStyle.cpp
@@ -0,0 +1,127 @@
+// Copyright Epic Games, Inc. All Rights Reserved.
+
+#include "MassEditorStyle.h"
+#include "Styling/SlateStyleRegistry.h"
+#include "Styling/CoreStyle.h"
+#include "Styling/SlateTypes.h"
+#include "EditorStyleSet.h"
+#include "Interfaces/IPluginManager.h"
+
+#define IMAGE_BRUSH( RelativePath, ... ) FSlateImageBrush( StyleSet->RootToContentDir( RelativePath, TEXT( ".png" ) ), __VA_ARGS__ )
+#define IMAGE_PLUGIN_BRUSH( RelativePath, ... ) FSlateImageBrush( FMeshEditorStyle::InContent( RelativePath, ".png" ), __VA_ARGS__ )
+#define BOX_BRUSH( RelativePath, ... ) FSlateBoxBrush( StyleSet->RootToContentDir( RelativePath, TEXT( ".png" ) ), __VA_ARGS__ )
+#define TTF_CORE_FONT( RelativePath, ... ) FSlateFontInfo( StyleSet->RootToCoreContentDir( RelativePath, TEXT( ".ttf" ) ), __VA_ARGS__ )
+
+#define DEFAULT_FONT(...) FCoreStyle::GetDefaultFontStyle(__VA_ARGS__)
+
+TSharedPtr<FSlateStyleSet> FPipeEditorStyle::StyleSet = nullptr;
+
+FString FPipeEditorStyle::InContent(const FString& RelativePath, const ANSICHAR* Extension)
+{
+	static FString ContentDir = IPluginManager::Get().FindPlugin(TEXT("PipeEditorModule"))->GetContentDir() / TEXT("Slate");
+	return (ContentDir / RelativePath) + Extension;
+}
+
+void FPipeEditorStyle::Initialize()
+{
+	if (StyleSet.IsValid())
+	{
+		return;
+	}
+
+	const FVector2D Icon8x8(8.0f, 8.0f);
+
+	StyleSet = MakeShared<FSlateStyleSet>(GetStyleSetName());
+
+	StyleSet->SetContentRoot(FPaths::EngineContentDir() / TEXT("Editor/Slate"));
+	StyleSet->SetCoreContentRoot(FPaths::EngineContentDir() / TEXT("Slate"));
+
+	const FScrollBarStyle ScrollBar = FEditorStyle::GetWidgetStyle<FScrollBarStyle>("ScrollBar");
+	const FTextBlockStyle& NormalText = FEditorStyle::Get().GetWidgetStyle<FTextBlockStyle>("NormalText");
+
+	// State
+	{
+		FTextBlockStyle StateIcon = FTextBlockStyle(NormalText)
+			.SetFont(FEditorStyle::Get().GetFontStyle("FontAwesome.12"))
+			.SetColorAndOpacity(FLinearColor(230.0f / 255.0f, 230.0f / 255.0f, 230.0f / 255.0f, 0.5f));
+		StyleSet->Set("Pipe.Icon", StateIcon);
+
+		FTextBlockStyle StateTitle = FTextBlockStyle(NormalText)
+			.SetFont(TTF_CORE_FONT("Fonts/Roboto-Bold", 12))
+			.SetColorAndOpacity(FLinearColor(230.0f / 255.0f, 230.0f / 255.0f, 230.0f / 255.0f));
+		StyleSet->Set("Pipe.State.Title", StateTitle);
+
+		FEditableTextBoxStyle StateTitleEditableText = FEditableTextBoxStyle()
+			.SetFont(TTF_CORE_FONT("Fonts/Roboto-Bold", 10))
+			.SetBackgroundImageNormal(BOX_BRUSH("Common/TextBox", FMargin(4.0f / 16.0f)))
+			.SetBackgroundImageHovered(BOX_BRUSH("Common/TextBox_Hovered", FMargin(4.0f / 16.0f)))
+			.SetBackgroundImageFocused(BOX_BRUSH("Common/TextBox_Hovered", FMargin(4.0f / 16.0f)))
+			.SetBackgroundImageReadOnly(BOX_BRUSH("Common/TextBox_ReadOnly", FMargin(4.0f / 16.0f)))
+			.SetScrollBarStyle(ScrollBar);
+		StyleSet->Set("Pipe.State.TitleEditableText", StateTitleEditableText);
+
+		StyleSet->Set("Pipe.State.TitleInlineEditableText", FInlineEditableTextBlockStyle()
+			.SetTextStyle(StateTitle)
+			.SetEditableTextBoxStyle(StateTitleEditableText));
+	}
+
+	// Task
+	{
+		FTextBlockStyle TaskTitle = FTextBlockStyle(NormalText)
+			.SetFont(TTF_CORE_FONT("Fonts/Roboto-Regular", 11))
+			.SetColorAndOpacity(FLinearColor(230.0f / 255.0f, 230.0f / 255.0f, 230.0f / 255.0f));
+		StyleSet->Set("Pipe.Task.Title", TaskTitle);
+
+		FEditableTextBoxStyle TaskTitleEditableText = FEditableTextBoxStyle()
+			.SetFont(TTF_CORE_FONT("Fonts/Roboto-Regular", 9))
+			.SetBackgroundImageNormal(BOX_BRUSH("Common/TextBox", FMargin(4.0f / 16.0f)))
+			.SetBackgroundImageHovered(BOX_BRUSH("Common/TextBox_Hovered", FMargin(4.0f / 16.0f)))
+			.SetBackgroundImageFocused(BOX_BRUSH("Common/TextBox_Hovered", FMargin(4.0f / 16.0f)))
+			.SetBackgroundImageReadOnly(BOX_BRUSH("Common/TextBox_ReadOnly", FMargin(4.0f / 16.0f)))
+			.SetScrollBarStyle(ScrollBar);
+		StyleSet->Set("Pipe.Task.TitleEditableText", TaskTitleEditableText);
+
+		StyleSet->Set("Pipe.Task.TitleInlineEditableText", FInlineEditableTextBlockStyle()
+			.SetTextStyle(TaskTitle)
+			.SetEditableTextBoxStyle(TaskTitleEditableText));
+	}
+
+	// Details
+	{
+		FTextBlockStyle StateTitle = FTextBlockStyle(NormalText)
+			.SetFont(TTF_CORE_FONT("Fonts/Roboto-Regular", 10))
+			.SetColorAndOpacity(FLinearColor(230.0f / 255.0f, 230.0f / 255.0f, 230.0f / 255.0f));
+		StyleSet->Set("Pipe.Details", StateTitle);
+	}
+
+	const FLinearColor SelectionColor = FColor(0, 0, 0, 32);
+	const FTableRowStyle& NormalTableRowStyle = FEditorStyle::Get().GetWidgetStyle<FTableRowStyle>("TableView.Row");
+	StyleSet->Set("Pipe.Selection",
+		FTableRowStyle(NormalTableRowStyle)
+		.SetActiveBrush(IMAGE_BRUSH("Common/Selection", Icon8x8, SelectionColor))
+		.SetActiveHoveredBrush(IMAGE_BRUSH("Common/Selection", Icon8x8, SelectionColor))
+		.SetInactiveBrush(IMAGE_BRUSH("Common/Selection", Icon8x8, SelectionColor))
+		.SetInactiveHoveredBrush(IMAGE_BRUSH("Common/Selection", Icon8x8, SelectionColor))
+		.SetSelectorFocusedBrush(IMAGE_BRUSH("Common/Selection", Icon8x8, SelectionColor))
+	);
+
+	FSlateStyleRegistry::RegisterSlateStyle(*StyleSet.Get());
+}
+
+
+void FPipeEditorStyle::Shutdown()
+{
+	if (StyleSet.IsValid())
+	{
+		FSlateStyleRegistry::UnRegisterSlateStyle(*StyleSet.Get());
+		ensure(StyleSet.IsUnique());
+		StyleSet.Reset();
+	}
+}
+
+
+FName FPipeEditorStyle::GetStyleSetName()
+{
+	static FName StyleName("PipeEditorStyle");
+	return StyleName;
+}
diff --git a/Engine/Plugins/Runtime/MassEntity/Source/MassEntityEditor/Private/MassEntityEditor.cpp b/Engine/Plugins/Runtime/MassEntity/Source/MassEntityEditor/Private/MassEntityEditor.cpp
new file mode 100644
--- /dev/null
+++ b/Engine/Plugins/Runtime/MassEntity/Source/MassEntityEditor/Private/MassEntityEditor.cpp
@@ -0,0 +1,132 @@
+// Copyright Epic Games, Inc. All Rights Reserved.
+
+#include "MassEntityEditor.h"
+#include "Widgets/Text/STextBlock.h"
+#include "Widgets/Layout/SBorder.h"
+#include "Modules/ModuleManager.h"
+#include "EditorStyleSet.h"
+#include "PropertyEditorModule.h"
+#include "MassEntityEditorModule.h"
+#include "Framework/Commands/Commands.h"
+#include "Framework/Application/SlateApplication.h"
+#include "Framework/MultiBox/MultiBoxBuilder.h"
+#include "MassSchematic.h"
+#include "IDetailsView.h"
+#include "Widgets/Docking/SDockTab.h"
+ 
+
+#define LOCTEXT_NAMESPACE "PipeEditor"
+
+const FName PipeEditorAppName(TEXT("PipeEditorApp"));
+const FName FPipeEditor::AssetDetailsTabId(TEXT("PipeEditor_AssetDetails"));
+
+void FPipeEditor::AddReferencedObjects(FReferenceCollector& Collector)
+{
+	if (PipeSchematic != nullptr)
+	{
+		Collector.AddReferencedObject(PipeSchematic);
+	}
+}
+
+void FPipeEditor::RegisterTabSpawners(const TSharedRef<class FTabManager>& InTabManager)
+{
+	WorkspaceMenuCategory = InTabManager->AddLocalWorkspaceMenuCategory(LOCTEXT("WorkspaceMenu_PipeEditor", "Pipe Editor"));
+	TSharedRef<FWorkspaceItem> WorkspaceMenuCategoryRef = WorkspaceMenuCategory.ToSharedRef();
+
+	FAssetEditorToolkit::RegisterTabSpawners(InTabManager);
+
+	InTabManager->RegisterTabSpawner(AssetDetailsTabId, FOnSpawnTab::CreateSP(this, &FPipeEditor::SpawnTab_AssetDetails))
+		.SetDisplayName(NSLOCTEXT("PipeEditor", "AssetDetailsTab", "Asset Details"))
+		.SetGroup(WorkspaceMenuCategoryRef)
+		.SetIcon(FSlateIcon(FEditorStyle::GetStyleSetName(), "LevelEditor.Tabs.Details"));
+}
+
+void FPipeEditor::UnregisterTabSpawners(const TSharedRef<class FTabManager>& InTabManager)
+{
+	FAssetEditorToolkit::UnregisterTabSpawners(InTabManager);
+	InTabManager->UnregisterTabSpawner(AssetDetailsTabId);
+}
+
+void FPipeEditor::InitEditor(const EToolkitMode::Type Mode, const TSharedPtr< class IToolkitHost >& InitToolkitHost, UPipeSchematic& InPipeSchematic)
+{
+	PipeSchematic = &InPipeSchematic;
+
+	// use the commented block while iterating on the layout since the named layout gets saved
+	//FGuid Result;
+	//FPlatformMisc::CreateGuid(Result);
+	//TSharedRef<FTabManager::FLayout> StandaloneDefaultLayout = FTabManager::NewLayout(*Result.ToString())
+	TSharedRef<FTabManager::FLayout> StandaloneDefaultLayout = FTabManager::NewLayout("Standalone_Pipe_Layout_v1")
+	->AddArea
+	(
+		FTabManager::NewPrimaryArea() ->SetOrientation(Orient_Vertical)
+		->Split
+		(
+			FTabManager::NewStack()
+			->SetSizeCoefficient(0.2f)
+			->AddTab(AssetDetailsTabId, ETabState::OpenedTab)
+			->SetForegroundTab(AssetDetailsTabId)
+		)
+	);
+
+	const bool bCreateDefaultStandaloneMenu = true;
+	const bool bCreateDefaultToolbar = true;
+	FAssetEditorToolkit::InitAssetEditor(Mode, InitToolkitHost, PipeEditorAppName, StandaloneDefaultLayout, bCreateDefaultStandaloneMenu, bCreateDefaultToolbar, PipeSchematic);
+}
+
+FName FPipeEditor::GetToolkitFName() const
+{
+	return FName("PipeEditor");
+}
+
+FText FPipeEditor::GetBaseToolkitName() const
+{
+	return NSLOCTEXT("PipeEditor", "AppLabel", "Pipe");
+}
+
+FString FPipeEditor::GetWorldCentricTabPrefix() const
+{
+	return NSLOCTEXT("PipeEditor", "WorldCentricTabPrefix", "Pipe").ToString();
+}
+
+FLinearColor FPipeEditor::GetWorldCentricTabColorScale() const
+{
+	return FLinearColor( 0.0f, 0.0f, 0.2f, 0.5f );
+}
+
+TSharedRef<SDockTab> FPipeEditor::SpawnTab_AssetDetails(const FSpawnTabArgs& Args)
+{
+	check(Args.GetTabId() == AssetDetailsTabId);
+
+	FPropertyEditorModule& PropertyEditorModule = FModuleManager::GetModuleChecked<FPropertyEditorModule>("PropertyEditor");
+	FDetailsViewArgs DetailsViewArgs;
+	DetailsViewArgs.NameAreaSettings = FDetailsViewArgs::HideNameArea;
+
+	AssetDetailsView = PropertyEditorModule.CreateDetailView(DetailsViewArgs);
+	AssetDetailsView->SetObject(PipeSchematic);
+	AssetDetailsView->OnFinishedChangingProperties().AddSP(this, &FPipeEditor::OnAssetFinishedChangingProperties);
+
+	TSharedRef<SDockTab> SpawnedTab = SNew(SDockTab)
+		.Label(NSLOCTEXT("PipeEditor", "AssetDetailsTab", "Pipe"))
+		[
+			AssetDetailsView.ToSharedRef()
+		];
+
+	return SpawnedTab;
+}
+
+void FPipeEditor::SaveAsset_Execute()
+{
+	// @todo pre-save validation
+
+	// save it
+	FAssetEditorToolkit::SaveAsset_Execute();
+}
+
+void FPipeEditor::OnAssetFinishedChangingProperties(const FPropertyChangedEvent& PropertyChangedEvent)
+{
+	FPipeEditorModule& PipeEditorModule = FModuleManager::GetModuleChecked<FPipeEditorModule>("PipeEditor");
+	PipeEditorModule.GetOnAssetPropertiesChanged().Broadcast(PipeSchematic, PropertyChangedEvent);
+}
+
+
+#undef LOCTEXT_NAMESPACE
diff --git a/Engine/Plugins/Runtime/MassEntity/Source/MassEntityEditor/Private/MassEntityEditor.h b/Engine/Plugins/Runtime/MassEntity/Source/MassEntityEditor/Private/MassEntityEditor.h
new file mode 100644
--- /dev/null
+++ b/Engine/Plugins/Runtime/MassEntity/Source/MassEntityEditor/Private/MassEntityEditor.h
@@ -0,0 +1,56 @@
+// Copyright Epic Games, Inc. All Rights Reserved.
+
+#pragma once
+
+#include "CoreMinimal.h"
+#include "Toolkits/IToolkitHost.h"
+#include "EditorUndoClient.h"
+#include "Toolkits/AssetEditorToolkit.h"
+#include "IMassEntityEditor.h"
+
+class IDetailsView;
+class UPipeSchematic;
+
+class FPipeEditor : public IPipeEditor, public FGCObject//, public FSelfRegisteringEditorUndoClient
+{
+public:
+	virtual void RegisterTabSpawners(const TSharedRef<class FTabManager>& TabManager) override;
+	virtual void UnregisterTabSpawners(const TSharedRef<class FTabManager>& TabManager) override;
+
+	void InitEditor(const EToolkitMode::Type Mode, const TSharedPtr< class IToolkitHost >& InitToolkitHost, UPipeSchematic& InPipeSchematic);
+
+	//~ Begin IToolkit Interface
+	virtual FName GetToolkitFName() const override;
+	virtual FText GetBaseToolkitName() const override;
+	virtual FString GetWorldCentricTabPrefix() const override;
+	virtual FLinearColor GetWorldCentricTabColorScale() const override;
+	//~ End IToolkit Interface
+
+	void OnAssetFinishedChangingProperties(const FPropertyChangedEvent& PropertyChangedEvent);
+
+	/** FEditorUndoClient interface */
+	//	virtual void PostUndo(bool bSuccess) override;
+	//	virtual void PostRedo(bool bSuccess) override;
+
+	/** FGCObject interface */
+	virtual void AddReferencedObjects(FReferenceCollector& Collector) override;
+	virtual FString GetReferencerName() const override
+	{
+		return TEXT("FPipeEditor");
+	}
+
+protected:
+	/** Called when "Save" is clicked for this asset */
+	virtual void SaveAsset_Execute() override;
+
+private:
+	TSharedRef<SDockTab> SpawnTab_AssetDetails(const FSpawnTabArgs& Args);
+
+	/** Pipe Schematic being edited */
+	UPipeSchematic* PipeSchematic;
+
+	/** Asset Property View */
+	TSharedPtr<class IDetailsView> AssetDetailsView;
+
+	static const FName AssetDetailsTabId;
+};
diff --git a/Engine/Plugins/Runtime/MassEntity/Source/MassEntityEditor/Private/MassEntityEditorModule.cpp b/Engine/Plugins/Runtime/MassEntity/Source/MassEntityEditor/Private/MassEntityEditorModule.cpp
new file mode 100644
--- /dev/null
+++ b/Engine/Plugins/Runtime/MassEntity/Source/MassEntityEditor/Private/MassEntityEditorModule.cpp
@@ -0,0 +1,85 @@
+// Copyright Epic Games, Inc. All Rights Reserved.
+
+#include "MassEntityEditorModule.h"
+#include "Modules/ModuleManager.h"
+#include "Developer/AssetTools/Public/IAssetTools.h"
+#include "Developer/AssetTools/Public/AssetToolsModule.h"
+#include "AssetTypeActions_MassSchematic.h"
+#include "MassEditorStyle.h"
+#include "MassEntityEditor.h"
+#include "AIGraphTypes.h" // Class cache
+#include "MassSchematic.h"
+#include "MassProcessor.h"
+
+
+#define LOCTEXT_NAMESPACE "PipeEditor"
+
+IMPLEMENT_MODULE(FPipeEditorModule, PipeEditor)
+
+void FPipeEditorModule::StartupModule()
+{
+	MenuExtensibilityManager = MakeShareable(new FExtensibilityManager);
+	ToolBarExtensibilityManager = MakeShareable(new FExtensibilityManager);
+
+	FPipeEditorStyle::Initialize();
+
+	// Register asset types
+	IAssetTools& AssetTools = FModuleManager::LoadModuleChecked<FAssetToolsModule>("AssetTools").Get();
+
+	TSharedPtr<FAssetTypeActions_PipeSchematic> PipeAssetTypeAction = MakeShareable(new FAssetTypeActions_PipeSchematic);
+	ItemDataAssetTypeActions.Add(PipeAssetTypeAction);
+	AssetTools.RegisterAssetTypeActions(PipeAssetTypeAction.ToSharedRef());
+
+	// Register the details customizers
+	FPropertyEditorModule& PropertyModule = FModuleManager::LoadModuleChecked<FPropertyEditorModule>("PropertyEditor");
+	//PropertyModule.RegisterCustomPropertyTypeLayout("StateTreeVariableDesc", FOnGetPropertyTypeCustomizationInstance::CreateStatic(&FStateTreeVariableDescDetails::MakeInstance));
+	PropertyModule.NotifyCustomizationModuleChanged();
+}
+
+void FPipeEditorModule::ShutdownModule()
+{
+	ProcessorClassCache.Reset();
+	MenuExtensibilityManager.Reset();
+	ToolBarExtensibilityManager.Reset();
+
+	FPipeEditorStyle::Shutdown();
+
+	// Unregister the data asset type actions
+	if (FModuleManager::Get().IsModuleLoaded("AssetTools"))
+	{
+		IAssetTools& AssetToolsModule = FModuleManager::GetModuleChecked<FAssetToolsModule>("AssetTools").Get();
+		for (int i = 0; i < ItemDataAssetTypeActions.Num(); i++)
+		{
+			if (ItemDataAssetTypeActions[i].IsValid())
+			{
+				AssetToolsModule.UnregisterAssetTypeActions(ItemDataAssetTypeActions[i].ToSharedRef());
+			}
+		}
+	}
+	ItemDataAssetTypeActions.Empty();
+
+	// Unregister the details customization
+	if (FModuleManager::Get().IsModuleLoaded("PropertyEditor"))
+	{
+		FPropertyEditorModule& PropertyModule = FModuleManager::LoadModuleChecked<FPropertyEditorModule>("PropertyEditor");
+		//PropertyModule.UnregisterCustomPropertyTypeLayout("StateTreeVariableDesc");
+		PropertyModule.NotifyCustomizationModuleChanged();
+	}
+}
+
+TSharedRef<IPipeEditor> FPipeEditorModule::CreatePipeEditor(const EToolkitMode::Type Mode, const TSharedPtr< class IToolkitHost >& InitToolkitHost, UPipeSchematic* PipeSchematic)
+{
+	if (!ProcessorClassCache.IsValid())
+	{
+		ProcessorClassCache = MakeShareable(new FGraphNodeClassHelper(UPipeProcessor::StaticClass()));
+		ProcessorClassCache->UpdateAvailableBlueprintClasses();
+	}
+
+	TSharedRef<FPipeEditor> NewEditor(new FPipeEditor());
+	if (ensure(PipeSchematic))
+	{
+		NewEditor->InitEditor(Mode, InitToolkitHost, *PipeSchematic);
+	}
+	return NewEditor;
+}
+#undef LOCTEXT_NAMESPACE
