AggregateTicking and Pipe merged into Engine/Plugins/Runtime/MassEntity

diff --git a/Engine/Plugins/Runtime/MassEntity/Source/MassEntityEditor/Private/MassSchematicFactory.cpp b/Engine/Plugins/Runtime/MassEntity/Source/MassEntityEditor/Private/MassSchematicFactory.cpp
new file mode 100644
--- /dev/null
+++ b/Engine/Plugins/Runtime/MassEntity/Source/MassEntityEditor/Private/MassSchematicFactory.cpp
@@ -0,0 +1,27 @@
+// Copyright Epic Games, Inc. All Rights Reserved.
+
+#include "MassSchematicFactory.h"
+#include "MassSchematic.h"
+
+#define LOCTEXT_NAMESPACE "PipeSchematicFactory"
+
+UPipeSchematicFactory::UPipeSchematicFactory(const FObjectInitializer& ObjectInitializer)
+	: Super(ObjectInitializer)
+{
+	SupportedClass = UPipeSchematic::StaticClass();
+	bCreateNew = true;
+	bEditAfterNew = true;
+}
+
+bool UPipeSchematicFactory::CanCreateNew() const
+{
+	return true;
+}
+
+UObject* UPipeSchematicFactory::FactoryCreateNew(UClass* Class,UObject* InParent,FName Name,EObjectFlags Flags,UObject* Context,FFeedbackContext* Warn)
+{
+	check(Class->IsChildOf(UPipeSchematic::StaticClass()));
+	return NewObject<UPipeSchematic>(InParent, Class, Name, Flags);
+}
+
+#undef LOCTEXT_NAMESPACE
diff --git a/Engine/Plugins/Runtime/MassEntity/Source/MassEntityEditor/Public/IMassEntityEditor.h b/Engine/Plugins/Runtime/MassEntity/Source/MassEntityEditor/Public/IMassEntityEditor.h
new file mode 100644
--- /dev/null
+++ b/Engine/Plugins/Runtime/MassEntity/Source/MassEntityEditor/Public/IMassEntityEditor.h
@@ -0,0 +1,13 @@
+// Copyright Epic Games, Inc. All Rights Reserved.
+
+#pragma once
+
+#include "CoreMinimal.h"
+#include "Toolkits/AssetEditorToolkit.h"
+
+
+/** Pipe Editor public interface */
+class MASSENTITYEDITOR_API IPipeEditor : public FAssetEditorToolkit
+{
+public:
+};
diff --git a/Engine/Plugins/Runtime/MassEntity/Source/MassEntityEditor/Public/MassEditorStyle.h b/Engine/Plugins/Runtime/MassEntity/Source/MassEntityEditor/Public/MassEditorStyle.h
new file mode 100644
--- /dev/null
+++ b/Engine/Plugins/Runtime/MassEntity/Source/MassEntityEditor/Public/MassEditorStyle.h
@@ -0,0 +1,26 @@
+// Copyright Epic Games, Inc. All Rights Reserved.
+
+#pragma once
+
+#include "Styling/SlateStyle.h"
+
+class MASSENTITYEDITOR_API FPipeEditorStyle
+{
+public:
+	static void Initialize();
+	static void Shutdown();
+	static TSharedPtr<ISlateStyle> Get() { return StyleSet; }
+	static FName GetStyleSetName();
+
+	static const FSlateBrush* GetBrush(FName PropertyName, const ANSICHAR* Specifier = NULL)
+	{
+		return StyleSet->GetBrush(PropertyName, Specifier);
+	}
+
+	static FString InContent(const FString& RelativePath, const ANSICHAR* Extension);
+
+
+private:
+
+	static TSharedPtr<FSlateStyleSet> StyleSet;
+};
diff --git a/Engine/Plugins/Runtime/MassEntity/Source/MassEntityEditor/Public/MassEntityEditorModule.h b/Engine/Plugins/Runtime/MassEntity/Source/MassEntityEditor/Public/MassEntityEditorModule.h
new file mode 100644
--- /dev/null
+++ b/Engine/Plugins/Runtime/MassEntity/Source/MassEntityEditor/Public/MassEntityEditorModule.h
@@ -0,0 +1,47 @@
+// Copyright Epic Games, Inc. All Rights Reserved.
+
+#pragma  once
+
+#include "CoreMinimal.h"
+#include "Modules/ModuleManager.h"
+#include "AssetTypeCategories.h"
+#include "Toolkits/IToolkitHost.h"
+#include "Toolkits/AssetEditorToolkit.h"
+
+
+class IPipeEditor;
+class FAssetTypeActions_Base;
+struct FGraphPanelNodeFactory;
+class UPipeSchematic;
+
+/**
+* The public interface to this module
+*/
+class MASSENTITYEDITOR_API FPipeEditorModule : public IModuleInterface, public IHasMenuExtensibility, public IHasToolBarExtensibility
+{
+public:
+	DECLARE_MULTICAST_DELEGATE_TwoParams(FOnAssetPropertiesChanged, class UPipeSchematic* /*PipeSchematic*/, const FPropertyChangedEvent& /*PropertyChangedEvent*/);
+
+	virtual void StartupModule() override;
+	virtual void ShutdownModule() override;
+
+	/** Creates an instance of Pipe editor. Only virtual so that it can be called across the DLL boundary. */
+	virtual TSharedRef<IPipeEditor> CreatePipeEditor(const EToolkitMode::Type Mode, const TSharedPtr< class IToolkitHost >& InitToolkitHost, UPipeSchematic* PipeSchematic);
+
+	virtual TSharedPtr<FExtensibilityManager> GetMenuExtensibilityManager() override { return MenuExtensibilityManager; }
+	virtual TSharedPtr<FExtensibilityManager> GetToolBarExtensibilityManager() override { return ToolBarExtensibilityManager; }
+
+	TSharedPtr<struct FGraphNodeClassHelper> GetProcassorClassCache() { return ProcessorClassCache; }
+	
+	FOnAssetPropertiesChanged& GetOnAssetPropertiesChanged() { return OnAssetPropertiesChanged; }
+
+protected:
+	TSharedPtr<struct FGraphNodeClassHelper> ProcessorClassCache;
+
+	TSharedPtr<FExtensibilityManager> MenuExtensibilityManager;
+	TSharedPtr<FExtensibilityManager> ToolBarExtensibilityManager;
+
+	TArray<TSharedPtr<FAssetTypeActions_Base>> ItemDataAssetTypeActions;
+
+	FOnAssetPropertiesChanged OnAssetPropertiesChanged;
+};
diff --git a/Engine/Plugins/Runtime/MassEntity/Source/MassEntityEditor/Public/MassSchematicFactory.h b/Engine/Plugins/Runtime/MassEntity/Source/MassEntityEditor/Public/MassSchematicFactory.h
new file mode 100644
--- /dev/null
+++ b/Engine/Plugins/Runtime/MassEntity/Source/MassEntityEditor/Public/MassSchematicFactory.h
@@ -0,0 +1,22 @@
+// Copyright Epic Games, Inc. All Rights Reserved.
+
+#pragma once
+
+#include "CoreMinimal.h"
+#include "UObject/ObjectMacros.h"
+#include "Factories/Factory.h"
+#include "MassSchematicFactory.generated.h"
+
+UCLASS()
+class UPipeSchematicFactory : public UFactory
+{
+	GENERATED_UCLASS_BODY()
+
+	// UFactory interface
+	virtual UObject* FactoryCreateNew(UClass* Class,UObject* InParent,FName Name,EObjectFlags Flags,UObject* Context,FFeedbackContext* Warn) override;
+	virtual bool CanCreateNew() const override;
+	// End of UFactory interface
+};
+
+
+
diff --git a/Engine/Plugins/Runtime/MassEntity/Source/MassEntityTestSuite/Private/LWComponentTestFarmPlot.cpp b/Engine/Plugins/Runtime/MassEntity/Source/MassEntityTestSuite/Private/LWComponentTestFarmPlot.cpp
new file mode 100644
--- /dev/null
+++ b/Engine/Plugins/Runtime/MassEntity/Source/MassEntityTestSuite/Private/LWComponentTestFarmPlot.cpp
@@ -0,0 +1,265 @@
+// Copyright Epic Games, Inc. All Rights Reserved.
+
+#include "LWComponentTestFarmPlot.h"
+
+#include "Components/HierarchicalInstancedStaticMeshComponent.h"
+#include "Engine/CollisionProfile.h"
+
+//@TODO: Can add a ReadyToHarvest tag component on when things are ready to harvest, to stop them ticking and signal that we need to create an icon
+
+
+//////////////////////////////////////////////////////////////////////////
+// UComponentUpdateSystem
+
+void UComponentUpdateSystem::PostInitProperties()
+{
+	Super::PostInitProperties();
+	ConfigureQueries();
+}
+
+//////////////////////////////////////////////////////////////////////////
+// UFarmHarvestTimerSetIcon
+
+void UFarmHarvestTimerSetIcon::Execute(UEntitySubsystem& EntitySubsystem, FLWComponentSystemExecutionContext& Context)
+{
+	QUICK_SCOPE_CYCLE_COUNTER(SET_ICON_SET_ICON_SET_ICON);
+
+	EntityQuery.ForEachEntityChunk(EntitySubsystem, Context, [this](FLWComponentSystemExecutionContext& Context) {
+
+		const int32 NumEntities = Context.GetEntitiesNum();
+		TConstArrayView<FFarmGridCellData> GridCoordList = Context.GetComponentView<FFarmGridCellData>();
+		TArrayView<FFarmVisualComponent> VisualList = Context.GetMutableComponentView<FFarmVisualComponent>();
+
+		for (int32 i = 0; i < NumEntities; ++i)
+		{
+			const FFarmGridCellData& GridCells = GridCoordList[i];
+
+			const FVector IconPosition(GridCells.CellX*GridCellWidth, GridCells.CellY*GridCellHeight, HarvestIconHeight);
+			const FTransform IconTransform(FQuat::Identity, IconPosition, FVector(HarvestIconScale, HarvestIconScale, HarvestIconScale));
+
+			VisualList[i].HarvestIconIndex = HarvestIconISMC->AddInstance(IconTransform);
+
+			FLWEntity ThisEntity = Context.GetEntity(i);
+			Context.Defer().RemoveTag<FFarmJustBecameReadyToHarvest>(ThisEntity);
+			Context.Defer().AddTag<FFarmReadyToHarvest>(ThisEntity);
+		}
+	});
+}
+
+//////////////////////////////////////////////////////////////////////////
+// ALWComponentTestFarmPlot
+
+ALWComponentTestFarmPlot::ALWComponentTestFarmPlot()
+{
+	PrimaryActorTick.bCanEverTick = true;
+	PrimaryActorTick.bStartWithTickEnabled = true;
+
+
+	USceneComponent* SceneComponent = CreateDefaultSubobject<USceneComponent>(TEXT("SceneComp"));
+
+	RootComponent = SceneComponent;
+
+	HarvestIconISMC = CreateDefaultSubobject<UHierarchicalInstancedStaticMeshComponent>(TEXT("HarvestIconISMC"));
+	HarvestIconISMC->SetupAttachment(SceneComponent);
+	HarvestIconISMC->SetCollisionProfileName(UCollisionProfile::NoCollision_ProfileName);
+}
+
+void ALWComponentTestFarmPlot::AddItemToGrid(UEntitySubsystem* EntitySystem, int32 X, int32 Y, FArchetypeHandle Archetype, int32 VisualIndex)
+{
+	FLWEntity NewItem = EntitySystem->CreateEntity(Archetype);
+	PlantedSquares[X + Y * GridWidth] = NewItem;
+
+	EntitySystem->GetComponentDataChecked<FFarmWaterComponent>(NewItem).DeltaWaterPerSecond = FMath::FRandRange(-0.01f, -0.001f);
+	EntitySystem->GetComponentDataChecked<FHarvestTimerComponent>(NewItem).NumSecondsLeft = 5 + (FMath::Rand() % 100);
+
+	FFarmGridCellData GridCoords;
+	GridCoords.CellX = X;
+	GridCoords.CellY = Y;
+	EntitySystem->GetComponentDataChecked<FFarmGridCellData>(NewItem) = GridCoords;
+
+	const FVector MeshPosition(X*GridCellWidth, Y*GridCellHeight, 0.0f);
+	const FRotator MeshRotation(0.0f, FMath::FRand()*360.0f, 0.0f);
+	const FVector MeshScale(1.0f, 1.0f, 1.0f); //@TODO: plumb in scale param?
+	const FTransform MeshTransform(MeshRotation, MeshPosition, MeshScale);
+
+	FFarmVisualComponent& VisualComp = EntitySystem->GetComponentDataChecked<FFarmVisualComponent>(NewItem);
+	VisualComp.VisualType = VisualIndex;
+	VisualComp.InstanceIndex = VisualDataISMCs[VisualComp.VisualType]->AddInstance(MeshTransform);
+}
+
+void ALWComponentTestFarmPlot::BeginPlay()
+{
+	Super::BeginPlay();
+
+	UEntitySubsystem* EntitySystem = UWorld::GetSubsystem<UEntitySubsystem>(GetWorld());
+
+	FArchetypeHandle CropArchetype = EntitySystem->CreateArchetype(TArray<const UScriptStruct*>{ FFarmWaterComponent::StaticStruct(), FFarmCropComponent::StaticStruct(), FHarvestTimerComponent::StaticStruct(), FFarmVisualComponent::StaticStruct(), FFarmGridCellData::StaticStruct() });
+	FArchetypeHandle FlowerArchetype = EntitySystem->CreateArchetype(TArray<const UScriptStruct*>{ FFarmWaterComponent::StaticStruct(), FFarmFlowerComponent::StaticStruct(), FHarvestTimerComponent::StaticStruct(), FFarmVisualComponent::StaticStruct(), FFarmGridCellData::StaticStruct() });
+
+	PerFrameSystems.Add(NewObject<UFarmWaterUpdateSystem>(this));
+
+	PerSecondSystems.Add(NewObject<UFarmHarvestTimerSystem_Flowers>(this));
+	PerSecondSystems.Add(NewObject<UFarmHarvestTimerSystem_Crops>(this));
+	PerSecondSystems.Add(NewObject<UFarmHarvestTimerExpired>(this));
+
+	UFarmHarvestTimerSetIcon* IconSetter = NewObject<UFarmHarvestTimerSetIcon>(this);
+	IconSetter->HarvestIconISMC = HarvestIconISMC;
+	IconSetter->GridCellWidth = GridCellWidth;
+	IconSetter->GridCellHeight = GridCellHeight;
+	IconSetter->HarvestIconHeight = 200.0f;
+	IconSetter->HarvestIconScale = HarvestIconScale;
+	PerSecondSystems.Add(IconSetter);
+
+	HarvestIconISMC->SetCullDistances(IconNearCullDistance, IconFarCullDistance);
+
+	// Create ISMCs for each mesh type
+	for (const FFarmVisualDataRow& VisualData : VisualDataTable)
+	{
+		UHierarchicalInstancedStaticMeshComponent* HISMC = NewObject<UHierarchicalInstancedStaticMeshComponent>(this);
+		HISMC->SetStaticMesh(VisualData.Mesh);
+		if (VisualData.MaterialOverride != nullptr)
+		{
+			HISMC->SetMaterial(0, VisualData.MaterialOverride);
+		}
+		HISMC->SetCullDistances(VisualNearCullDistance, VisualFarCullDistance);
+		HISMC->SetupAttachment(RootComponent);
+		HISMC->RegisterComponent();
+
+		VisualDataISMCs.Add(HISMC);
+	}
+
+	// Plant us a grid
+	const int32 NumGridCells = GridWidth * GridHeight;
+	PlantedSquares.AddDefaulted(NumGridCells);
+
+	for (int32 Y = 0; Y < GridHeight; ++Y)
+	{
+		for (int32 X = 0; X < GridWidth; ++X)
+		{
+			const bool bIsOdd = ((X + Y) & 1) != 0;
+			const int32 VisualIndex = bIsOdd ? TestDataCropIndicies[FMath::RandRange(0, TestDataCropIndicies.Num() - 1)] : TestDataFlowerIndicies[FMath::RandRange(0, TestDataFlowerIndicies.Num() - 1)];
+
+			AddItemToGrid(EntitySystem, X, Y, bIsOdd ? CropArchetype : FlowerArchetype, VisualIndex);
+		}
+	}
+}
+
+void ALWComponentTestFarmPlot::TickActor(float DeltaTime, enum ELevelTick TickType, FActorTickFunction& ThisTickFunction)
+{
+	Super::TickActor(DeltaTime, TickType, ThisTickFunction);
+
+	QUICK_SCOPE_CYCLE_COUNTER(HeyaTick);
+
+	UEntitySubsystem* EntitySystem = UWorld::GetSubsystem<UEntitySubsystem>(GetWorld());
+
+	// Run every frame systems
+	{
+		FLWComponentSystemExecutionContext ExecContext(DeltaTime);
+		TSharedPtr<FLWCCommandBuffer> DeferredCommandBuffer = MakeShareable(new FLWCCommandBuffer());
+		ExecContext.SetDeferredCommandBuffer(DeferredCommandBuffer);
+
+		for (UComponentUpdateSystem* System : PerFrameSystems)
+		{
+			check(System);
+			System->Execute(*EntitySystem, ExecContext);
+		}
+	}
+
+	// Run per-second systems when it's time
+	NextSecondTimer -= DeltaTime;
+	while (NextSecondTimer < 0.0f)
+	{
+		FLWComponentSystemExecutionContext ExecContext(1.f);
+		TSharedPtr<FLWCCommandBuffer> DeferredCommandBuffer = MakeShareable(new FLWCCommandBuffer());
+		ExecContext.SetDeferredCommandBuffer(DeferredCommandBuffer);
+
+		NextSecondTimer += 1.0f;
+		for (UComponentUpdateSystem* System : PerSecondSystems)
+		{
+			check(System);
+			System->Execute(*EntitySystem, ExecContext);
+		}
+	}
+
+#if 0
+	// Update visuals
+	{
+		QUICK_SCOPE_CYCLE_COUNTER(FARM_BASE_MESHES);
+
+		for (int32 Y = 0; Y < GridHeight; ++Y)
+		{
+			for (int32 X = 0; X < GridWidth; ++X)
+			{
+				FLWEntity GridEntity = PlantedSquares[X + Y * GridWidth];
+				if (EntitySystem->IsValidEntity(GridEntity))
+				{
+					FFarmVisualComponent& VisualComp = EntitySystem->GetComponentDataChecked<FFarmVisualComponent>(GridEntity);
+					if (VisualComp.InstanceIndex < 0)
+					{
+						const FVector MeshPosition(X*GridCellWidth, Y*GridCellHeight, 0.0f);
+						const FRotator MeshRotation(0.0f, FMath::FRand()*360.0f, 0.0f);
+						const FVector MeshScale(1.0f, 1.0f, 1.0f); //@TODO: plumb in scale param?
+						const FTransform MeshTransform(MeshRotation, MeshPosition, MeshScale);
+
+						VisualComp.InstanceIndex = VisualDataISMCs[VisualComp.VisualType]->AddInstance(MeshTransform);
+					}
+				}
+			}
+		}
+	}
+#endif
+	//ForEachComponentChunk
+	// Update visuals
+
+#if 0
+	static int32 QQQZZZ = 2;
+
+	const float HarvestIconHeight = 200.0f;
+
+	if (QQQZZZ == 1)
+	{
+		QUICK_SCOPE_CYCLE_COUNTER(FARM_HARVEST_ICONS);
+		EntitySystem->ForEachComponentChunk<FHarvestTimerComponent>(
+			[&](TArrayView<FHarvestTimerComponent> Timers)
+		{
+			for (FHarvestTimerComponent& Timer : Timers)
+			{
+				if ((Timer.NumSecondsLeft == 0) && (Timer.HarvestIconIndex < 0))
+				{
+					const FVector IconPosition(0.0f, 0.0f, HarvestIconHeight);
+					const FTransform IconTransform(FQuat::Identity, IconPosition, FVector(HarvestIconScale, HarvestIconScale, HarvestIconScale));
+
+					Timer.HarvestIconIndex = HarvestIconISMC->AddInstance(IconTransform);
+				}
+			}
+		});
+	}
+	else if (QQQZZZ == 0)
+	{
+		QUICK_SCOPE_CYCLE_COUNTER(FARM_HARVEST_ICONS);
+
+		for (int32 Y = 0; Y < GridHeight; ++Y)
+		{
+			for (int32 X = 0; X < GridWidth; ++X)
+			{
+				FLWEntity GridEntity = PlantedSquares[X + Y * GridWidth];
+				if (EntitySystem->IsValidEntity(GridEntity))
+				{
+					FHarvestTimerComponent& HarvestTimerData = EntitySystem->GetComponentDataChecked<FHarvestTimerComponent>(GridEntity);
+					if ((HarvestTimerData.NumSecondsLeft == 0) && (HarvestTimerData.HarvestIconIndex < 0))
+					{
+						const FVector IconPosition(X*GridCellWidth, Y*GridCellHeight, HarvestIconHeight);
+						const FTransform IconTransform(FQuat::Identity, IconPosition, FVector(HarvestIconScale, HarvestIconScale, HarvestIconScale));
+
+						HarvestTimerData.HarvestIconIndex = HarvestIconISMC->AddInstance(IconTransform);
+					}
+				}
+			}
+		}
+	}
+#endif
+}
+
+//////////////////////////////////////////////////////////////////////////
+// UFarmWaterUpdateSystem
+
diff --git a/Engine/Plugins/Runtime/MassEntity/Source/MassEntityTestSuite/Private/MassCompositeProcessorTest.cpp b/Engine/Plugins/Runtime/MassEntity/Source/MassEntityTestSuite/Private/MassCompositeProcessorTest.cpp
new file mode 100644
--- /dev/null
+++ b/Engine/Plugins/Runtime/MassEntity/Source/MassEntityTestSuite/Private/MassCompositeProcessorTest.cpp
@@ -0,0 +1,74 @@
+// Copyright Epic Games, Inc. All Rights Reserved.
+
+#include "CoreMinimal.h"
+#include "AITestsCommon.h"
+
+#include "Engine/World.h"
+#include "MassEntitySystem.h"
+#include "MassEntityTypes.h"
+#include "MassEntityTestTypes.h"
+#include "MassExecutor.h"
+
+#define LOCTEXT_NAMESPACE "PipeTest"
+
+PRAGMA_DISABLE_OPTIMIZATION
+
+namespace FPipeCompositeProcessorTest
+{
+
+struct FCompositeProcessorTest_Empty : FEntityTestBase
+{
+	virtual bool InstantTest() override
+	{
+		CA_ASSUME(EntitySubsystem);
+
+		UPipeCompositeProcessor* CompositeProcessor = NewObject<UPipeCompositeProcessor>(EntitySubsystem);
+		check(CompositeProcessor);
+		FPipeContext PipeContext(*EntitySubsystem, /*DeltaSeconds=*/0.f);
+		// it should just run, no warnings
+		UE::Pipe::Executor::Run(*CompositeProcessor, PipeContext);
+		return true;
+	}
+};
+IMPLEMENT_AI_INSTANT_TEST(FCompositeProcessorTest_Empty, "System.Pipe.Processor.Composite.Empty");
+
+
+struct FCompositeProcessorTest_MultipleSubProcessors : FEntityTestBase
+{
+	virtual bool InstantTest() override
+	{
+		CA_ASSUME(EntitySubsystem);
+
+		UPipeCompositeProcessor* CompositeProcessor = NewObject<UPipeCompositeProcessor>(EntitySubsystem);
+		check(CompositeProcessor);
+
+		int Result = 0;
+		{
+			TArray<UPipeProcessor*> Processors;
+			for (int i = 0; i < 3; ++i)
+			{
+				UPipeTestProcessorBase* Processor = NewObject<UPipeTestProcessorBase>(EntitySubsystem);
+				Processor->ExecutionFunction = [Processor, &Result, i](UEntitySubsystem& InEntitySubsystem, FLWComponentSystemExecutionContext& Context) {
+						check(Processor);
+						Result += FMath::Pow(10.f, float(i));
+					};
+				Processors.Add(Processor);
+			}
+
+			CompositeProcessor->SetChildProcessors(MoveTemp(Processors));
+		}
+
+		FPipeContext PipeContext(*EntitySubsystem, /*DeltaSeconds=*/0.f);
+		UE::Pipe::Executor::Run(*CompositeProcessor, PipeContext);
+		AITEST_EQUAL("All of the child processors should get run", Result, 111);
+		return true;
+	}
+};
+IMPLEMENT_AI_INSTANT_TEST(FCompositeProcessorTest_MultipleSubProcessors, "System.Pipe.Processor.Composite.MultipleSubProcessors");
+
+} // FPipeCompositeProcessorTest
+
+PRAGMA_ENABLE_OPTIMIZATION
+
+#undef LOCTEXT_NAMESPACE
+
diff --git a/Engine/Plugins/Runtime/MassEntity/Source/MassEntityTestSuite/Private/MassDependencyTest.cpp b/Engine/Plugins/Runtime/MassEntity/Source/MassEntityTestSuite/Private/MassDependencyTest.cpp
new file mode 100644
--- /dev/null
+++ b/Engine/Plugins/Runtime/MassEntity/Source/MassEntityTestSuite/Private/MassDependencyTest.cpp
@@ -0,0 +1,289 @@
+// Copyright Epic Games, Inc. All Rights Reserved.
+
+#include "CoreMinimal.h"
+#include "AITestsCommon.h"
+
+#include "MassProcessorDependencySolver.h"
+#include "MassEntityTestTypes.h"
+
+#define LOCTEXT_NAMESPACE "PipeTest"
+
+PRAGMA_DISABLE_OPTIMIZATION
+
+namespace FPipeDependencySolverTest
+{
+
+template<typename T>
+static FName GetProcessorName()
+{
+	return T::StaticClass()->GetFName();
+}
+
+struct FDependencySolverBase : FAITestBase
+{
+	TArray<UPipeTestProcessorBase*> Processors;
+	TArray<FProcessorDependencySolver::FOrderInfo> Result;
+	
+	virtual bool SetUp() override
+	{
+		Processors.Reset();
+		return true;
+	}
+
+	void Solve()
+	{
+		Result.Reset();
+		FProcessorDependencySolver Solver(MakeArrayView((UPipeProcessor**)Processors.GetData(), Processors.Num()), TEXT("Test") );
+		Solver.ResolveDependencies(Result);
+	}
+};	
+
+struct FTrivialDependency : FDependencySolverBase
+{
+	virtual bool SetUp() override
+	{
+		{
+			UPipeTestProcessorBase* Proc = Processors.Add_GetRef(NewObject<UPipeTestProcessor_A>());
+			Proc->GetMutableExecutionOrder().ExecuteAfter.Add(GetProcessorName<UPipeTestProcessor_C>());
+		}
+
+		{
+			UPipeTestProcessorBase* Proc = Processors.Add_GetRef(NewObject<UPipeTestProcessor_B>());
+			Proc->GetMutableExecutionOrder().ExecuteAfter.Add(GetProcessorName<UPipeTestProcessor_A>());
+		}
+
+		Processors.Add(NewObject<UPipeTestProcessor_C>());
+
+		return true;
+	}
+
+	virtual bool InstantTest() override
+	{
+		Solve();
+
+		AITEST_TRUE("C is expected to be first", Result[0].Name == GetProcessorName<UPipeTestProcessor_C>());
+		AITEST_TRUE("A is expected to be second", Result[1].Name == GetProcessorName<UPipeTestProcessor_A>());
+		AITEST_TRUE("B is expected to be third", Result[2].Name == GetProcessorName<UPipeTestProcessor_B>());
+
+		return true;
+	}
+};
+IMPLEMENT_AI_INSTANT_TEST(FTrivialDependency, "System.Pipe.Dependencies.Trivial");
+
+
+struct FDeepGroup : FDependencySolverBase
+{
+	virtual bool SetUp() override
+	{
+		{
+			UPipeTestProcessorBase* Proc = Processors.Add_GetRef(NewObject<UPipeTestProcessor_A>());
+			Proc->GetMutableExecutionOrder().ExecuteAfter.Add(TEXT("W.X.Y.Z"));
+			Proc->GetMutableExecutionOrder().ExecuteInGroup = TEXT("P.Q.R");
+		}
+
+		{
+			UPipeTestProcessorBase* Proc = Processors.Add_GetRef(NewObject<UPipeTestProcessor_B>());
+			Proc->GetMutableExecutionOrder().ExecuteInGroup = TEXT("W.X.Y.Z");
+		}
+
+		return true;
+	}
+
+	virtual bool InstantTest() override
+	{
+		Solve();
+
+		// dump all the group information from the Result collection for easier ordering testing
+		for (int32 i = 0; i < Result.Num(); ++i)
+		{
+			if (Result[i].NodeType != EDependencyNodeType::Processor)
+			{
+				Result.RemoveAt(i--, 1, /*bAllowShrinking=*/false);
+			}
+		}
+
+		AITEST_TRUE("B is expected to be first", Result[0].Name == GetProcessorName<UPipeTestProcessor_B>());
+		AITEST_TRUE("A is expected to be second", Result[1].Name == GetProcessorName<UPipeTestProcessor_A>());
+
+		return true;
+	}
+};
+IMPLEMENT_AI_INSTANT_TEST(FDeepGroup, "System.Pipe.Dependencies.DeepGroup");
+
+struct FComplexScenario : FDependencySolverBase
+{
+	virtual bool SetUp() override
+	{
+		{
+			UPipeTestProcessorBase* Proc = Processors.Add_GetRef(NewObject<UPipeTestProcessor_A>());
+			Proc->GetMutableExecutionOrder().ExecuteInGroup = TEXT("X.Z");
+			Proc->GetMutableExecutionOrder().ExecuteAfter.Add(TEXT("X.Y"));
+			Proc->GetMutableExecutionOrder().ExecuteAfter.Add(UPipeTestProcessor_E::StaticClass()->GetFName());
+		}
+
+		{
+			UPipeTestProcessorBase* Proc = Processors.Add_GetRef(NewObject<UPipeTestProcessor_B>());
+			Proc->GetMutableExecutionOrder().ExecuteInGroup = TEXT("X.Y");
+		}
+
+		{
+			UPipeTestProcessorBase* Proc = Processors.Add_GetRef(NewObject<UPipeTestProcessor_C>());
+			Proc->GetMutableExecutionOrder().ExecuteInGroup = TEXT("X.Y");
+		}
+
+		{
+			UPipeTestProcessorBase* Proc = Processors.Add_GetRef(NewObject<UPipeTestProcessor_D>());
+			Proc->GetMutableExecutionOrder().ExecuteBefore.Add(UPipeTestProcessor_A::StaticClass()->GetFName());
+		}
+
+		{
+			UPipeTestProcessorBase* Proc = Processors.Add_GetRef(NewObject<UPipeTestProcessor_E>());
+			Proc->GetMutableExecutionOrder().ExecuteInGroup = TEXT("X.Z");
+		}
+
+		{
+			UPipeTestProcessorBase* Proc = Processors.Add_GetRef(NewObject<UPipeTestProcessor_F>());
+			Proc->GetMutableExecutionOrder().ExecuteAfter.Add(UPipeTestProcessor_A::StaticClass()->GetFName());
+		}
+		return true;
+	}
+
+	virtual bool InstantTest() override
+	{
+		Solve();
+
+		// dump all the group information from the Result collection for easier ordering testing
+		for (int32 i = 0; i < Result.Num(); ++i)
+		{
+			if (Result[i].NodeType != EDependencyNodeType::Processor)
+			{
+				Result.RemoveAt(i--, 1, /*bAllowShrinking=*/false);
+			}
+		}
+
+		AITEST_TRUE("D is the only fully dependency-less processor so should be first", Result[0].Name == GetProcessorName<UPipeTestProcessor_D>());		
+		AITEST_TRUE("B and C come next", (Result[1].Name == GetProcessorName<UPipeTestProcessor_B>() || Result[2].Name == GetProcessorName<UPipeTestProcessor_B>()) && (Result[1].Name == GetProcessorName<UPipeTestProcessor_C>() || Result[2].Name == GetProcessorName<UPipeTestProcessor_C>()));
+		AITEST_TRUE("Following by E", Result[3].Name == GetProcessorName<UPipeTestProcessor_E>());
+		AITEST_TRUE("Then A", Result[4].Name == GetProcessorName<UPipeTestProcessor_A>());
+		AITEST_TRUE("F is last", Result[5].Name == GetProcessorName<UPipeTestProcessor_F>());
+
+		return true;
+	}
+};
+IMPLEMENT_AI_INSTANT_TEST(FComplexScenario, "System.Pipe.Dependencies.Complex");
+
+
+struct FBadInput_Empty : FDependencySolverBase
+{
+	virtual bool InstantTest() override
+	{
+		Solve();
+		AITEST_TRUE("Empty input should be handled gracefully", Result.Num() == 0);
+		return true;
+	}
+};
+IMPLEMENT_AI_INSTANT_TEST(FBadInput_Empty, "System.Pipe.Dependencies.BadInput.Empty");
+
+struct FBadInput_Null : FDependencySolverBase
+{
+	virtual bool InstantTest() override
+	{
+		GetTestRunner().AddExpectedError(TEXT("nullptr found in Processors"), EAutomationExpectedErrorFlags::Contains, 1);
+		Processors.Reset();
+		Processors.Add(nullptr);
+		Solve();
+		AITEST_TRUE("Single nullptr input should be handled gracefully", Result.Num() == 0);
+		return true;
+	}
+};
+IMPLEMENT_AI_INSTANT_TEST(FBadInput_Null, "System.Pipe.Dependencies.BadInput.SingleNull");
+
+struct FBadInput_MultipleNulls : FDependencySolverBase
+{
+	virtual bool InstantTest() override
+	{
+		GetTestRunner().AddExpectedError(TEXT("nullptr found in Processors"), EAutomationExpectedErrorFlags::Contains, 3);
+		Processors.Reset();
+		Processors.Add(nullptr);
+		Processors.Add(nullptr);
+		Processors.Add(nullptr);
+		Solve();
+		AITEST_TRUE("Multiple nullptr inputs should be handled gracefully", Result.Num() == 0);
+		return true;
+	}
+};
+IMPLEMENT_AI_INSTANT_TEST(FBadInput_MultipleNulls, "System.Pipe.Dependencies.BadInput.MultipleNulls");
+
+struct FBadInput_MixedNulls : FDependencySolverBase
+{
+	virtual bool InstantTest() override
+	{
+		GetTestRunner().AddExpectedError(TEXT("nullptr found in Processors"), EAutomationExpectedErrorFlags::Contains, 2);
+		Processors.Reset();
+		Processors.Add(nullptr);
+		Processors.Add(NewObject<UPipeTestProcessor_A>());
+		Processors.Add(nullptr);
+		Processors.Add(NewObject<UPipeTestProcessor_B>());
+		Solve();
+		AITEST_TRUE("Mixed nullptr and proper inputs should be handled gracefully", Result.Num() == 2);
+		return true;
+	}
+};
+IMPLEMENT_AI_INSTANT_TEST(FBadInput_MixedNulls, "System.Pipe.Dependencies.BadInput.NullsMixedIn");
+
+struct FBadInput_Duplicates : FDependencySolverBase
+{
+	virtual bool InstantTest() override
+	{
+		GetTestRunner().AddExpectedError(TEXT("already registered. Duplicates are not supported"), EAutomationExpectedErrorFlags::Contains, 2);
+		Processors.Reset();
+		Processors.Add(NewObject<UPipeTestProcessor_A>());
+		Processors.Add(NewObject<UPipeTestProcessor_A>());
+		Processors.Add(NewObject<UPipeTestProcessor_A>());
+		Solve();
+		AITEST_TRUE("Duplicates in input should be handled gracefully", Result.Num() == 1);
+		return true;
+	}
+};
+IMPLEMENT_AI_INSTANT_TEST(FBadInput_Duplicates, "System.Pipe.Dependencies.BadInput.Duplicates");
+
+struct FGroupNamesGeneration : FAITestBase
+{
+	virtual bool InstantTest() override
+	{
+		TArray<FString> SubGroupNames;
+		FName EmptyName;
+		FProcessorDependencySolver::CreateSubGroupNames(EmptyName, SubGroupNames);
+
+		AITEST_TRUE("Empty group name is supported", SubGroupNames.Num() > 0);
+		AITEST_TRUE("Empty group name handled like any other name", SubGroupNames[0] == EmptyName.ToString());
+		
+		FProcessorDependencySolver::CreateSubGroupNames(TEXT("X"), SubGroupNames);
+		AITEST_TRUE("Trivial group name is supported", SubGroupNames.Num() > 0);
+		AITEST_TRUE("Trivial group name shouldn\'t get decorated", SubGroupNames[0] == TEXT("X"));
+
+		FProcessorDependencySolver::CreateSubGroupNames(TEXT("W.X.Y.Z"), SubGroupNames);
+		AITEST_TRUE("Complex group name should be result in a number of group names equal to group name\'s depth", SubGroupNames.Num() == 4);
+		AITEST_TRUE("Group name W.X.Y.Z should contain subgroup W", SubGroupNames.Find(TEXT("W")) != INDEX_NONE);
+		AITEST_TRUE("Group name W.X.Y.Z should contain subgroup W.X", SubGroupNames.Find(TEXT("W.X")) != INDEX_NONE);
+		AITEST_TRUE("Group name W.X.Y.Z should contain subgroup W.X.Y", SubGroupNames.Find(TEXT("W.X.Y")) != INDEX_NONE);
+		AITEST_TRUE("Group name W.X.Y.Z should contain subgroup W.X.Y.Z", SubGroupNames.Find(TEXT("W.X.Y.Z")) != INDEX_NONE);
+		AITEST_TRUE("Split up of group name W.X.Y.Z should result in a given order"
+			, SubGroupNames[0] == TEXT("W") && SubGroupNames[1] == TEXT("W.X") && SubGroupNames[2] == TEXT("W.X.Y") && SubGroupNames[3] == TEXT("W.X.Y.Z"));
+
+		return true;
+	}
+};
+
+IMPLEMENT_AI_INSTANT_TEST(FGroupNamesGeneration, "System.Pipe.Dependencies.SubgroupNames");
+
+// tests to add:
+// * handling missing dependencies
+// * single non-null processor
+// * circular dependencies  
+
+} // FPipeDependencySolverTest
+
+PRAGMA_ENABLE_OPTIMIZATION
+
+#undef LOCTEXT_NAMESPACE
diff --git a/Engine/Plugins/Runtime/MassEntity/Source/MassEntityTestSuite/Private/MassEntityTest.cpp b/Engine/Plugins/Runtime/MassEntity/Source/MassEntityTestSuite/Private/MassEntityTest.cpp
new file mode 100644
--- /dev/null
+++ b/Engine/Plugins/Runtime/MassEntity/Source/MassEntityTestSuite/Private/MassEntityTest.cpp
@@ -0,0 +1,364 @@
+// Copyright Epic Games, Inc. All Rights Reserved.
+
+#include "CoreMinimal.h"
+#include "AITestsCommon.h"
+
+#include "Engine/World.h"
+#include "MassEntitySystem.h"
+#include "MassEntityTypes.h"
+#include "MassEntityTestTypes.h"
+#include "MassExecutor.h"
+
+#define LOCTEXT_NAMESPACE "PipeTest"
+
+PRAGMA_DISABLE_OPTIMIZATION
+
+namespace FPipeEntityTest
+{
+#if WITH_AGGREGATETICKING_DEBUG
+struct FEntityTest_ArchetypeCreation : FEntityTestBase
+{
+	virtual bool InstantTest() override
+	{
+		AITEST_TRUE("Floats archetype should have been created", FloatsArchetype.IsValid());
+		AITEST_TRUE("Floats archetype should have been created", FloatsArchetype.IsValid());
+		AITEST_TRUE("Ints archetype should have been created", IntsArchetype.IsValid());
+
+		TArray<const UScriptStruct*> FragmentsList;
+		EntitySubsystem->DebugGetArchetypeComponentTypes(FloatsArchetype, FragmentsList);
+		AITEST_EQUAL("Floats archetype should contain just a single fragment", FragmentsList.Num(), 1);
+		AITEST_EQUAL("Floats archetype\'s lone fragment should be of Float fragment type", FragmentsList[0], FTestFragment_Float::StaticStruct());
+
+		FragmentsList.Reset();
+		EntitySubsystem->DebugGetArchetypeComponentTypes(IntsArchetype, FragmentsList);
+		AITEST_EQUAL("Ints archetype should contain just a single fragment", FragmentsList.Num(), 1);
+		AITEST_EQUAL("Ints archetype\'s lone fragment should be of Ints fragment type", FragmentsList[0], FTestFragment_Int::StaticStruct());
+
+		FragmentsList.Reset();
+		EntitySubsystem->DebugGetArchetypeComponentTypes(FloatsIntsArchetype, FragmentsList);
+		AITEST_EQUAL("FloatsInts archetype should contain exactly two fragments", FragmentsList.Num(), 2);
+		AITEST_TRUE("FloatsInts archetype\'s should contain both expected fragment types", FragmentsList.Find(FTestFragment_Int::StaticStruct()) != INDEX_NONE && FragmentsList.Find(FTestFragment_Float::StaticStruct()) != INDEX_NONE);
+		
+		return true;
+	}
+};
+IMPLEMENT_AI_INSTANT_TEST(FEntityTest_ArchetypeCreation, "System.Pipe.Entity.AchetypesCreation");
+
+
+struct FEntityTest_ArchetypeEquivalence : FEntityTestBase
+{
+	virtual bool InstantTest() override
+	{
+		TArray<const UScriptStruct*> FragmentsA = { FTestFragment_Float::StaticStruct(), FTestFragment_Int::StaticStruct() };
+		TArray<const UScriptStruct*> FragmentsB = { FTestFragment_Int::StaticStruct(), FTestFragment_Float::StaticStruct() };
+		const FArchetypeHandle ArchetypeA = EntitySubsystem->CreateArchetype(FragmentsA);
+		const FArchetypeHandle ArchetypeB = EntitySubsystem->CreateArchetype(FragmentsB);
+
+		AITEST_EQUAL("Archetype creation is expected to be independent of fragments ordering", ArchetypeA, ArchetypeB);
+
+		return true;
+	}
+};
+IMPLEMENT_AI_INSTANT_TEST(FEntityTest_ArchetypeEquivalence, "System.Pipe.Entity.AchetypeEquivalance");
+
+
+struct FEntityTest_MultipleEntitiesCreation : FEntityTestBase
+{
+	virtual bool InstantTest() override
+	{
+		CA_ASSUME(EntitySubsystem);
+		int32 Counts[] = { 10, 100, 1000 };
+		int32 TotalCreatedCount = 0;
+		FArchetypeHandle Archetypes[] = { FloatsArchetype, IntsArchetype, FloatsIntsArchetype };
+
+		for (int ArchetypeIndex = 0; ArchetypeIndex < (sizeof(Archetypes) / sizeof(FArchetypeHandle)); ++ArchetypeIndex)
+		{
+			for (int i = 0; i < Counts[ArchetypeIndex]; ++i)
+			{
+				EntitySubsystem->CreateEntity(Archetypes[ArchetypeIndex]);
+			}
+			TotalCreatedCount += Counts[ArchetypeIndex];
+		}
+		AITEST_EQUAL("The total number of entities must match the number created", EntitySubsystem->DebugGetEntityCount(), TotalCreatedCount);
+		AITEST_EQUAL("10 entities of FloatsArchetype should have been created", EntitySubsystem->DebugGetArchetypeEntitiesCount(FloatsArchetype), 10);
+		AITEST_EQUAL("100 entities of IntsArchetype should have been created", EntitySubsystem->DebugGetArchetypeEntitiesCount(IntsArchetype), 100);
+		AITEST_EQUAL("1000 entities of FloatsIntsArchetype should have been created", EntitySubsystem->DebugGetArchetypeEntitiesCount(FloatsIntsArchetype), 1000);
+
+		return true;
+	}
+};
+IMPLEMENT_AI_INSTANT_TEST(FEntityTest_MultipleEntitiesCreation, "System.Pipe.Entity.MultipleEntitiesCreation");
+
+
+struct FEntityTest_EntityBatchCreation : FEntityTestBase
+{
+	virtual bool InstantTest() override
+	{
+		CA_ASSUME(EntitySubsystem);
+		const int32 Count = 123;
+		TArray<FLWEntity> Entities;
+		EntitySubsystem->BatchCreateEntities(FloatsIntsArchetype, Count, Entities);
+		AITEST_EQUAL("Batch creation should create the expected number of entities", Entities.Num(), Count);
+		AITEST_EQUAL("The total number of entities present must match the number requested", EntitySubsystem->DebugGetEntityCount(), Count);
+
+		return true;
+	}
+};
+IMPLEMENT_AI_INSTANT_TEST(FEntityTest_EntityBatchCreation, "System.Pipe.Entity.BatchCreation");
+
+struct FEntityTest_BatchCreatingSingleEntity : FEntityTestBase
+{
+	virtual bool InstantTest() override
+	{
+		CA_ASSUME(EntitySubsystem);
+		TArray<FLWEntity> Entities;
+		EntitySubsystem->BatchCreateEntities(FloatsIntsArchetype, /*Count=*/1, Entities);
+		AITEST_EQUAL("Batch creation should have created a single entity", Entities.Num(), 1);
+		AITEST_EQUAL("The total number of entities present must match the number created by batch creation", EntitySubsystem->DebugGetEntityCount(), 1);
+
+		return true;
+	}
+};
+IMPLEMENT_AI_INSTANT_TEST(FEntityTest_BatchCreatingSingleEntity, "System.Pipe.Entity.BatchCreatingSingleEntity");
+
+
+struct FEntityTest_EntityCreation : FEntityTestBase
+{
+	virtual bool InstantTest() override
+	{
+		CA_ASSUME(EntitySubsystem);		
+		const FLWEntity Entity = EntitySubsystem->CreateEntity(FloatsArchetype);
+		AITEST_EQUAL("There should be one entity across the whole system", EntitySubsystem->DebugGetEntityCount(), 1);
+		AITEST_EQUAL("Entity\'s archetype should be the Float one", EntitySubsystem->GetArchetypeForEntity(Entity), FloatsArchetype);
+		AITEST_EQUAL("The created entity should have been added to the Floats archetype", EntitySubsystem->DebugGetArchetypeEntitiesCount(FloatsArchetype), 1);
+		AITEST_EQUAL("Other archetypes should not get affected", EntitySubsystem->DebugGetArchetypeEntitiesCount(IntsArchetype) + EntitySubsystem->DebugGetArchetypeEntitiesCount(FloatsIntsArchetype), 0);
+
+		return true;
+	}
+};
+IMPLEMENT_AI_INSTANT_TEST(FEntityTest_EntityCreation, "System.Pipe.Entity.EntityCreation");
+
+
+struct FEntityTest_EntityCreationFromInstances : FEntityTestBase
+{
+	virtual bool InstantTest() override
+	{
+		CA_ASSUME(EntitySubsystem);
+		const FLWEntity Entity = EntitySubsystem->CreateEntity(MakeArrayView(&InstanceInt, 1));
+		AITEST_EQUAL("There should be one entity across the whole system", EntitySubsystem->DebugGetEntityCount(), 1);
+		AITEST_EQUAL("Entity\'s archetype should be the Ints one", EntitySubsystem->GetArchetypeForEntity(Entity), IntsArchetype);
+		AITEST_EQUAL("The created entity should have been added to the Ints archetype", EntitySubsystem->DebugGetArchetypeEntitiesCount(IntsArchetype), 1);
+		AITEST_EQUAL("The entity should have the new component with the correct value set", EntitySubsystem->GetComponentDataChecked<FTestFragment_Int>(Entity).Value, TestIntValue);
+
+		return true;
+	}
+};
+IMPLEMENT_AI_INSTANT_TEST(FEntityTest_EntityCreationFromInstances, "System.Pipe.Entity.EntityCreationFromInstances");
+
+#if 0 // this test is compiled out since AddComponentToEntity will fails an ensure if a redundant fragment gets added
+struct FEntityTest_AddingRedundantFragment : FEntityTestBase
+{
+	virtual bool InstantTest() override
+	{
+		CA_ASSUME(EntitySubsystem);
+		const FLWEntity Entity = EntitySubsystem->CreateEntity(FloatsArchetype);
+		EntitySubsystem->AddComponentToEntity(Entity, FTestFragment_Float::StaticStruct());		
+		AITEST_EQUAL("Adding a fragment that a given entity\s archetype already has should do nothing", EntitySubsystem->GetArchetypeForEntity(Entity), FloatsArchetype);
+
+		return true;
+	}
+};
+IMPLEMENT_AI_INSTANT_TEST(FEntityTest_AddingRedundantFragment, "System.Pipe.Entity.AddingRedundantFragment");
+#endif //0
+
+
+struct FEntityTest_AddingFragmentType : FEntityTestBase
+{
+	virtual bool InstantTest() override
+	{
+		CA_ASSUME(EntitySubsystem);
+		const FLWEntity Entity = EntitySubsystem->CreateEntity(FloatsArchetype);
+		EntitySubsystem->AddComponentToEntity(Entity, FTestFragment_Int::StaticStruct());
+		AITEST_EQUAL("There should be one entity across the whole system", EntitySubsystem->DebugGetEntityCount(), 1);
+		AITEST_EQUAL("The original archetype should now have no entities", EntitySubsystem->DebugGetArchetypeEntitiesCount(FloatsArchetype), 0);
+		AITEST_EQUAL("The destination archetype should now store a single entity", EntitySubsystem->DebugGetArchetypeEntitiesCount(FloatsIntsArchetype), 1);
+		AITEST_EQUAL("The remaining archetype should not be affected", EntitySubsystem->DebugGetArchetypeEntitiesCount(IntsArchetype), 0);
+		// this test was originally failing due to FEntityData.CurrentArchetype not getting updated during entity moving between archetypes
+		AITEST_EQUAL("The entity should get associated with the new archetype", EntitySubsystem->GetArchetypeForEntity(Entity), FloatsIntsArchetype);
+
+		return true;
+	}
+};
+IMPLEMENT_AI_INSTANT_TEST(FEntityTest_AddingFragmentType, "System.Pipe.Entity.AddingFragmentType");
+
+struct FEntityTest_AddingFragmentInstance : FEntityTestBase
+{
+	virtual bool InstantTest() override
+	{
+		CA_ASSUME(EntitySubsystem);
+		const FLWEntity Entity = EntitySubsystem->CreateEntity(FloatsArchetype);
+		EntitySubsystem->AddComponentInstanceListToEntity(Entity, MakeArrayView(&InstanceInt, 1));
+		AITEST_EQUAL("There should be one entity across the whole system", EntitySubsystem->DebugGetEntityCount(), 1);
+		AITEST_EQUAL("The original archetype should now have no entities", EntitySubsystem->DebugGetArchetypeEntitiesCount(FloatsArchetype), 0);
+		AITEST_EQUAL("The destination archetype should now store a single entity", EntitySubsystem->DebugGetArchetypeEntitiesCount(FloatsIntsArchetype), 1);
+		AITEST_EQUAL("The archetype containing just the new fragment should not be affected", EntitySubsystem->DebugGetArchetypeEntitiesCount(IntsArchetype), 0);
+		// this test was originally failing due to FEntityData.CurrentArchetype not getting updated during entity moving between archetypes
+		AITEST_EQUAL("The entity should get associated with the new archetype", EntitySubsystem->GetArchetypeForEntity(Entity), FloatsIntsArchetype);
+		AITEST_EQUAL("The entity should have the new component with the correct value set", EntitySubsystem->GetComponentDataChecked<FTestFragment_Int>(Entity).Value, TestIntValue);
+
+		return true;
+	}
+};
+IMPLEMENT_AI_INSTANT_TEST(FEntityTest_AddingFragmentInstance, "System.Pipe.Entity.AddingFragmentType");
+
+
+struct FEntityTest_RemovingFragment : FEntityTestBase
+{
+	virtual bool InstantTest() override
+	{
+		CA_ASSUME(EntitySubsystem);
+		const FLWEntity Entity = EntitySubsystem->CreateEntity(FloatsIntsArchetype);
+		EntitySubsystem->RemoveComponentFromEntity(Entity, FTestFragment_Float::StaticStruct());
+		AITEST_EQUAL("There should be just one entity across the whole system", EntitySubsystem->DebugGetEntityCount(), 1);
+		AITEST_EQUAL("The original archetype should now have no entities", EntitySubsystem->DebugGetArchetypeEntitiesCount(FloatsIntsArchetype), 0);
+		AITEST_EQUAL("The destination archetype should now store a single entity", EntitySubsystem->DebugGetArchetypeEntitiesCount(IntsArchetype), 1);
+		AITEST_EQUAL("The remaining archetype should not be affected", EntitySubsystem->DebugGetArchetypeEntitiesCount(FloatsArchetype), 0);
+		// this test was originally failing due to FEntityData.CurrentArchetype not getting updated during entity moving between archetypes
+		AITEST_EQUAL("The entity should get associated with the new archetype", EntitySubsystem->GetArchetypeForEntity(Entity), IntsArchetype);
+
+		return true;
+	}
+};
+IMPLEMENT_AI_INSTANT_TEST(FEntityTest_RemovingFragment, "System.Pipe.Entity.RemovingFragment");
+
+struct FEntityTest_RemovingLastFragment : FEntityTestBase
+{
+	virtual bool InstantTest() override
+	{
+		CA_ASSUME(EntitySubsystem);
+		const FLWEntity Entity = EntitySubsystem->CreateEntity(FloatsArchetype);
+		EntitySubsystem->RemoveComponentFromEntity(Entity, FTestFragment_Float::StaticStruct());
+		AITEST_EQUAL("There should be one entity across the whole system", EntitySubsystem->DebugGetEntityCount(), 1);
+		AITEST_EQUAL("The original archetype should now have no entities", EntitySubsystem->DebugGetArchetypeEntitiesCount(FloatsArchetype), 0);
+		// this test was originally failing due to FEntityData.CurrentArchetype not getting updated during entity moving between archetypes
+		AITEST_EQUAL("The entity should not get associated to any archetype", EntitySubsystem->GetArchetypeForEntity(Entity), EmptyArchetype);
+
+		return true;
+	}
+};
+IMPLEMENT_AI_INSTANT_TEST(FEntityTest_RemovingLastFragment, "System.Pipe.Entity.RemovingLastFragment");
+
+struct FEntityTest_DestroyEntity : FEntityTestBase
+{
+	virtual bool InstantTest() override
+	{
+		CA_ASSUME(EntitySubsystem);
+		const FLWEntity Entity = EntitySubsystem->CreateEntity(FloatsArchetype);
+		AITEST_EQUAL("The entity should get associated to the right archetype", EntitySubsystem->GetArchetypeForEntity(Entity), FloatsArchetype);
+		EntitySubsystem->DestroyEntity(Entity);
+		AITEST_EQUAL("There should not be any entity across the whole system", EntitySubsystem->DebugGetEntityCount(), 0);
+		AITEST_EQUAL("The original archetype should now have no entities", EntitySubsystem->DebugGetArchetypeEntitiesCount(FloatsArchetype), 0);
+		return true;
+	}
+};
+IMPLEMENT_AI_INSTANT_TEST(FEntityTest_DestroyEntity, "System.Pipe.Entity.DestroyEntity");
+
+struct FEntityTest_EntityReservationAndBuilding : FEntityTestBase
+{
+	virtual bool InstantTest() override
+	{
+		CA_ASSUME(EntitySubsystem);
+		const FLWEntity ReservedEntity = EntitySubsystem->ReserveEntity();
+		AITEST_TRUE("The reserved entity should be a valid entity", EntitySubsystem->IsEntityValid(ReservedEntity));
+		AITEST_FALSE("The reserved entity should not be a valid entity", EntitySubsystem->IsEntityBuilt(ReservedEntity));
+		EntitySubsystem->BuildEntity(ReservedEntity, FloatsArchetype);
+		AITEST_TRUE("The reserved entity should be a valid entity", EntitySubsystem->IsEntityBuilt(ReservedEntity));
+		AITEST_EQUAL("There should be one entity across the whole system", EntitySubsystem->DebugGetEntityCount(), 1);
+		AITEST_EQUAL("Entity\'s archetype should be the Float one", EntitySubsystem->GetArchetypeForEntity(ReservedEntity), FloatsArchetype);
+		AITEST_EQUAL("The created entity should have been added to the Floats archetype", EntitySubsystem->DebugGetArchetypeEntitiesCount(FloatsArchetype), 1);
+		AITEST_EQUAL("Other archetypes should not get affected", EntitySubsystem->DebugGetArchetypeEntitiesCount(IntsArchetype) + EntitySubsystem->DebugGetArchetypeEntitiesCount(FloatsIntsArchetype), 0);
+		EntitySubsystem->DestroyEntity(ReservedEntity);
+		AITEST_EQUAL("There should not be any entity across the whole system", EntitySubsystem->DebugGetEntityCount(), 0);
+		AITEST_EQUAL("The original archetype should now have no entities", EntitySubsystem->DebugGetArchetypeEntitiesCount(FloatsArchetype), 0);
+
+		return true;
+	}
+};
+IMPLEMENT_AI_INSTANT_TEST(FEntityTest_EntityReservationAndBuilding, "System.Pipe.Entity.EntityReservationAndBuilding");
+
+struct FEntityTest_EntityReservationAndBuildingFromInstances : FEntityTestBase
+{
+	virtual bool InstantTest() override
+	{
+		CA_ASSUME(EntitySubsystem);
+		const FLWEntity ReservedEntity = EntitySubsystem->ReserveEntity();
+		AITEST_TRUE("The reserved entity should be a valid entity", EntitySubsystem->IsEntityValid(ReservedEntity));
+		AITEST_FALSE("The reserved entity should not be a valid entity", EntitySubsystem->IsEntityBuilt(ReservedEntity));
+		EntitySubsystem->BuildEntity(ReservedEntity, MakeArrayView(&InstanceInt, 1));
+		AITEST_TRUE("The reserved entity should not be a valid entity", EntitySubsystem->IsEntityBuilt(ReservedEntity));
+		AITEST_EQUAL("There should be one entity across the whole system", EntitySubsystem->DebugGetEntityCount(), 1);
+		AITEST_EQUAL("Entity\'s archetype should be the Ints one", EntitySubsystem->GetArchetypeForEntity(ReservedEntity), IntsArchetype);
+		AITEST_EQUAL("The created entity should have been added to the Ints archetype", EntitySubsystem->DebugGetArchetypeEntitiesCount(IntsArchetype), 1);
+		AITEST_EQUAL("The entity should have the new component with the correct value set", EntitySubsystem->GetComponentDataChecked<FTestFragment_Int>(ReservedEntity).Value, TestIntValue);
+		EntitySubsystem->DestroyEntity(ReservedEntity);
+		AITEST_EQUAL("There should not be any entity across the whole system", EntitySubsystem->DebugGetEntityCount(), 0);
+		AITEST_EQUAL("The original archetype should now have no entities", EntitySubsystem->DebugGetArchetypeEntitiesCount(FloatsArchetype), 0);
+
+		return true;
+	}
+};
+IMPLEMENT_AI_INSTANT_TEST(FEntityTest_EntityReservationAndBuildingFromInstances, "System.Pipe.Entity.EntityReservationAndBuildingFromInstances");
+
+struct FEntityTest_ReleaseEntity : FEntityTestBase
+{
+	virtual bool InstantTest() override
+	{
+		CA_ASSUME(EntitySubsystem);
+		const FLWEntity ReservedEntity = EntitySubsystem->ReserveEntity();
+		AITEST_TRUE("The reserved entity should be a valid entity", EntitySubsystem->IsEntityValid(ReservedEntity));
+		AITEST_FALSE("The reserved entity should not be a valid entity", EntitySubsystem->IsEntityBuilt(ReservedEntity));
+		AITEST_EQUAL("There should only be one entity across the whole system", EntitySubsystem->DebugGetEntityCount(), 1);
+		AITEST_EQUAL("The entity should not get associated to any archetype", EntitySubsystem->GetArchetypeForEntity(ReservedEntity), FArchetypeHandle());
+		EntitySubsystem->ReleaseReservedEntity(ReservedEntity);
+		AITEST_EQUAL("There should not be any entity across the whole system", EntitySubsystem->DebugGetEntityCount(), 0);
+		return true;
+	}
+};
+IMPLEMENT_AI_INSTANT_TEST(FEntityTest_ReleaseEntity, "System.Pipe.Entity.ReleaseEntity");
+
+struct FEntityTest_ReserveAPreviouslyBuiltEntity : FEntityTestBase
+{
+	virtual bool InstantTest() override
+	{
+		CA_ASSUME(EntitySubsystem);
+		{
+			const FLWEntity Entity = EntitySubsystem->CreateEntity(IntsArchetype);
+			AITEST_EQUAL("The entity should get associated to the right archetype", EntitySubsystem->GetArchetypeForEntity(Entity), IntsArchetype);
+			EntitySubsystem->DestroyEntity(Entity);
+			AITEST_EQUAL("There should not be any entity across the whole system", EntitySubsystem->DebugGetEntityCount(), 0);
+			AITEST_EQUAL("The original archetype should now have no entities", EntitySubsystem->DebugGetArchetypeEntitiesCount(IntsArchetype), 0);
+		}
+		
+		const FLWEntity ReservedEntity = EntitySubsystem->ReserveEntity();
+		AITEST_TRUE("The reserved entity should be a valid entity", EntitySubsystem->IsEntityValid(ReservedEntity));
+		AITEST_FALSE("The reserved entity should not be a valid entity", EntitySubsystem->IsEntityBuilt(ReservedEntity));
+		AITEST_EQUAL("There should only be one entity across the whole system", EntitySubsystem->DebugGetEntityCount(), 1);
+		AITEST_EQUAL("The entity should not get associated to any archetype", EntitySubsystem->GetArchetypeForEntity(ReservedEntity), FArchetypeHandle());
+		EntitySubsystem->BuildEntity(ReservedEntity, FloatsArchetype);
+		AITEST_TRUE("The reserved entity should not be a valid entity", EntitySubsystem->IsEntityBuilt(ReservedEntity));
+		AITEST_EQUAL("The entity should get associated to the right archetype", EntitySubsystem->GetArchetypeForEntity(ReservedEntity), FloatsArchetype);
+		EntitySubsystem->DestroyEntity(ReservedEntity);
+		AITEST_EQUAL("There should not be any entity across the whole system", EntitySubsystem->DebugGetEntityCount(), 0);
+		AITEST_EQUAL("The original archetype should now have no entities", EntitySubsystem->DebugGetArchetypeEntitiesCount(FloatsArchetype), 0);
+		return true;
+	}
+};
+IMPLEMENT_AI_INSTANT_TEST(FEntityTest_ReserveAPreviouslyBuiltEntity, "System.Pipe.Entity.ReserveAPreviouslyBuiltEntity");
+
+#endif // WITH_AGGREGATETICKING_DEBUG
+
+} // FPipeEntityTestTest
+
+PRAGMA_ENABLE_OPTIMIZATION
+
+#undef LOCTEXT_NAMESPACE
diff --git a/Engine/Plugins/Runtime/MassEntity/Source/MassEntityTestSuite/Private/MassEntityTestSuiteModule.cpp b/Engine/Plugins/Runtime/MassEntity/Source/MassEntityTestSuite/Private/MassEntityTestSuiteModule.cpp
new file mode 100644
--- /dev/null
+++ b/Engine/Plugins/Runtime/MassEntity/Source/MassEntityTestSuite/Private/MassEntityTestSuiteModule.cpp
@@ -0,0 +1,13 @@
+// Copyright Epic Games, Inc. All Rights Reserved.
+
+#include "MassEntityTestSuiteModule.h"
+
+#define LOCTEXT_NAMESPACE "PipeTestSuite"
+
+class FPipeTestSuiteModule : public IPipeTestSuiteModule
+{
+};
+
+IMPLEMENT_MODULE(FPipeTestSuiteModule, PipeTestSuite)
+
+#undef LOCTEXT_NAMESPACE
diff --git a/Engine/Plugins/Runtime/MassEntity/Source/MassEntityTestSuite/Private/MassExecutionTest.cpp b/Engine/Plugins/Runtime/MassEntity/Source/MassEntityTestSuite/Private/MassExecutionTest.cpp
new file mode 100644
--- /dev/null
+++ b/Engine/Plugins/Runtime/MassEntity/Source/MassEntityTestSuite/Private/MassExecutionTest.cpp
@@ -0,0 +1,177 @@
+// Copyright Epic Games, Inc. All Rights Reserved.
+
+#include "CoreMinimal.h"
+#include "Engine/World.h"
+#include "MassEntitySystem.h"
+#include "MassEntityTypes.h"
+#include "MassEntityTestTypes.h"
+#include "MassExecutor.h"
+
+#define LOCTEXT_NAMESPACE "PipeTest"
+
+PRAGMA_DISABLE_OPTIMIZATION
+
+//----------------------------------------------------------------------//
+// tests 
+//----------------------------------------------------------------------//
+namespace FPipeExecutionTest
+{
+
+struct FExecution_Setup : FExecutionTestBase
+{
+	virtual bool InstantTest() override
+	{
+		AITEST_NOT_NULL("World needs to exist for the test to be performed", World);
+		AITEST_NOT_NULL("EntitySubsystem needs to be created for the test to be performed", EntitySubsystem);
+
+		return true;
+	}
+};
+IMPLEMENT_AI_INSTANT_TEST(FExecution_Setup, "System.Pipe.Execution.Setup");
+
+
+struct FExecution_EmptyArray : FExecutionTestBase
+{
+	virtual bool InstantTest() override
+	{
+		CA_ASSUME(EntitySubsystem); // if EntitySubsystem null InstantTest won't be called at all
+		const float DeltaSeconds = 0.f;
+		FPipeContext PipeContext(*EntitySubsystem, DeltaSeconds);
+		// no test performed, let's just see if it results in errors/warnings
+		UE::Pipe::Executor::RunProcessorsView(TArrayView<UPipeProcessor*>(), PipeContext);
+		return true;
+	}
+};
+IMPLEMENT_AI_INSTANT_TEST(FExecution_EmptyArray, "System.Pipe.Execution.EmptyArray");
+
+
+struct FExecution_EmptyPipeline : FExecutionTestBase
+{
+	virtual bool InstantTest() override
+	{
+		CA_ASSUME(EntitySubsystem); // if EntitySubsystem null InstantTest won't be called at all
+		const float DeltaSeconds = 0.f;
+		FPipeContext PipeContext(*EntitySubsystem, DeltaSeconds);
+		FRuntimePipeline Pipeline;
+		// no test performed, let's just see if it results in errors/warnings
+		UE::Pipe::Executor::Run(Pipeline, PipeContext);
+		return true;
+	}
+};
+IMPLEMENT_AI_INSTANT_TEST(FExecution_EmptyPipeline, "System.Pipe.Execution.EmptyPipeline");
+
+
+struct FExecution_InvalidPipeContext : FExecutionTestBase
+{
+	virtual bool InstantTest() override
+	{
+		CA_ASSUME(EntitySubsystem);
+		const float DeltaSeconds = 0.f;
+		FPipeContext PipeContext;
+		// test assumption
+		AITEST_NULL("FPipeContext\'s default constructor is expected to set FPipeContext.EntitySubsystem to null", PipeContext.EntitySubsystem);
+		
+		GetTestRunner().AddExpectedError(TEXT("PipeContext.EntitySubsystem is null"), EAutomationExpectedErrorFlags::Contains, 1);		
+		// note that using RunProcessorsView is to bypass reasonable tests UE::Pipe::Executor::Run(Pipeline,...) does that are 
+		// reported via ensures which are not handled by the automation framework
+		UE::Pipe::Executor::RunProcessorsView(TArrayView<UPipeProcessor*>(), PipeContext);
+		return true;
+	}
+};
+IMPLEMENT_AI_INSTANT_TEST(FExecution_InvalidPipeContext, "System.Pipe.Execution.InvalidPipeContext");
+
+
+#if WITH_PIPE_DEBUG
+struct FExecution_SingleNullProcessor : FExecutionTestBase
+{
+	virtual bool InstantTest() override
+	{
+		CA_ASSUME(EntitySubsystem);
+		const float DeltaSeconds = 0.f;
+		FPipeContext PipeContext(*EntitySubsystem, DeltaSeconds);
+		TArray<UPipeProcessor*> Processors;
+		Processors.Add(nullptr);
+		
+
+		GetTestRunner().AddExpectedError(TEXT("Processors contains nullptr"), EAutomationExpectedErrorFlags::Contains, 1);
+		// note that using RunProcessorsView is to bypass reasonable tests UE::Pipe::Executor::Run(Pipeline,...) does that are 
+		// reported via ensures which are not handled by the automation framework
+		UE::Pipe::Executor::RunProcessorsView(Processors, PipeContext);
+		return true;
+	}
+};
+IMPLEMENT_AI_INSTANT_TEST(FExecution_SingleNullProcessor, "System.Pipe.Execution.SingleNullProcessor");
+
+
+struct FExecution_SingleValidProcessor : FExecutionTestBase
+{
+	virtual bool InstantTest() override
+	{
+		CA_ASSUME(EntitySubsystem);
+		const float DeltaSeconds = 0.f;
+		FPipeContext PipeContext(*EntitySubsystem, DeltaSeconds);
+		UPipeTestProcessorBase* Processor = NewObject<UPipeTestProcessorBase>(EntitySubsystem);
+		check(Processor);
+
+		// nothing should break. The actual result of processing is getting tested in PipeProcessorTests.cpp
+		UE::Pipe::Executor::Run(*Processor, PipeContext);
+		return true;
+	}
+};
+IMPLEMENT_AI_INSTANT_TEST(FExecution_SingleValidProcessor, "System.Pipe.Execution.SingleValidProcessor");
+
+
+struct FExecution_MultipleNullProcessors : FExecutionTestBase
+{
+	virtual bool InstantTest() override
+	{
+		CA_ASSUME(EntitySubsystem);
+		const float DeltaSeconds = 0.f;
+		FPipeContext PipeContext(*EntitySubsystem, DeltaSeconds);
+		TArray<UPipeProcessor*> Processors;
+		Processors.Add(nullptr);
+		Processors.Add(nullptr);
+		Processors.Add(nullptr);
+
+		GetTestRunner().AddExpectedError(TEXT("Processors contains nullptr"), EAutomationExpectedErrorFlags::Contains, 1);
+		// note that using RunProcessorsView is to bypass reasonable tests UE::Pipe::Executor::Run(Pipeline,...) does that are 
+		// reported via ensures which are not handled by the automation framework
+		UE::Pipe::Executor::RunProcessorsView(Processors, PipeContext);
+		return true;
+	}
+};
+IMPLEMENT_AI_INSTANT_TEST(FExecution_MultipleNullProcessors, "System.Pipe.Execution.MultipleNullProcessors");
+#endif // WITH_PIPE_DEBUG
+
+
+struct FExecution_Sparse : FEntityTestBase
+{
+	virtual bool InstantTest() override
+	{
+		CA_ASSUME(EntitySubsystem);
+		const float DeltaSeconds = 0.f;
+		FPipeContext PipeContext(*EntitySubsystem, DeltaSeconds);
+		UPipeTestProcessorBase* Processor = NewObject<UPipeTestProcessorBase>(EntitySubsystem);
+		check(Processor);
+
+		FRuntimePipeline Pipeline;
+		{
+			TArray<UPipeProcessor*> Processors;
+			Processors.Add(Processor);
+			Pipeline.SetProcessors(MoveTemp(Processors));
+		}
+
+		FArchetypeChunkCollection ChunkCollection(FloatsArchetype);
+		// nothing should break. The actual result of processing is getting tested in PipeProcessorTests.cpp
+		
+		UE::Pipe::Executor::RunSparse(Pipeline, PipeContext, ChunkCollection);
+
+		return true;
+	}
+};
+IMPLEMENT_AI_INSTANT_TEST(FExecution_Sparse, "System.Pipe.Execution.Sparse");
+} // FPipeExecutionTest
+
+PRAGMA_ENABLE_OPTIMIZATION
+
+#undef LOCTEXT_NAMESPACE
diff --git a/Engine/Plugins/Runtime/MassEntity/Source/MassEntityTestSuite/Private/MassProcessorTest.cpp b/Engine/Plugins/Runtime/MassEntity/Source/MassEntityTestSuite/Private/MassProcessorTest.cpp
new file mode 100644
--- /dev/null
+++ b/Engine/Plugins/Runtime/MassEntity/Source/MassEntityTestSuite/Private/MassProcessorTest.cpp
@@ -0,0 +1,125 @@
+// Copyright Epic Games, Inc. All Rights Reserved.
+
+#include "CoreMinimal.h"
+#include "AITestsCommon.h"
+
+#include "Engine/World.h"
+#include "MassEntitySystem.h"
+#include "MassEntityTypes.h"
+#include "MassEntityTestTypes.h"
+#include "MassExecutor.h"
+
+#define LOCTEXT_NAMESPACE "PipeTest"
+
+PRAGMA_DISABLE_OPTIMIZATION
+
+namespace FPipeProcessorTest
+{
+
+template<typename TProcessor>
+int32 SimpleProcessorRun(UPipeEntitySubsystem& EntitySubsystem)
+{
+	int32 EntityProcessedCount = 0;
+	TProcessor* Processor = NewObject<TProcessor>(&EntitySubsystem);
+	Processor->ExecutionFunction = [Processor, &EntityProcessedCount](UEntitySubsystem& InEntitySubsystem, FLWComponentSystemExecutionContext& Context) {
+		check(Processor);
+		Processor->TestGetQuery().ForEachEntityChunk(InEntitySubsystem, Context, [Processor, &EntityProcessedCount](FLWComponentSystemExecutionContext& Context)
+			{
+				EntityProcessedCount += Context.GetEntitiesNum();
+			});
+	};
+
+	FPipeContext PipeContext(EntitySubsystem, /*DeltaSeconds=*/0.f);
+	UE::Pipe::Executor::Run(*Processor, PipeContext);
+
+	return EntityProcessedCount;
+}
+
+struct FProcessorTest_NoEntities : FEntityTestBase
+{
+	virtual bool InstantTest() override
+	{
+		CA_ASSUME(EntitySubsystem);
+		
+		int32 EntityProcessedCount = SimpleProcessorRun<UPipeTestProcessor_Floats>(*EntitySubsystem);
+		AITEST_EQUAL("No entities have been created yet so the processor shouldn't do any work.", EntityProcessedCount, 0);
+				
+		return true;
+	}
+};
+IMPLEMENT_AI_INSTANT_TEST(FProcessorTest_NoEntities, "System.Pipe.Processor.NoEntities");
+
+
+struct FProcessorTest_NoMatchingEntities : FEntityTestBase
+{
+	virtual bool InstantTest() override
+	{
+		CA_ASSUME(EntitySubsystem);
+
+		TArray<FLWEntity> EntitiesCreated;
+		EntitySubsystem->BatchCreateEntities(IntsArchetype, 100, EntitiesCreated);
+		int32 EntityProcessedCount = SimpleProcessorRun<UPipeTestProcessor_Floats>(*EntitySubsystem);
+		AITEST_EQUAL("No matching entities have been created yet so the processor shouldn't do any work.", EntityProcessedCount, 0);
+
+		return true;
+	}
+};
+IMPLEMENT_AI_INSTANT_TEST(FProcessorTest_NoMatchingEntities, "System.Pipe.Processor.NoMatchingEntities");
+
+struct FProcessorTest_OneMatchingArchetype : FEntityTestBase
+{
+	virtual bool InstantTest() override
+	{
+		CA_ASSUME(EntitySubsystem);
+
+		const int32 EntitiesToCreate = 100;
+		TArray<FLWEntity> EntitiesCreated;
+		EntitySubsystem->BatchCreateEntities(FloatsArchetype, EntitiesToCreate, EntitiesCreated);
+		int32 EntityProcessedCount = SimpleProcessorRun<UPipeTestProcessor_Floats>(*EntitySubsystem);
+		AITEST_EQUAL("No matching entities have been created yet so the processor shouldn't do any work.", EntityProcessedCount, EntitiesToCreate);
+
+		return true;
+	}
+};
+IMPLEMENT_AI_INSTANT_TEST(FProcessorTest_OneMatchingArchetype, "System.Pipe.Processor.OneMatchingArchetype");
+
+
+struct FProcessorTest_MultipleMatchingArchetype : FEntityTestBase
+{
+	virtual bool InstantTest() override
+	{
+		CA_ASSUME(EntitySubsystem);
+
+		const int32 EntitiesToCreate = 100;
+		TArray<FLWEntity> EntitiesCreated;
+		EntitySubsystem->BatchCreateEntities(FloatsArchetype, EntitiesToCreate, EntitiesCreated);
+		EntitySubsystem->BatchCreateEntities(FloatsIntsArchetype, EntitiesToCreate, EntitiesCreated);
+		EntitySubsystem->BatchCreateEntities(IntsArchetype, EntitiesToCreate, EntitiesCreated);
+		// note that only two of these archetypes match 
+		int32 EntityProcessedCount = SimpleProcessorRun<UPipeTestProcessor_Floats>(*EntitySubsystem);
+		AITEST_EQUAL("Two of given three archetypes should match.", EntityProcessedCount, EntitiesToCreate * 2);
+
+		return true;
+	}
+};
+IMPLEMENT_AI_INSTANT_TEST(FProcessorTest_MultipleMatchingArchetype, "System.Pipe.Processor.MultipleMatchingArchetype");
+
+
+struct FProcessorTest_Requirements : FEntityTestBase
+{
+	virtual bool InstantTest() override
+	{
+		// add an array to the configurable ExecutionFunction an
+
+		return true;
+	}
+};
+IMPLEMENT_AI_INSTANT_TEST(FProcessorTest_Requirements, "System.Pipe.Processor.Requirements");
+
+
+} // FPipeProcessorTestTest
+
+PRAGMA_ENABLE_OPTIMIZATION
+
+#undef LOCTEXT_NAMESPACE
+
diff --git a/Engine/Plugins/Runtime/MassEntity/Source/MassEntityTestSuite/Private/MassQueryTest.cpp b/Engine/Plugins/Runtime/MassEntity/Source/MassEntityTestSuite/Private/MassQueryTest.cpp
new file mode 100644
--- /dev/null
+++ b/Engine/Plugins/Runtime/MassEntity/Source/MassEntityTestSuite/Private/MassQueryTest.cpp
@@ -0,0 +1,550 @@
+// Copyright Epic Games, Inc. All Rights Reserved.
+
+#include "CoreMinimal.h"
+#include "AITestsCommon.h"
+
+#include "Engine/World.h"
+#include "MassEntitySystem.h"
+#include "MassEntityTypes.h"
+#include "MassEntityTestTypes.h"
+#include "MassExecutor.h"
+
+#define LOCTEXT_NAMESPACE "PipeTest"
+
+PRAGMA_DISABLE_OPTIMIZATION
+
+/**
+mz@todo:
+	- add a test for requirement mode (read/write, absent, [future] optional)
+	- test Present/Absent modes for tags
+	- test Present/Absent modes for non-tag fragments
+	- read/write tag fragments 
+	- constructs like Query.AddRequirement<FTestFragment_Tag>(ELWComponentAccess::ReadWrite);
+	- constructs like Query.AddRequirement<FTestFragment_Int>(ELWComponentAccess::ReadWrite, Absent);
+*/
+
+
+namespace FPipeQueryTest
+{
+
+struct FQueryTest_ProcessorRequirements : FEntityTestBase
+{
+	virtual bool InstantTest() override
+	{
+		CA_ASSUME(EntitySubsystem);
+
+		UPipeTestProcessor_Floats* Processor = NewObject<UPipeTestProcessor_Floats>(EntitySubsystem);
+		TConstArrayView<FLWComponentRequirement> Requirements = Processor->TestGetQuery().GetRequirements();
+		
+		AITEST_TRUE("Query should have extracted some requirements from the given Processor", Requirements.Num() > 0);
+		AITEST_TRUE("There should be exactly one requirement", Requirements.Num() == 1);
+		AITEST_TRUE("The requirement should be of the Float fragment type", Requirements[0].StructType == FTestFragment_Float::StaticStruct());
+
+		return true;
+	}
+};
+IMPLEMENT_AI_INSTANT_TEST(FQueryTest_ProcessorRequirements, "System.Pipe.Query.ProcessorRequiements");
+
+
+struct FQueryTest_ExplicitRequirements : FEntityTestBase
+{
+	virtual bool InstantTest() override
+	{
+		CA_ASSUME(EntitySubsystem);
+				
+		FLWComponentQuery Query({ FTestFragment_Float::StaticStruct()});
+		TConstArrayView<FLWComponentRequirement> Requirements = Query.GetRequirements();
+
+		AITEST_TRUE("Query should have extracted some requirements from the given Processor", Requirements.Num() > 0);
+		AITEST_TRUE("There should be exactly one requirement", Requirements.Num() == 1);
+		AITEST_TRUE("The requirement should be of the Float fragment type", Requirements[0].StructType == FTestFragment_Float::StaticStruct());
+
+		return true;
+	}
+};
+IMPLEMENT_AI_INSTANT_TEST(FQueryTest_ExplicitRequirements, "System.Pipe.Query.ExplicitRequiements");
+
+
+struct FQueryTest_FragmentViewBinding : FEntityTestBase
+{
+	virtual bool InstantTest() override
+	{
+		CA_ASSUME(EntitySubsystem);
+
+		FLWEntity Entity = EntitySubsystem->CreateEntity(FloatsArchetype);
+		FTestFragment_Float& TestedFragment = EntitySubsystem->GetComponentDataChecked<FTestFragment_Float>(Entity);
+		AITEST_TRUE("Initial value of the component should match expectations", TestedFragment.Value == 0.f);
+
+		UPipeTestProcessor_Floats* Processor = NewObject<UPipeTestProcessor_Floats>(EntitySubsystem);
+		Processor->ExecutionFunction = [Processor](UEntitySubsystem& InEntitySubsystem, FLWComponentSystemExecutionContext& Context) {
+			check(Processor);
+			//FLWComponentQuery Query(*Processor);
+
+			Processor->TestGetQuery().ForEachEntityChunk(InEntitySubsystem, Context, [](FLWComponentSystemExecutionContext& Context)
+				{
+					TArrayView<FTestFragment_Float> Floats = Context.GetMutableComponentView<FTestFragment_Float>();
+
+					for (int32 i = 0; i < Context.GetEntitiesNum(); ++i)
+					{
+						Floats[i].Value = 13.f;
+					}
+				});
+		};
+
+		FPipeContext PipeContext(*EntitySubsystem, /*DeltaSeconds=*/0.f);
+		UE::Pipe::Executor::Run(*Processor, PipeContext);
+
+		AITEST_EQUAL("Fragment value should have changed to the expected value", TestedFragment.Value, 13.f);
+
+		return true;
+	}
+};
+IMPLEMENT_AI_INSTANT_TEST(FQueryTest_FragmentViewBinding, "System.Pipe.Query.FragmentViewBinding");
+
+struct FQueryTest_ExecuteSingleArchetype : FEntityTestBase
+{
+	virtual bool InstantTest() override
+	{
+		CA_ASSUME(EntitySubsystem);
+
+		const int32 NumToCreate = 10;
+		TArray<FLWEntity> EntitiesCreated;
+		EntitySubsystem->BatchCreateEntities(FloatsArchetype, NumToCreate, EntitiesCreated);
+		
+		int TotalProcessed = 0;
+
+		FLWComponentSystemExecutionContext ExecContext;
+		FLWComponentQuery Query({ FTestFragment_Float::StaticStruct() });
+		Query.ForEachEntityChunk(*EntitySubsystem, ExecContext, [&TotalProcessed](FLWComponentSystemExecutionContext& Context)
+			{
+				TotalProcessed += Context.GetEntitiesNum();
+				TArrayView<FTestFragment_Float> Floats = Context.GetMutableComponentView<FTestFragment_Float>();
+				
+				for (int32 i = 0; i < Context.GetEntitiesNum(); ++i)
+				{
+					Floats[i].Value = 13.f;
+				}
+			});
+
+		AITEST_TRUE("The number of entities processed needs to match expectations", TotalProcessed == NumToCreate);
+
+		for (FLWEntity& Entity : EntitiesCreated)
+		{
+			const FTestFragment_Float& TestedFragment = EntitySubsystem->GetComponentDataChecked<FTestFragment_Float>(Entity);
+			AITEST_EQUAL("Every fragment value should have changed to the expected value", TestedFragment.Value, 13.f);
+		}
+
+		return true;
+	}
+};
+IMPLEMENT_AI_INSTANT_TEST(FQueryTest_ExecuteSingleArchetype, "System.Pipe.Query.ExecuteSingleArchetype");
+
+
+struct FQueryTest_ExecuteMultipleArchetypes : FEntityTestBase
+{
+	virtual bool InstantTest() override
+	{
+		CA_ASSUME(EntitySubsystem);
+
+		const int32 FloatsArchetypeCreated = 7;
+		const int32 IntsArchetypeCreated = 11;
+		const int32 FloatsIntsArchetypeCreated = 13;
+		TArray<FLWEntity> EntitiesCreated;
+		EntitySubsystem->BatchCreateEntities(IntsArchetype, IntsArchetypeCreated, EntitiesCreated);
+		// clear to store only the float-related entities
+		EntitiesCreated.Reset();
+		EntitySubsystem->BatchCreateEntities(FloatsArchetype, FloatsArchetypeCreated, EntitiesCreated);
+		EntitySubsystem->BatchCreateEntities(FloatsIntsArchetype, FloatsIntsArchetypeCreated, EntitiesCreated);
+
+		int TotalProcessed = 0;
+		FLWComponentSystemExecutionContext ExecContext;
+		FLWComponentQuery Query({ FTestFragment_Float::StaticStruct() });
+		Query.ForEachEntityChunk(*EntitySubsystem, ExecContext, [&TotalProcessed](FLWComponentSystemExecutionContext& Context)
+			{
+				TotalProcessed += Context.GetEntitiesNum();
+				TArrayView<FTestFragment_Float> Floats = Context.GetMutableComponentView<FTestFragment_Float>();
+
+				for (int32 i = 0; i < Context.GetEntitiesNum(); ++i)
+				{
+					Floats[i].Value = 13.f;
+				}
+			});
+
+		AITEST_TRUE("The number of entities processed needs to match expectations", TotalProcessed == FloatsIntsArchetypeCreated + FloatsArchetypeCreated);
+
+		for (FLWEntity& Entity : EntitiesCreated)
+		{
+			const FTestFragment_Float& TestedFragment = EntitySubsystem->GetComponentDataChecked<FTestFragment_Float>(Entity);
+			AITEST_EQUAL("Every fragment value should have changed to the expected value", TestedFragment.Value, 13.f);
+		}
+
+		return true;
+	}
+};
+
+IMPLEMENT_AI_INSTANT_TEST(FQueryTest_ExecuteMultipleArchetypes, "System.Pipe.Query.ExecuteMultipleArchetypes");
+
+
+struct FQueryTest_ExecuteSparse : FEntityTestBase
+{
+	virtual bool InstantTest() override
+	{
+		CA_ASSUME(EntitySubsystem);
+
+		const int32 NumToCreate = 10;
+		TArray<FLWEntity> AllEntitiesCreated;
+		EntitySubsystem->BatchCreateEntities(FloatsArchetype, NumToCreate, AllEntitiesCreated);
+		
+		TArray<int32> IndicesToProcess = { 1, 2, 3, 6, 7};
+		TArray<FLWEntity> EntitiesToProcess;
+		TArray<FLWEntity> EntitiesToIgnore;
+		for (int32 i = 0; i < AllEntitiesCreated.Num(); ++i)
+		{
+			if (IndicesToProcess.Find(i) != INDEX_NONE)
+			{
+				EntitiesToProcess.Add(AllEntitiesCreated[i]);
+			}
+			else
+			{
+				EntitiesToIgnore.Add(AllEntitiesCreated[i]);
+			}
+		}
+
+
+		int TotalProcessed = 0;
+
+		FLWComponentSystemExecutionContext ExecContext;
+		FLWComponentQuery().AddRequirement<FTestFragment_Float>(ELWComponentAccess::ReadWrite)
+			.ForEachEntityChunk(FArchetypeChunkCollection(FloatsArchetype, EntitiesToProcess), *EntitySubsystem, ExecContext, [&TotalProcessed](FLWComponentSystemExecutionContext& Context)
+			{
+				TotalProcessed += Context.GetEntitiesNum();
+				TArrayView<FTestFragment_Float> Floats = Context.GetMutableComponentView<FTestFragment_Float>();
+
+				for (int32 i = 0; i < Context.GetEntitiesNum(); ++i)
+				{
+					Floats[i].Value = 13.f;
+				}
+			});
+
+		AITEST_TRUE("The number of entities processed needs to match expectations", TotalProcessed == IndicesToProcess.Num());
+
+		for (FLWEntity& Entity : EntitiesToProcess)
+		{
+			const FTestFragment_Float& TestedFragment = EntitySubsystem->GetComponentDataChecked<FTestFragment_Float>(Entity);
+			AITEST_EQUAL("Every fragment value should have changed to the expected value", TestedFragment.Value, 13.f);
+		}
+		
+		for (FLWEntity& Entity : EntitiesToIgnore)
+		{
+			const FTestFragment_Float& TestedFragment = EntitySubsystem->GetComponentDataChecked<FTestFragment_Float>(Entity);
+			AITEST_EQUAL("Untouched entites should retain default fragment value ", TestedFragment.Value, 0.f);
+		}
+
+		return true;
+	}
+};
+IMPLEMENT_AI_INSTANT_TEST(FQueryTest_ExecuteSparse, "System.Pipe.Query.ExecuteSparse");
+
+
+struct FQueryTest_TagPresent : FEntityTestBase
+{
+	virtual bool InstantTest() override
+	{
+		CA_ASSUME(EntitySubsystem);
+
+		TArray<const UScriptStruct*> Fragments = {FTestFragment_Float::StaticStruct(), FTestFragment_Tag::StaticStruct()};
+		const FArchetypeHandle FloatsTagArchetype = EntitySubsystem->CreateArchetype(Fragments);
+
+		FLWComponentQuery Query;
+		Query.AddRequirement<FTestFragment_Float>(ELWComponentAccess::ReadWrite);
+		Query.AddTagRequirement<FTestFragment_Tag>(ELWComponentPresence::All);
+		Query.CacheArchetypes(*EntitySubsystem);
+
+		AITEST_EQUAL("There's a single archetype matching the requirements", Query.GetArchetypes().Num(), 1);
+		AITEST_TRUE("The only valid archetype is FloatsTagArchetype", FloatsTagArchetype == Query.GetArchetypes()[0]);
+
+		return true;
+	}
+};
+IMPLEMENT_AI_INSTANT_TEST(FQueryTest_TagPresent, "System.Pipe.Query.TagPresent");
+
+
+struct FQueryTest_TagAbsent : FEntityTestBase
+{
+	virtual bool InstantTest() override
+	{
+		CA_ASSUME(EntitySubsystem);
+
+		TArray<const UScriptStruct*> Fragments = { FTestFragment_Float::StaticStruct(), FTestFragment_Tag::StaticStruct() };
+		const FArchetypeHandle FloatsTagArchetype = EntitySubsystem->CreateArchetype(Fragments);
+
+		FLWComponentQuery Query;
+		Query.AddRequirement<FTestFragment_Float>(ELWComponentAccess::ReadWrite);
+		Query.AddTagRequirement<FTestFragment_Tag>(ELWComponentPresence::None);
+		Query.CacheArchetypes(*EntitySubsystem);
+
+		AITEST_EQUAL("There are exactly two archetypes matching the requirements", Query.GetArchetypes().Num(), 2);
+		AITEST_TRUE("FloatsTagArchetype is not amongst matching archetypes"
+			, !(FloatsTagArchetype == Query.GetArchetypes()[0] || FloatsTagArchetype == Query.GetArchetypes()[1]));
+
+		return true;
+	}
+};
+IMPLEMENT_AI_INSTANT_TEST(FQueryTest_TagAbsent, "System.Pipe.Query.TagAbsent");
+
+
+/** using a fragment as a tag */
+struct FQueryTest_FragmentPresent : FEntityTestBase
+{
+	virtual bool InstantTest() override
+	{
+		CA_ASSUME(EntitySubsystem);
+
+		FLWComponentQuery Query;
+		Query.AddRequirement<FTestFragment_Int>(ELWComponentAccess::ReadOnly);
+		Query.CacheArchetypes(*EntitySubsystem);
+
+		AITEST_EQUAL("There are exactly two archetypes matching the requirements", Query.GetArchetypes().Num(), 2);
+		AITEST_TRUE("FloatsArchetype is not amongst matching archetypes"
+			, !(FloatsArchetype == Query.GetArchetypes()[0] || FloatsArchetype == Query.GetArchetypes()[1]));
+
+		return true;
+	}
+};
+IMPLEMENT_AI_INSTANT_TEST(FQueryTest_FragmentPresent, "System.Pipe.Query.FragmentPresent");
+
+
+struct FQueryTest_OnlySingleAbsentFragment : FEntityTestBase
+{
+	virtual bool InstantTest() override
+	{
+		CA_ASSUME(EntitySubsystem);
+
+		FLWComponentQuery Query;
+		Query.AddRequirement<FTestFragment_Int>(ELWComponentAccess::ReadOnly, ELWComponentPresence::None);
+		
+		AITEST_FALSE("The query is not valid", Query.CheckValidity());
+
+		GetTestRunner().AddExpectedError(TEXT("requirements not valid"), EAutomationExpectedErrorFlags::Contains, 1);
+		Query.CacheArchetypes(*EntitySubsystem);
+
+		// this is an invalid query. We expect an error and an empty valid archetypes array
+		AITEST_EQUAL("The query is invalid and we expect no archetypes match", Query.GetArchetypes().Num(), 0);
+
+		return true;
+	}
+};
+IMPLEMENT_AI_INSTANT_TEST(FQueryTest_OnlySingleAbsentFragment, "System.Pipe.Query.OnlySingleAbsentFragment");
+
+
+struct FQueryTest_OnlyMultipleAbsentFragments : FEntityTestBase
+{
+	virtual bool InstantTest() override
+	{
+		CA_ASSUME(EntitySubsystem);
+
+		FLWComponentQuery Query;
+		Query.AddRequirement<FTestFragment_Int>(ELWComponentAccess::ReadOnly, ELWComponentPresence::None);
+		Query.AddRequirement<FTestFragment_Float>(ELWComponentAccess::ReadOnly, ELWComponentPresence::None);
+
+		AITEST_FALSE("The query is not valid", Query.CheckValidity());
+
+		GetTestRunner().AddExpectedError(TEXT("requirements not valid"), EAutomationExpectedErrorFlags::Contains, 1);
+		Query.CacheArchetypes(*EntitySubsystem);
+
+		// this is an invalid query. We expect an error and an empty valid archetypes array
+		AITEST_EQUAL("The query is invalid and we expect no archetypes match", Query.GetArchetypes().Num(), 0);
+
+		return true;
+	}
+};
+IMPLEMENT_AI_INSTANT_TEST(FQueryTest_OnlyMultipleAbsentFragments, "System.Pipe.Query.OnlyMultipleAbsentFragments");
+
+
+struct FQueryTest_AbsentAndPresentFragments : FEntityTestBase
+{
+	virtual bool InstantTest() override
+	{
+		CA_ASSUME(EntitySubsystem);
+
+		FLWComponentQuery Query;
+		Query.AddRequirement<FTestFragment_Int>(ELWComponentAccess::None, ELWComponentPresence::None);
+		Query.AddRequirement<FTestFragment_Float>(ELWComponentAccess::ReadOnly, ELWComponentPresence::All);
+
+		AITEST_TRUE("The query is valid", Query.CheckValidity());
+		Query.CacheArchetypes(*EntitySubsystem);
+		AITEST_EQUAL("There is only one archetype matching the query", Query.GetArchetypes().Num(), 1);
+		AITEST_TRUE("FloatsArchetype is the only one matching the query", FloatsArchetype == Query.GetArchetypes()[0]);
+
+		return true;
+	}
+};
+IMPLEMENT_AI_INSTANT_TEST(FQueryTest_AbsentAndPresentFragments, "System.Pipe.Query.AbsentAndPresentFragments");
+
+
+struct FQueryTest_SingleOptionalFragment : FEntityTestBase
+{
+	virtual bool InstantTest() override
+	{
+		CA_ASSUME(EntitySubsystem);
+
+		FLWComponentQuery Query;
+		Query.AddRequirement<FTestFragment_Int>(ELWComponentAccess::ReadWrite, ELWComponentPresence::Optional);
+		Query.CacheArchetypes(*EntitySubsystem);
+
+		AITEST_EQUAL("There are exactly two archetypes matching the requirements", Query.GetArchetypes().Num(), 2);
+		AITEST_TRUE("FloatsArchetype is not amongst matching archetypes"
+			, !(FloatsArchetype == Query.GetArchetypes()[0] || FloatsArchetype == Query.GetArchetypes()[1]));
+
+		return true;
+	}
+};
+IMPLEMENT_AI_INSTANT_TEST(FQueryTest_SingleOptionalFragment, "System.Pipe.Query.SingleOptionalFragment");
+
+
+struct FQueryTest_MultipleOptionalFragment : FEntityTestBase
+{
+	virtual bool InstantTest() override
+	{
+		CA_ASSUME(EntitySubsystem);
+
+		FLWComponentQuery Query;
+		Query.AddRequirement<FTestFragment_Int>(ELWComponentAccess::ReadWrite, ELWComponentPresence::Optional);
+		Query.AddRequirement<FTestFragment_Float>(ELWComponentAccess::ReadWrite, ELWComponentPresence::Optional);
+		Query.CacheArchetypes(*EntitySubsystem);
+
+		AITEST_EQUAL("All three archetype meet requirements", Query.GetArchetypes().Num(), 3);
+		return true;
+	}
+};
+IMPLEMENT_AI_INSTANT_TEST(FQueryTest_MultipleOptionalFragment, "System.Pipe.Query.MultipleOptionalFragment");
+
+
+/** This test configures a query to fetch archetypes that have a Float fragment (we have two of these) with an optional 
+ *  Int fragment (of which we'll have one among the Float ones) */
+struct FQueryTest_UsingOptionalFragment : FEntityTestBase
+{
+	virtual bool InstantTest() override
+	{
+		CA_ASSUME(EntitySubsystem);
+
+		EntitySubsystem->CreateEntity(FloatsArchetype);
+		const FLWEntity EntityWithFloatsInts = EntitySubsystem->CreateEntity(FloatsIntsArchetype);
+		EntitySubsystem->CreateEntity(IntsArchetype);
+
+		const int32 IntValueSet = 123;
+		int TotalProcessed = 0;
+		int EmptyIntsViewCount = 0;
+
+		FLWComponentQuery Query;
+		Query.AddRequirement<FTestFragment_Int>(ELWComponentAccess::ReadWrite, ELWComponentPresence::Optional);
+		Query.AddRequirement<FTestFragment_Float>(ELWComponentAccess::ReadWrite, ELWComponentPresence::All);
+		FLWComponentSystemExecutionContext ExecContext;
+		Query.ForEachEntityChunk(*EntitySubsystem, ExecContext, [&TotalProcessed, &EmptyIntsViewCount, IntValueSet](FLWComponentSystemExecutionContext& Context) {
+			++TotalProcessed;
+			TArrayView<FTestFragment_Int> Ints = Context.GetMutableComponentView<FTestFragment_Int>();
+			if (Ints.Num() == 0)
+			{
+				++EmptyIntsViewCount;
+			}
+			else
+			{
+				for (FTestFragment_Int& IntFragment : Ints)
+				{	
+					IntFragment.Value = IntValueSet;
+				}
+			}
+			});
+
+		AITEST_EQUAL("Two archetypes total should get processed", TotalProcessed, 2);
+		AITEST_EQUAL("Only one of these archetypes should get an empty Ints array view", EmptyIntsViewCount, 1);
+
+		const FTestFragment_Int& TestFragment = EntitySubsystem->GetComponentDataChecked<FTestFragment_Int>(EntityWithFloatsInts);
+		AITEST_TRUE("The optional fragment\'s value should get modified where present", TestFragment.Value == IntValueSet);
+
+		return true;
+	}
+};
+IMPLEMENT_AI_INSTANT_TEST(FQueryTest_UsingOptionalFragment, "System.Pipe.Query.UsingOptionalFragment");
+
+
+struct FQueryTest_AnyFragment : FEntityTestBase
+{
+	virtual bool InstantTest() override
+	{
+		CA_ASSUME(EntitySubsystem);
+
+		// From FEntityTestBase:
+		// FArchetypeHandle FloatsArchetype;
+		// FArchetypeHandle IntsArchetype;
+		// FArchetypeHandle FloatsIntsArchetype;
+		const FArchetypeHandle BoolArchetype = EntitySubsystem->CreateArchetype({ FTestFragment_Bool::StaticStruct() });
+		const FArchetypeHandle BoolFloatArchetype = EntitySubsystem->CreateArchetype({ FTestFragment_Bool::StaticStruct(), FTestFragment_Float::StaticStruct() });
+
+		FLWComponentQuery Query;
+		Query.AddRequirement<FTestFragment_Int>(ELWComponentAccess::ReadWrite, ELWComponentPresence::Any);
+		Query.AddRequirement<FTestFragment_Bool>(ELWComponentAccess::ReadWrite, ELWComponentPresence::Any);
+		// this query should match: 
+		// IntsArchetype, FloatsIntsArchetype, BoolArchetype, BoolFloatArchetype
+		Query.CacheArchetypes(*EntitySubsystem);
+
+		AITEST_EQUAL("Archetypes containing Int or Bool should meet requirements", Query.GetArchetypes().Num(), 4);
+
+		// populate the archetypes so that we can test fragment binding
+		for (auto ArchetypeHandle : Query.GetArchetypes())
+		{
+			EntitySubsystem->CreateEntity(ArchetypeHandle);
+		}
+
+		FLWComponentSystemExecutionContext TestContext;
+		Query.ForEachEntityChunk(*EntitySubsystem, TestContext, [this](FLWComponentSystemExecutionContext& Context)
+			{
+				TArrayView<FTestFragment_Bool> BoolView = Context.GetMutableComponentView<FTestFragment_Bool>();
+				TArrayView<FTestFragment_Int> IntView = Context.GetMutableComponentView<FTestFragment_Int>();
+				
+				GetTestRunner().TestTrue("Every matching archetype needs to host Bool or Int fragments", BoolView.Num() || IntView.Num());
+			});
+
+		return true;
+	}
+};
+IMPLEMENT_AI_INSTANT_TEST(FQueryTest_AnyFragment, "System.Pipe.Query.AnyFragment");
+
+
+struct FQueryTest_AnyTag : FEntityTestBase
+{
+	virtual bool InstantTest() override
+	{
+		CA_ASSUME(EntitySubsystem);
+
+		const FArchetypeHandle ABArchetype = EntitySubsystem->CreateArchetype({ FTestFragment_Int::StaticStruct(), FTestTag_A::StaticStruct(), FTestTag_B::StaticStruct() });
+		const FArchetypeHandle ACArchetype = EntitySubsystem->CreateArchetype({ FTestFragment_Int::StaticStruct(), FTestTag_A::StaticStruct(), FTestTag_C::StaticStruct() });
+		const FArchetypeHandle BCArchetype = EntitySubsystem->CreateArchetype({ FTestFragment_Int::StaticStruct(), FTestTag_B::StaticStruct(), FTestTag_C::StaticStruct() });
+		const FArchetypeHandle BDArchetype = EntitySubsystem->CreateArchetype({ FTestFragment_Int::StaticStruct(), FTestTag_B::StaticStruct(), FTestTag_D::StaticStruct() });
+		const FArchetypeHandle FloatACArchetype = EntitySubsystem->CreateArchetype({ FTestFragment_Float::StaticStruct(), FTestTag_A::StaticStruct(), FTestTag_C::StaticStruct() });
+
+		FLWComponentQuery Query;
+		// at least one fragment requirement needs to be present for the query to be valid
+		Query.AddRequirement<FTestFragment_Int>(ELWComponentAccess::ReadOnly); 
+		Query.AddTagRequirement<FTestTag_A>(ELWComponentPresence::Any);
+		Query.AddTagRequirement<FTestTag_C>(ELWComponentPresence::Any);
+		// this query should match: 
+		// ABArchetype, ACArchetype and ABCrchetype but not BDArchetype nor FEntityTestBase.IntsArchetype
+		Query.CacheArchetypes(*EntitySubsystem);
+
+		AITEST_EQUAL("Only Archetypes tagged with A or C should matched the query", Query.GetArchetypes().Num(), 3);
+		AITEST_TRUE("ABArchetype should be amongst the matched archetypes", Query.GetArchetypes().Find(ABArchetype) != INDEX_NONE);
+		AITEST_TRUE("ACArchetype should be amongst the matched archetypes", Query.GetArchetypes().Find(ACArchetype) != INDEX_NONE);
+		AITEST_TRUE("BCArchetype should be amongst the matched archetypes", Query.GetArchetypes().Find(BCArchetype) != INDEX_NONE);
+
+		return true;
+	}
+};
+IMPLEMENT_AI_INSTANT_TEST(FQueryTest_AnyTag, "System.Pipe.Query.AnyTag");
+
+} // FPipeQueryTest
+
+PRAGMA_ENABLE_OPTIMIZATION
+
+#undef LOCTEXT_NAMESPACE
+
diff --git a/Engine/Plugins/Runtime/MassEntity/Source/MassEntityTestSuite/Private/MassTestTypes.cpp b/Engine/Plugins/Runtime/MassEntity/Source/MassEntityTestSuite/Private/MassTestTypes.cpp
new file mode 100644
--- /dev/null
+++ b/Engine/Plugins/Runtime/MassEntity/Source/MassEntityTestSuite/Private/MassTestTypes.cpp
@@ -0,0 +1,84 @@
+// Copyright Epic Games, Inc. All Rights Reserved.
+
+#include "MassEntityTestTypes.h"
+#include "MassEntitySystem.h"
+#include "MassExecutor.h"
+#include "Engine/World.h"
+
+//----------------------------------------------------------------------//
+// Test bases 
+//----------------------------------------------------------------------//
+bool FExecutionTestBase::SetUp()
+{
+	World = FAITestHelpers::GetWorld();
+	EntitySubsystem = NewObject<UPipeEntitySubsystem>(World);
+	check(EntitySubsystem);
+	struct FSubsystemCollection_TestInit : FSubsystemCollectionBase
+	{
+		FSubsystemCollection_TestInit(){}
+	};
+	FSubsystemCollection_TestInit Collection;
+	EntitySubsystem->Initialize(Collection);
+
+	return true;
+}
+
+bool FEntityTestBase::SetUp()
+{
+	FExecutionTestBase::SetUp();
+	check(EntitySubsystem);
+
+	const UScriptStruct* FragmentTypes[] = { FTestFragment_Float::StaticStruct(), FTestFragment_Int::StaticStruct() };
+
+	EmptyArchetype = EntitySubsystem->CreateArchetype(MakeArrayView<const UScriptStruct*>(nullptr, 0));
+	FloatsArchetype = EntitySubsystem->CreateArchetype(MakeArrayView(&FragmentTypes[0], 1));
+	IntsArchetype = EntitySubsystem->CreateArchetype(MakeArrayView(&FragmentTypes[1], 1));
+	FloatsIntsArchetype = EntitySubsystem->CreateArchetype(MakeArrayView(FragmentTypes, 2));
+
+	FTestFragment_Int IntFrag;
+	IntFrag.Value = TestIntValue;
+	InstanceInt = FInstancedStruct::Make(IntFrag);
+
+	return true;
+}
+
+
+//----------------------------------------------------------------------//
+// Processors 
+//----------------------------------------------------------------------//
+UPipeTestProcessorBase::UPipeTestProcessorBase()
+{
+#if WITH_EDITORONLY_DATA
+	bCanShowUpInSettings = false;
+#endif // WITH_EDITORONLY_DATA
+	bAutoRegisterWithProcessingPhases = false;
+	ExecutionFlags = int32(EProcessorExecutionFlags::All);
+
+	ExecutionFunction = [](UEntitySubsystem& InEntitySubsystem, FLWComponentSystemExecutionContext& Context) {};
+	RequirementsFunction = [](FLWComponentQuery& Query){};
+}
+
+UPipeTestProcessor_Floats::UPipeTestProcessor_Floats()
+{
+	RequirementsFunction = [this](FLWComponentQuery& Query)
+	{
+		EntityQuery.AddRequirement<FTestFragment_Float>(ELWComponentAccess::ReadWrite);
+	};
+}
+
+UPipeTestProcessor_Ints::UPipeTestProcessor_Ints()
+{
+	RequirementsFunction = [this](FLWComponentQuery& Query)
+	{
+		EntityQuery.AddRequirement<FTestFragment_Int>(ELWComponentAccess::ReadWrite);
+	};
+}
+
+UPipeTestProcessor_FloatsInts::UPipeTestProcessor_FloatsInts()
+{
+	RequirementsFunction = [this](FLWComponentQuery& Query)
+	{
+		EntityQuery.AddRequirement<FTestFragment_Float>(ELWComponentAccess::ReadWrite);
+		EntityQuery.AddRequirement<FTestFragment_Int>(ELWComponentAccess::ReadWrite);
+	};
+}
\ No newline at end of file
diff --git a/Engine/Plugins/Runtime/MassEntity/Source/MassEntityTestSuite/Public/LWComponentTestFarmPlot.h b/Engine/Plugins/Runtime/MassEntity/Source/MassEntityTestSuite/Public/LWComponentTestFarmPlot.h
new file mode 100644
--- /dev/null
+++ b/Engine/Plugins/Runtime/MassEntity/Source/MassEntityTestSuite/Public/LWComponentTestFarmPlot.h
@@ -0,0 +1,347 @@
+// Copyright Epic Games, Inc. All Rights Reserved.
+
+#pragma once
+
+#include "CoreMinimal.h"
+#include "GameFramework/Actor.h"
+#include "Engine/DataTable.h"
+
+#include "LWComponentTypes.h"
+#include "EntitySubsystem.h"
+#include "LWCCommandBuffer.h"
+
+#include "LWComponentTestFarmPlot.generated.h"
+
+class UStaticMesh;
+class UMaterialInterface;
+class UHierarchicalInstancedStaticMeshComponent;
+
+USTRUCT()
+struct FFarmVisualDataRow : public FTableRowBase
+{
+	GENERATED_BODY()
+
+	UPROPERTY(EditAnywhere, Category=Farm)
+	UStaticMesh* Mesh = nullptr;
+
+	UPROPERTY(EditAnywhere, Category=Farm)
+	UMaterialInterface* MaterialOverride = nullptr;
+};
+
+USTRUCT()
+struct FFarmJustBecameReadyToHarvest : public FComponentTag
+{
+	GENERATED_BODY()
+};
+
+USTRUCT()
+struct FFarmReadyToHarvest : public FComponentTag
+{
+	GENERATED_BODY()
+};
+
+USTRUCT()
+struct FFarmGridCellData : public FLWComponentData
+{
+	GENERATED_BODY()
+
+	uint16 CellX = 0;
+	uint16 CellY = 0;
+};
+
+USTRUCT()
+struct FFarmWaterComponent : public FLWComponentData
+{
+	GENERATED_BODY()
+
+	float CurrentWater = 1.0f;
+	float DeltaWaterPerSecond = -0.01f;
+};
+
+USTRUCT()
+struct FFarmFlowerComponent : public FLWComponentData
+{
+	GENERATED_BODY()
+		
+	uint32 NumBonusTicks = 0;
+	uint16 FlowerType = 0;
+};
+
+USTRUCT()
+struct FFarmCropComponent : public FLWComponentData
+{
+	GENERATED_BODY()
+
+	uint16 CropType = 0;
+};
+
+
+USTRUCT()
+struct FFarmVisualComponent : public FLWComponentData
+{
+	GENERATED_BODY()
+
+	int32 InstanceIndex = -1;
+	int32 HarvestIconIndex = -1;
+	uint16 VisualType = 0;
+};
+
+USTRUCT()
+struct FHarvestTimerComponent : public FLWComponentData
+{
+	GENERATED_BODY()
+
+	uint32 NumSecondsLeft = 15;
+};
+
+
+UCLASS()
+class UComponentUpdateSystem : public UObject
+{
+	GENERATED_BODY()
+public:
+	virtual void PostInitProperties() override;
+	virtual void Execute(UEntitySubsystem& EntitySubsystem, FLWComponentSystemExecutionContext& Context) {}
+protected:
+	virtual void ConfigureQueries() {}
+protected:
+	FLWComponentQuery EntityQuery;
+};
+
+
+UCLASS()
+class UFarmWaterUpdateSystem : public UComponentUpdateSystem
+{
+	GENERATED_BODY()
+
+public:
+	virtual void ConfigureQueries() override
+	{
+		EntityQuery.AddRequirement<FFarmWaterComponent>(ELWComponentAccess::ReadWrite);
+	}
+
+	virtual void Execute(UEntitySubsystem& EntitySubsystem, FLWComponentSystemExecutionContext& Context) override
+	{
+		QUICK_SCOPE_CYCLE_COUNTER(UFarmWaterUpdateSystem_Run);
+		EntityQuery.ForEachEntityChunk(EntitySubsystem, Context, [this](FLWComponentSystemExecutionContext& Context) {
+
+			TArrayView<FFarmWaterComponent> WaterList = Context.GetMutableComponentView<FFarmWaterComponent>();
+
+			for (FFarmWaterComponent& WaterComponent : WaterList)
+			{
+				WaterComponent.CurrentWater = FMath::Clamp(WaterComponent.CurrentWater + WaterComponent.DeltaWaterPerSecond * 0.01f, 0.0f, 1.0f);
+			}
+		});
+	}
+};
+
+
+UCLASS()
+class UFarmHarvestTimerSystem_Flowers : public UComponentUpdateSystem
+{
+	GENERATED_BODY()
+
+public:
+	virtual void ConfigureQueries() override
+	{
+		EntityQuery.AddRequirement<FHarvestTimerComponent>(ELWComponentAccess::ReadWrite);
+		EntityQuery.AddRequirement<FFarmWaterComponent>(ELWComponentAccess::ReadOnly);
+		EntityQuery.AddRequirement<FFarmFlowerComponent>(ELWComponentAccess::ReadWrite);
+	}
+
+	virtual void Execute(UEntitySubsystem& EntitySubsystem, FLWComponentSystemExecutionContext& Context) override
+	{
+		QUICK_SCOPE_CYCLE_COUNTER(UFarmHarvestTimerSystem_Flowers_Run);
+		
+		EntityQuery.ForEachEntityChunk(EntitySubsystem, Context, [this](FLWComponentSystemExecutionContext& Context) {
+			const int32 NumEntities = Context.GetEntitiesNum();
+			const float WellWateredThreshold = 0.25f;
+			TArrayView<FHarvestTimerComponent> TimerList = Context.GetMutableComponentView<FHarvestTimerComponent>();
+			TConstArrayView<FFarmWaterComponent> WaterList = Context.GetComponentView<FFarmWaterComponent>();
+			TArrayView<FFarmFlowerComponent> FlowerList = Context.GetMutableComponentView<FFarmFlowerComponent>();
+			
+			for (int32 i = 0; i < NumEntities; ++i)
+			{
+				if (TimerList[i].NumSecondsLeft > 0)
+				{
+					--TimerList[i].NumSecondsLeft;
+
+					if (WaterList[i].CurrentWater > WellWateredThreshold)
+					{
+						++FlowerList[i].NumBonusTicks;
+					}
+				}
+			}
+		});
+	}
+};
+
+
+
+
+UCLASS()
+class UFarmHarvestTimerSystem_Crops : public UComponentUpdateSystem
+{
+	GENERATED_BODY()
+
+public:
+	virtual void ConfigureQueries() override
+	{
+		EntityQuery.AddRequirement<FHarvestTimerComponent>(ELWComponentAccess::ReadWrite);
+		EntityQuery.AddRequirement<FFarmWaterComponent>(ELWComponentAccess::ReadOnly);
+		EntityQuery.AddRequirement<FFarmCropComponent>(ELWComponentAccess::ReadOnly);
+	}
+
+	virtual void Execute(UEntitySubsystem& EntitySubsystem, FLWComponentSystemExecutionContext& Context) override
+	{
+		QUICK_SCOPE_CYCLE_COUNTER(UFarmHarvestTimerSystem_Crops_Run);
+		
+		EntityQuery.ForEachEntityChunk(EntitySubsystem, Context, [this](FLWComponentSystemExecutionContext& Context) {
+			
+			const int32 NumEntities = Context.GetEntitiesNum();
+			const float WellWateredThreshold = 0.25f;
+			TArrayView<FHarvestTimerComponent> TimerList = Context.GetMutableComponentView<FHarvestTimerComponent>();
+			TConstArrayView<FFarmWaterComponent> WaterList = Context.GetMutableComponentView<FFarmWaterComponent>();
+				
+			for (int32 i = 0; i < NumEntities; ++i)
+			{
+				const uint32 TimeToSubtract = (WaterList[i].CurrentWater > WellWateredThreshold) ? 2 : 1;
+				TimerList[i].NumSecondsLeft = (TimerList[i].NumSecondsLeft >= TimeToSubtract) ? (TimerList[i].NumSecondsLeft - TimeToSubtract) : 0;
+			}
+		});
+	}
+};
+
+//////////////////////////////////////////////////////////////////////
+
+UCLASS()
+class UFarmHarvestTimerExpired : public UComponentUpdateSystem
+{
+	GENERATED_BODY()
+
+public:
+	virtual void ConfigureQueries() override
+	{
+		EntityQuery.AddRequirement<FHarvestTimerComponent>(ELWComponentAccess::ReadOnly);
+		EntityQuery.AddTagRequirement<FFarmJustBecameReadyToHarvest>(ELWComponentPresence::None);
+		EntityQuery.AddTagRequirement<FFarmReadyToHarvest>(ELWComponentPresence::None);
+	}
+
+	virtual void Execute(UEntitySubsystem& EntitySubsystem, FLWComponentSystemExecutionContext& Context) override
+	{
+		QUICK_SCOPE_CYCLE_COUNTER(UFarmHarvestTimerExpired_Run);
+		
+		EntityQuery.ForEachEntityChunk(EntitySubsystem, Context, [this](FLWComponentSystemExecutionContext& Context) {
+			const int32 NumEntities = Context.GetEntitiesNum();
+			TConstArrayView<FHarvestTimerComponent> TimerList = Context.GetComponentView<FHarvestTimerComponent>();
+
+			for (int32 i = 0; i < NumEntities; ++i)
+			{
+				if (TimerList[i].NumSecondsLeft == 0)
+				{
+					Context.Defer().AddTag<FFarmJustBecameReadyToHarvest>(Context.GetEntity(i));
+				}
+			}
+		});
+	}
+};
+
+
+//////////////////////////////////////////////////////////////////////
+
+UCLASS()
+class UFarmHarvestTimerSetIcon : public UComponentUpdateSystem
+{
+	GENERATED_BODY()
+
+public:
+	UHierarchicalInstancedStaticMeshComponent* HarvestIconISMC;
+	float GridCellWidth;
+	float GridCellHeight;
+	float HarvestIconHeight;
+	float HarvestIconScale;
+
+public:
+	virtual void ConfigureQueries() override
+	{
+		EntityQuery.AddRequirement<FFarmGridCellData>(ELWComponentAccess::ReadOnly);
+		EntityQuery.AddRequirement<FFarmVisualComponent>(ELWComponentAccess::ReadWrite);
+		EntityQuery.AddTagRequirement<FFarmJustBecameReadyToHarvest>(ELWComponentPresence::All);
+	}
+
+	virtual void Execute(UEntitySubsystem& EntitySubsystem, FLWComponentSystemExecutionContext& Context) override;
+};
+
+//////////////////////////////////////////////////////////////////////
+
+
+UCLASS(config=Game)
+class ALWComponentTestFarmPlot : public AActor
+{
+	GENERATED_BODY()
+
+public:
+	int32 GridWidth = 40*7;
+	int32 GridHeight = 20*7;
+
+	UPROPERTY(EditAnywhere, Category=Farm)
+	float GridCellWidth = 150.0f;
+	
+	UPROPERTY(EditAnywhere, Category=Farm)
+	float GridCellHeight = 150.0f;
+	
+	UPROPERTY(EditAnywhere, Category=Farm)
+	float HarvestIconScale = 0.3f;
+	
+	TArray<FLWEntity> PlantedSquares;
+
+	
+	UPROPERTY(EditDefaultsOnly, Category=Farm)
+	TArray<FFarmVisualDataRow> VisualDataTable;
+
+	UPROPERTY(Transient)
+	TArray<UHierarchicalInstancedStaticMeshComponent*> VisualDataISMCs;
+
+	float NextSecondTimer = 0.0f;
+
+	UPROPERTY(Transient)
+	TArray<UComponentUpdateSystem*> PerFrameSystems;
+
+	UPROPERTY(Transient)
+	TArray<UComponentUpdateSystem*> PerSecondSystems;
+
+	// Indicies into VisualDataTable for flowers
+	UPROPERTY(EditAnywhere, Category=Farm)
+	TArray<int32> TestDataFlowerIndicies;
+
+	// Indicies into VisualDataTable for crops
+	UPROPERTY(EditAnywhere, Category=Farm)
+	TArray<int32> TestDataCropIndicies;
+
+	UPROPERTY(EditAnywhere, Category = Farm)
+	float VisualNearCullDistance = 1000.0f;
+
+	UPROPERTY(EditAnywhere, Category = Farm)
+	float VisualFarCullDistance = 1200.0f;
+
+	UPROPERTY(EditAnywhere, Category = Farm)
+	float IconNearCullDistance = 400.0f;
+
+	UPROPERTY(EditAnywhere, Category = Farm)
+	float IconFarCullDistance = 800.0f;
+
+private:
+	UPROPERTY(VisibleAnywhere, BlueprintReadOnly, meta=(AllowPrivateAccess = "true"), Category=Farm)
+	UHierarchicalInstancedStaticMeshComponent* HarvestIconISMC;
+
+	TArray<int32> FreeHarvestIconIndicies;
+
+private:
+	void AddItemToGrid(UEntitySubsystem* EntitySystem, int32 X, int32 Y, FArchetypeHandle Archetype, int32 VisualIndex);
+
+public:
+	ALWComponentTestFarmPlot();
+
+	virtual void BeginPlay() override;
+	virtual void TickActor(float DeltaTime, enum ELevelTick TickType, FActorTickFunction& ThisTickFunction) override;
+};
diff --git a/Engine/Plugins/Runtime/MassEntity/Source/MassEntityTestSuite/Public/MassEntityTestSuiteModule.h b/Engine/Plugins/Runtime/MassEntity/Source/MassEntityTestSuite/Public/MassEntityTestSuiteModule.h
new file mode 100644
--- /dev/null
+++ b/Engine/Plugins/Runtime/MassEntity/Source/MassEntityTestSuite/Public/MassEntityTestSuiteModule.h
@@ -0,0 +1,37 @@
+// Copyright Epic Games, Inc. All Rights Reserved.
+
+#pragma once
+
+#include "CoreMinimal.h"
+#include "Modules/ModuleManager.h"
+
+/**
+* The public interface to this module
+*/
+class IPipeTestSuiteModule : public IModuleInterface
+{
+
+public:
+
+	/**
+	* Singleton-like access to this module's interface.  This is just for convenience!
+	* Beware of calling this during the shutdown phase, though.  Your module might have been unloaded already.
+	*
+	* @return Returns singleton instance, loading the module on demand if needed
+	*/
+	static inline IPipeTestSuiteModule& Get()
+	{
+		return FModuleManager::LoadModuleChecked<IPipeTestSuiteModule>("PipeTestSuite");
+	}
+
+	/**
+	* Checks to see if this module is loaded and ready.  It is only valid to call Get() if IsAvailable() returns true.
+	*
+	* @return True if the module is loaded and ready to use
+	*/
+	static inline bool IsAvailable()
+	{
+		return FModuleManager::Get().IsModuleLoaded("PipeTestSuite");
+	}
+};
+
