Removed the PipeEntitySystem and renamed EntitySubsystem to MassEntitySubsystem.
It also involves using the MassEntitySubsystem instance created as the WorldSubsystem. MassSimulationSubsystem still holds a pointer to an entity subsystem instance, just gets it from the world.
Updated all of the code to use the main MassEntitySubsystem.

diff --git a/Engine/Plugins/AI/MassAI/Source/MassAIBehavior/Private/MassLookAtProcessors.cpp b/Engine/Plugins/AI/MassAI/Source/MassAIBehavior/Private/MassLookAtProcessors.cpp
--- a/Engine/Plugins/AI/MassAI/Source/MassAIBehavior/Private/MassLookAtProcessors.cpp
+++ b/Engine/Plugins/AI/MassAI/Source/MassAIBehavior/Private/MassLookAtProcessors.cpp
@@ -117,7 +117,7 @@ void UMassLookAtProcessor::ConfigureQueries()
 	EntityQuery_Conditional.SetChunkFilter(&FMassVisualizationChunkFragment::AreAnyEntitiesVisibleInChunk);
 }
 
-void UMassLookAtProcessor::Execute(UEntitySubsystem& EntitySubsystem, FLWComponentSystemExecutionContext& Context)
+void UMassLookAtProcessor::Execute(UMassEntitySubsystem& EntitySubsystem, FLWComponentSystemExecutionContext& Context)
 {
 	QUICK_SCOPE_CYCLE_COUNTER(LookAtProcessor_Run);
 
@@ -233,7 +233,7 @@ void UMassLookAtProcessor::Execute(UEntitySubsystem& EntitySubsystem, FLWCompone
 		});
 }
 
-void UMassLookAtProcessor::FindNewGazeTarget(const UEntitySubsystem& EntitySubsystem, const float CurrentTime, const FTransform& Transform, FMassLookAtFragment& LookAt) const
+void UMassLookAtProcessor::FindNewGazeTarget(const UMassEntitySubsystem& EntitySubsystem, const float CurrentTime, const FTransform& Transform, FMassLookAtFragment& LookAt) const
 {
 	const FLWEntity LastTrackedEntity = LookAt.GazeTrackedEntity;
 	
@@ -349,7 +349,7 @@ void UMassLookAtProcessor::UpdateLookAtTrajectory(const FTransform& Transform, c
 	}
 }
 
-void UMassLookAtProcessor::UpdateLookAtTrackedEntity(const UEntitySubsystem& EntitySubsystem, const FTransform& Transform, const bool bDisplayDebug, FMassLookAtFragment& LookAt) const
+void UMassLookAtProcessor::UpdateLookAtTrackedEntity(const UMassEntitySubsystem& EntitySubsystem, const FTransform& Transform, const bool bDisplayDebug, FMassLookAtFragment& LookAt) const
 {
 	// Update direction toward target
 	if (EntitySubsystem.IsEntityValid(LookAt.TrackedEntity))
@@ -371,7 +371,7 @@ void UMassLookAtProcessor::UpdateLookAtTrackedEntity(const UEntitySubsystem& Ent
 	}
 }
 
-bool UMassLookAtProcessor::UpdateGazeTrackedEntity(const UEntitySubsystem& EntitySubsystem, const FTransform& Transform, const bool bDisplayDebug, FMassLookAtFragment& LookAt) const
+bool UMassLookAtProcessor::UpdateGazeTrackedEntity(const UMassEntitySubsystem& EntitySubsystem, const FTransform& Transform, const bool bDisplayDebug, FMassLookAtFragment& LookAt) const
 {
 	bool bHasTarget = false;
 	
diff --git a/Engine/Plugins/AI/MassAI/Source/MassAIBehavior/Private/MassStateTreeExecutionContext.cpp b/Engine/Plugins/AI/MassAI/Source/MassAIBehavior/Private/MassStateTreeExecutionContext.cpp
--- a/Engine/Plugins/AI/MassAI/Source/MassAIBehavior/Private/MassStateTreeExecutionContext.cpp
+++ b/Engine/Plugins/AI/MassAI/Source/MassAIBehavior/Private/MassStateTreeExecutionContext.cpp
@@ -2,11 +2,11 @@
 
 #include "MassStateTreeExecutionContext.h"
 #include "MassStateTreeTypes.h"
-#include "EntitySubsystem.h"
+#include "MassEntitySubsystem.h"
 #include "MassSignalSubsystem.h"
 #include "Engine/World.h"
 
-FMassStateTreeExecutionContext::FMassStateTreeExecutionContext(UEntitySubsystem& InEntitySubsystem,
+FMassStateTreeExecutionContext::FMassStateTreeExecutionContext(UMassEntitySubsystem& InEntitySubsystem,
                                                                FLWComponentSystemExecutionContext& InContext):
 	EntitySubsystem(&InEntitySubsystem), EntitySubsystemExecutionContext(&InContext)
 {
diff --git a/Engine/Plugins/AI/MassAI/Source/MassAIBehavior/Private/MassStateTreeProcessors.cpp b/Engine/Plugins/AI/MassAI/Source/MassAIBehavior/Private/MassStateTreeProcessors.cpp
--- a/Engine/Plugins/AI/MassAI/Source/MassAIBehavior/Private/MassStateTreeProcessors.cpp
+++ b/Engine/Plugins/AI/MassAI/Source/MassAIBehavior/Private/MassStateTreeProcessors.cpp
@@ -18,7 +18,7 @@ CSV_DEFINE_CATEGORY(StateTreeProcessor, true);
 namespace UE::MassBehavior
 {
 
-bool SetExternalFragments(FMassStateTreeExecutionContext& Context, const UEntitySubsystem& EntitySubsystem)
+bool SetExternalFragments(FMassStateTreeExecutionContext& Context, const UMassEntitySubsystem& EntitySubsystem)
 {
 	bool bFoundAllFragments = true;
 	const FEntityView EntityView(EntitySubsystem, Context.GetEntity());
@@ -152,7 +152,7 @@ void UMassStateTreeFragmentInitializer::ConfigureQueries()
 	EntityQuery.AddRequirement<FMassStateTreeFragment>(ELWComponentAccess::ReadOnly);
 }
 
-void UMassStateTreeFragmentInitializer::Execute(UEntitySubsystem& EntitySubsystem, FLWComponentSystemExecutionContext& Context)
+void UMassStateTreeFragmentInitializer::Execute(UMassEntitySubsystem& EntitySubsystem, FLWComponentSystemExecutionContext& Context)
 {
 	// StateTree processor relies on signals to be ticked but we need an 'initial tick' to set the tree in the proper state.
 	// The initializer provides that by sending a signal to all new entities that use StateTree.
@@ -202,7 +202,7 @@ void UMassStateTreeFragmentDestructor::ConfigureQueries()
 	EntityQuery.AddRequirement<FMassStateTreeFragment>(ELWComponentAccess::ReadOnly);
 }
 
-void UMassStateTreeFragmentDestructor::Execute(UEntitySubsystem& EntitySubsystem, FLWComponentSystemExecutionContext& Context)
+void UMassStateTreeFragmentDestructor::Execute(UMassEntitySubsystem& EntitySubsystem, FLWComponentSystemExecutionContext& Context)
 {
 	FMassStateTreeExecutionContext StateTreeContext(EntitySubsystem, Context);
 	UMassStateTreeSubsystem* MassStateTreeSubsystem = UWorld::GetSubsystem<UMassStateTreeSubsystem>(EntitySubsystem.GetWorld());
@@ -273,7 +273,7 @@ void UMassStateTreeProcessor::ConfigureQueries()
 	EntityQuery.AddRequirement<FMassStateTreeFragment>(ELWComponentAccess::ReadWrite);
 }
 
-void UMassStateTreeProcessor::SignalEntities(UEntitySubsystem& EntitySubsystem, FLWComponentSystemExecutionContext& Context, FMassSignalNameLookup& EntitySignals)
+void UMassStateTreeProcessor::SignalEntities(UMassEntitySubsystem& EntitySubsystem, FLWComponentSystemExecutionContext& Context, FMassSignalNameLookup& EntitySignals)
 {
 	if (!MassStateTreeSubsystem)
 	{
diff --git a/Engine/Plugins/AI/MassAI/Source/MassAIBehavior/Private/MassZoneGraphAnnotationFragments.cpp b/Engine/Plugins/AI/MassAI/Source/MassAIBehavior/Private/MassZoneGraphAnnotationFragments.cpp
--- a/Engine/Plugins/AI/MassAI/Source/MassAIBehavior/Private/MassZoneGraphAnnotationFragments.cpp
+++ b/Engine/Plugins/AI/MassAI/Source/MassAIBehavior/Private/MassZoneGraphAnnotationFragments.cpp
@@ -1,7 +1,7 @@
 // Copyright Epic Games, Inc. All Rights Reserved.
 
 #include "MassZoneGraphAnnotationFragments.h"
-#include "EntitySubsystem.h"
+#include "MassEntitySubsystem.h"
 #include "MassSimulationLOD.h"
 
 namespace UE::Mass::ZoneGraphAnnotations
diff --git a/Engine/Plugins/AI/MassAI/Source/MassAIBehavior/Private/MassZoneGraphAnnotationProcessors.cpp b/Engine/Plugins/AI/MassAI/Source/MassAIBehavior/Private/MassZoneGraphAnnotationProcessors.cpp
--- a/Engine/Plugins/AI/MassAI/Source/MassAIBehavior/Private/MassZoneGraphAnnotationProcessors.cpp
+++ b/Engine/Plugins/AI/MassAI/Source/MassAIBehavior/Private/MassZoneGraphAnnotationProcessors.cpp
@@ -31,7 +31,7 @@ void UMassZoneGraphAnnotationTagsInitializer::ConfigureQueries()
 	EntityQuery.AddRequirement<FMassZoneGraphLaneLocationFragment>(ELWComponentAccess::ReadOnly);
 }
 
-void UMassZoneGraphAnnotationTagsInitializer::Execute(UEntitySubsystem& EntitySubsystem, FLWComponentSystemExecutionContext& Context)
+void UMassZoneGraphAnnotationTagsInitializer::Execute(UMassEntitySubsystem& EntitySubsystem, FLWComponentSystemExecutionContext& Context)
 {
 	EntityQuery.ForEachEntityChunk(EntitySubsystem, Context, [this](FLWComponentSystemExecutionContext& Context)
 	{
@@ -83,7 +83,7 @@ void UMassZoneGraphAnnotationTagUpdateProcessor::ConfigureQueries()
 	EntityQuery.AddChunkRequirement<FMassSimulationVariableTickChunkFragment>(ELWComponentAccess::ReadOnly);
 }
 
-void UMassZoneGraphAnnotationTagUpdateProcessor::Execute(UEntitySubsystem& EntitySubsystem, FLWComponentSystemExecutionContext& Context)
+void UMassZoneGraphAnnotationTagUpdateProcessor::Execute(UMassEntitySubsystem& EntitySubsystem, FLWComponentSystemExecutionContext& Context)
 {
 	TransientEntitiesToSignal.Reset();
 
@@ -136,7 +136,7 @@ void UMassZoneGraphAnnotationTagUpdateProcessor::UpdateAnnotationTags(FMassZoneG
 	}
 }
 
-void UMassZoneGraphAnnotationTagUpdateProcessor::SignalEntities(UEntitySubsystem& EntitySubsystem, FLWComponentSystemExecutionContext& Context, FMassSignalNameLookup& EntitySignals)
+void UMassZoneGraphAnnotationTagUpdateProcessor::SignalEntities(UMassEntitySubsystem& EntitySubsystem, FLWComponentSystemExecutionContext& Context, FMassSignalNameLookup& EntitySignals)
 {
 	if (!ZoneGraphAnnotationSubsystem)
 	{
diff --git a/Engine/Plugins/AI/MassAI/Source/MassAIBehavior/Public/MassLookAtProcessors.h b/Engine/Plugins/AI/MassAI/Source/MassAIBehavior/Public/MassLookAtProcessors.h
--- a/Engine/Plugins/AI/MassAI/Source/MassAIBehavior/Public/MassLookAtProcessors.h
+++ b/Engine/Plugins/AI/MassAI/Source/MassAIBehavior/Public/MassLookAtProcessors.h
@@ -29,20 +29,20 @@ protected:
 
 	virtual void ConfigureQueries() override;
 	virtual void Initialize(UObject& Owner) override;
-	virtual void Execute(UEntitySubsystem& EntitySubsystem, FLWComponentSystemExecutionContext& Context) override;
+	virtual void Execute(UMassEntitySubsystem& EntitySubsystem, FLWComponentSystemExecutionContext& Context) override;
 
 	/** Selects a nearby target if possible or use a random fixed direction */
-	void FindNewGazeTarget(const UEntitySubsystem& EntitySubsystem, const float CurrentTime, const FTransform& Transform, FMassLookAtFragment& LookAt) const;
+	void FindNewGazeTarget(const UMassEntitySubsystem& EntitySubsystem, const float CurrentTime, const FTransform& Transform, FMassLookAtFragment& LookAt) const;
 
 	/** Updates look direction based on look at trajectory. */
 	void UpdateLookAtTrajectory(const FTransform& Transform, const FMassZoneGraphLaneLocationFragment& ZoneGraphLocation,
 								const FMassLookAtTrajectoryFragment& LookAtTrajectory, const bool bDisplayDebug, FMassLookAtFragment& LookAt) const;
 
 	/** Updates look at based on tracked entity. */
-	void UpdateLookAtTrackedEntity(const UEntitySubsystem& EntitySubsystem, const FTransform& Transform, const bool bDisplayDebug, FMassLookAtFragment& LookAt) const;
+	void UpdateLookAtTrackedEntity(const UMassEntitySubsystem& EntitySubsystem, const FTransform& Transform, const bool bDisplayDebug, FMassLookAtFragment& LookAt) const;
 
 	/** Updates gaze based on tracked entity. */
-	bool UpdateGazeTrackedEntity(const UEntitySubsystem& EntitySubsystem, const FTransform& Transform, const bool bDisplayDebug, FMassLookAtFragment& LookAt) const;
+	bool UpdateGazeTrackedEntity(const UMassEntitySubsystem& EntitySubsystem, const FTransform& Transform, const bool bDisplayDebug, FMassLookAtFragment& LookAt) const;
 
 	/** Builds look at trajectory along the current path. */
 	void BuildTrajectory(const FMassZoneGraphLaneLocationFragment& LaneLocation, const FMassZoneGraphShortPathFragment& ShortPath,
diff --git a/Engine/Plugins/AI/MassAI/Source/MassAIBehavior/Public/MassStateTreeExecutionContext.h b/Engine/Plugins/AI/MassAI/Source/MassAIBehavior/Public/MassStateTreeExecutionContext.h
--- a/Engine/Plugins/AI/MassAI/Source/MassAIBehavior/Public/MassStateTreeExecutionContext.h
+++ b/Engine/Plugins/AI/MassAI/Source/MassAIBehavior/Public/MassStateTreeExecutionContext.h
@@ -7,7 +7,7 @@
 #include "MassStateTreeExecutionContext.generated.h"
 
 struct FLWComponentSystemExecutionContext;
-class UEntitySubsystem;
+class UMassEntitySubsystem;
 class UMassSignalSubsystem;
 
 /**
@@ -20,9 +20,9 @@ struct MASSAIBEHAVIOR_API FMassStateTreeExecutionContext : public FStateTreeExec
 public:
 	/** Should never be used but has to be public for 'void UScriptStruct::TCppStructOps<FMassStateTreeExecutionContext>::ConstructForTests(void *)' */
 	FMassStateTreeExecutionContext() = default;
-	FMassStateTreeExecutionContext(UEntitySubsystem& InEntitySubsystem, FLWComponentSystemExecutionContext& InContext);
+	FMassStateTreeExecutionContext(UMassEntitySubsystem& InEntitySubsystem, FLWComponentSystemExecutionContext& InContext);
 
-	UEntitySubsystem& GetEntitySubsystem() const { return *EntitySubsystem; }
+	UMassEntitySubsystem& GetEntitySubsystem() const { return *EntitySubsystem; }
 	FLWComponentSystemExecutionContext& GetEntitySubsystemExecutionContext() const { return *EntitySubsystemExecutionContext; }
 
 	FLWEntity GetEntity() const { return Entity; }
@@ -38,7 +38,7 @@ protected:
 
 	virtual void BeginGatedTransition(const FStateTreeExecutionState& Exec) override;
 
-	UEntitySubsystem* EntitySubsystem = nullptr;
+	UMassEntitySubsystem* EntitySubsystem = nullptr;
 	UMassSignalSubsystem* SignalSubsystem = nullptr;
 	FLWComponentSystemExecutionContext* EntitySubsystemExecutionContext = nullptr;
 	FLWEntity Entity;
diff --git a/Engine/Plugins/AI/MassAI/Source/MassAIBehavior/Public/MassStateTreeProcessors.h b/Engine/Plugins/AI/MassAI/Source/MassAIBehavior/Public/MassStateTreeProcessors.h
--- a/Engine/Plugins/AI/MassAI/Source/MassAIBehavior/Public/MassStateTreeProcessors.h
+++ b/Engine/Plugins/AI/MassAI/Source/MassAIBehavior/Public/MassStateTreeProcessors.h
@@ -22,7 +22,7 @@ public:
 
 protected:
 	virtual void ConfigureQueries() override;
-	virtual void Execute(UEntitySubsystem& EntitySubsystem, FLWComponentSystemExecutionContext& Context) override;
+	virtual void Execute(UMassEntitySubsystem& EntitySubsystem, FLWComponentSystemExecutionContext& Context) override;
 
 	FLWComponentQuery EntityQuery;
 };
@@ -37,7 +37,7 @@ public:
 
 protected:
 	virtual void ConfigureQueries() override;
-	virtual void Execute(UEntitySubsystem& EntitySubsystem, FLWComponentSystemExecutionContext& Context) override;
+	virtual void Execute(UMassEntitySubsystem& EntitySubsystem, FLWComponentSystemExecutionContext& Context) override;
 
 	FLWComponentQuery EntityQuery;
 };
@@ -56,7 +56,7 @@ public:
 protected:
 	virtual void Initialize(UObject& Owner) override;
 	virtual void ConfigureQueries() override;
-	virtual void SignalEntities(UEntitySubsystem& EntitySubsystem, FLWComponentSystemExecutionContext& Context, FMassSignalNameLookup& EntitySignals) override;
+	virtual void SignalEntities(UMassEntitySubsystem& EntitySubsystem, FLWComponentSystemExecutionContext& Context, FMassSignalNameLookup& EntitySignals) override;
 
 	UPROPERTY(Transient)
 	UMassStateTreeSubsystem* MassStateTreeSubsystem;
diff --git a/Engine/Plugins/AI/MassAI/Source/MassAIBehavior/Public/MassZoneGraphAnnotationProcessors.h b/Engine/Plugins/AI/MassAI/Source/MassAIBehavior/Public/MassZoneGraphAnnotationProcessors.h
--- a/Engine/Plugins/AI/MassAI/Source/MassAIBehavior/Public/MassZoneGraphAnnotationProcessors.h
+++ b/Engine/Plugins/AI/MassAI/Source/MassAIBehavior/Public/MassZoneGraphAnnotationProcessors.h
@@ -26,7 +26,7 @@ public:
 protected:
 	virtual void Initialize(UObject& Owner) override;
 	virtual void ConfigureQueries() override;
-	virtual void Execute(UEntitySubsystem& EntitySubsystem, FLWComponentSystemExecutionContext& Context) override;
+	virtual void Execute(UMassEntitySubsystem& EntitySubsystem, FLWComponentSystemExecutionContext& Context) override;
 
 	FLWComponentQuery EntityQuery;
 
@@ -48,11 +48,11 @@ public:
 protected:
 	virtual void Initialize(UObject& Owner) override;
 	virtual void ConfigureQueries() override;
-	virtual void Execute(UEntitySubsystem& EntitySubsystem, FLWComponentSystemExecutionContext& Context) override;
+	virtual void Execute(UMassEntitySubsystem& EntitySubsystem, FLWComponentSystemExecutionContext& Context) override;
 
 	void UpdateAnnotationTags(FMassZoneGraphAnnotationTagsFragment& AnnotationTags, const FMassZoneGraphLaneLocationFragment& LaneLocation, FLWEntity Entity);
 
-	virtual void SignalEntities(UEntitySubsystem& EntitySubsystem, FLWComponentSystemExecutionContext& Context, FMassSignalNameLookup& EntitySignals) override;
+	virtual void SignalEntities(UMassEntitySubsystem& EntitySubsystem, FLWComponentSystemExecutionContext& Context, FMassSignalNameLookup& EntitySignals) override;
 
 	UPROPERTY(Transient)
 	UZoneGraphAnnotationSubsystem* ZoneGraphAnnotationSubsystem = nullptr;
diff --git a/Engine/Plugins/AI/MassAI/Source/MassAIDebug/Private/GameplayDebuggerCategory_Mass.cpp b/Engine/Plugins/AI/MassAI/Source/MassAIDebug/Private/GameplayDebuggerCategory_Mass.cpp
--- a/Engine/Plugins/AI/MassAI/Source/MassAIDebug/Private/GameplayDebuggerCategory_Mass.cpp
+++ b/Engine/Plugins/AI/MassAI/Source/MassAIDebug/Private/GameplayDebuggerCategory_Mass.cpp
@@ -10,7 +10,7 @@
 #include "MassActorSubsystem.h"
 #include "GameFramework/PlayerController.h"
 #include "MassAgentComponent.h"
-#include "MassEntitySystem.h"
+#include "MassEntitySubsystem.h"
 #include "MassAIMovementFragments.h"
 #include "MassLookAtFragments.h"
 #include "MassStateTreeFragments.h"
@@ -165,7 +165,7 @@ void FGameplayDebuggerCategory_Mass::CollectData(APlayerController* OwnerPC, AAc
 		}
 	}
 
-	UPipeEntitySubsystem* EntitySystem = UPipeEntitySubsystem::GetCurrent(World);
+	UMassEntitySubsystem* EntitySystem = UMassEntitySubsystem::GetCurrent(World);
 	if (EntitySystem)
 	{
 		AddTextLine(FString::Printf(TEXT("{Green}Entities count active{grey}/all: {white}%d{grey}/%d"), EntitySystem->DebugGetEntityCount(), EntitySystem->DebugGetEntityCount()));
diff --git a/Engine/Plugins/AI/MassAI/Source/MassAIDebug/Private/MassDebugStateTreeProcessor.cpp b/Engine/Plugins/AI/MassAI/Source/MassAIDebug/Private/MassDebugStateTreeProcessor.cpp
--- a/Engine/Plugins/AI/MassAI/Source/MassAIDebug/Private/MassDebugStateTreeProcessor.cpp
+++ b/Engine/Plugins/AI/MassAI/Source/MassAIDebug/Private/MassDebugStateTreeProcessor.cpp
@@ -23,7 +23,7 @@ void UMassDebugStateTreeProcessor::ConfigureQueries()
 	EntityQuery.AddRequirement<FDataFragment_Transform>(ELWComponentAccess::ReadOnly);
 }
 
-void UMassDebugStateTreeProcessor::Execute(UEntitySubsystem& EntitySubsystem, FLWComponentSystemExecutionContext& Context)
+void UMassDebugStateTreeProcessor::Execute(UMassEntitySubsystem& EntitySubsystem, FLWComponentSystemExecutionContext& Context)
 {
 #if WITH_MASS_DEBUG
 	UWorld* World = GetWorld();
diff --git a/Engine/Plugins/AI/MassAI/Source/MassAIDebug/Public/MassDebugStateTreeProcessor.h b/Engine/Plugins/AI/MassAI/Source/MassAIDebug/Public/MassDebugStateTreeProcessor.h
--- a/Engine/Plugins/AI/MassAI/Source/MassAIDebug/Public/MassDebugStateTreeProcessor.h
+++ b/Engine/Plugins/AI/MassAI/Source/MassAIDebug/Public/MassDebugStateTreeProcessor.h
@@ -5,7 +5,7 @@
 #include "MassProcessor.h"
 #include "MassDebugStateTreeProcessor.generated.h"
 
-class UEntitySubsystem;
+class UMassEntitySubsystem;
 struct FMassStateTreeFragment;
 struct FLWComponentQuery;
 struct FLWComponentSystemExecutionContext;
@@ -19,7 +19,7 @@ protected:
 	UMassDebugStateTreeProcessor();
 
 	virtual void ConfigureQueries() override;
-	virtual void Execute(UEntitySubsystem& EntitySubsystem, FLWComponentSystemExecutionContext& Context) override;
+	virtual void Execute(UMassEntitySubsystem& EntitySubsystem, FLWComponentSystemExecutionContext& Context) override;
 
 	FLWComponentQuery EntityQuery;
 };
diff --git a/Engine/Plugins/AI/MassAI/Source/MassAIMovement/Private/AssignRandomNavLocationProcessor.cpp b/Engine/Plugins/AI/MassAI/Source/MassAIMovement/Private/AssignRandomNavLocationProcessor.cpp
--- a/Engine/Plugins/AI/MassAI/Source/MassAIMovement/Private/AssignRandomNavLocationProcessor.cpp
+++ b/Engine/Plugins/AI/MassAI/Source/MassAIMovement/Private/AssignRandomNavLocationProcessor.cpp
@@ -21,7 +21,7 @@ void UAssignRandomNavLocationProcessor::ConfigureQueries()
 	EntityQuery.AddRequirement<FDataFragment_NavLocation>(ELWComponentAccess::ReadWrite);
 }
 
-void UAssignRandomNavLocationProcessor::Execute(UEntitySubsystem& EntitySubsystem, FLWComponentSystemExecutionContext& Context)
+void UAssignRandomNavLocationProcessor::Execute(UMassEntitySubsystem& EntitySubsystem, FLWComponentSystemExecutionContext& Context)
 {
 	if (!ensure(NavigationSystem))
 	{
diff --git a/Engine/Plugins/AI/MassAI/Source/MassAIMovement/Private/MassAvoidanceProcessors.cpp b/Engine/Plugins/AI/MassAI/Source/MassAIMovement/Private/MassAvoidanceProcessors.cpp
--- a/Engine/Plugins/AI/MassAI/Source/MassAIMovement/Private/MassAvoidanceProcessors.cpp
+++ b/Engine/Plugins/AI/MassAI/Source/MassAIMovement/Private/MassAvoidanceProcessors.cpp
@@ -416,7 +416,7 @@ void UMassAvoidanceProcessor::Initialize(UObject& Owner)
 	WeakMovementSubsystem = UWorld::GetSubsystem<UMassMovementSubsystem>(Owner.GetWorld());
 }
 
-void UMassAvoidanceProcessor::Execute(UEntitySubsystem& EntitySubsystem, FLWComponentSystemExecutionContext& Context)
+void UMassAvoidanceProcessor::Execute(UMassEntitySubsystem& EntitySubsystem, FLWComponentSystemExecutionContext& Context)
 {
 	QUICK_SCOPE_CYCLE_COUNTER(UMassAvoidanceProcessor);
 
@@ -1164,7 +1164,7 @@ void UMassStandingAvoidanceProcessor::Initialize(UObject& Owner)
 	WeakMovementSubsystem = UWorld::GetSubsystem<UMassMovementSubsystem>(Owner.GetWorld());
 }
 
-void UMassStandingAvoidanceProcessor::Execute(UEntitySubsystem& EntitySubsystem, FLWComponentSystemExecutionContext& Context)
+void UMassStandingAvoidanceProcessor::Execute(UMassEntitySubsystem& EntitySubsystem, FLWComponentSystemExecutionContext& Context)
 {
 	QUICK_SCOPE_CYCLE_COUNTER(UMassStandingAvoidanceProcessor);
 
@@ -1423,7 +1423,7 @@ void UMassAvoidanceObstacleProcessor::Initialize(UObject& Owner)
 	WeakMovementSubsystem = UWorld::GetSubsystem<UMassMovementSubsystem>(Owner.GetWorld());
 }
 
-void UMassAvoidanceObstacleProcessor::Execute(UEntitySubsystem& EntitySubsystem, FLWComponentSystemExecutionContext& Context)
+void UMassAvoidanceObstacleProcessor::Execute(UMassEntitySubsystem& EntitySubsystem, FLWComponentSystemExecutionContext& Context)
 {
 	using namespace UE::MassAvoidance;
 
@@ -1552,7 +1552,7 @@ void UMassNavigationBoundaryProcessor::Initialize(UObject& Owner)
 	WeakNavData = NavData;
 }
 
-void UMassNavigationBoundaryProcessor::Execute(UEntitySubsystem& EntitySubsystem, FLWComponentSystemExecutionContext& Context)
+void UMassNavigationBoundaryProcessor::Execute(UMassEntitySubsystem& EntitySubsystem, FLWComponentSystemExecutionContext& Context)
 {
 	using namespace UE::MassAvoidance;
 
@@ -1627,7 +1627,7 @@ void UMassLaneBoundaryProcessor::Initialize(UObject& Owner)
 	WeakZoneGraph = UWorld::GetSubsystem<UZoneGraphSubsystem>(Owner.GetWorld());
 }
 
-void UMassLaneBoundaryProcessor::Execute(UEntitySubsystem& EntitySubsystem, FLWComponentSystemExecutionContext& Context)
+void UMassLaneBoundaryProcessor::Execute(UMassEntitySubsystem& EntitySubsystem, FLWComponentSystemExecutionContext& Context)
 {
 	using namespace UE::MassAvoidance;
 
@@ -1904,7 +1904,7 @@ void UMassLaneCacheBoundaryProcessor::Initialize(UObject& Owner)
 	WeakWorld = Owner.GetWorld();
 }
 
-void UMassLaneCacheBoundaryProcessor::Execute(UEntitySubsystem& EntitySubsystem, FLWComponentSystemExecutionContext& Context)
+void UMassLaneCacheBoundaryProcessor::Execute(UMassEntitySubsystem& EntitySubsystem, FLWComponentSystemExecutionContext& Context)
 {
 	using namespace UE::MassAvoidance;
 
diff --git a/Engine/Plugins/AI/MassAI/Source/MassAIMovement/Private/MassMovementSubsystem.cpp b/Engine/Plugins/AI/MassAI/Source/MassAIMovement/Private/MassMovementSubsystem.cpp
--- a/Engine/Plugins/AI/MassAI/Source/MassAIMovement/Private/MassMovementSubsystem.cpp
+++ b/Engine/Plugins/AI/MassAI/Source/MassAIMovement/Private/MassMovementSubsystem.cpp
@@ -19,6 +19,6 @@ void UMassMovementSubsystem::Initialize(FSubsystemCollectionBase& Collection)
 	Collection.InitializeDependency<UMassSimulationSubsystem>();
 	ZoneGraphSubsystem = Collection.InitializeDependency<UZoneGraphSubsystem>();
 
-	EntitySubsystem = UPipeEntitySubsystem::GetCurrent(GetWorld());
+	EntitySubsystem = UMassEntitySubsystem::GetCurrent(GetWorld());
 }
 
diff --git a/Engine/Plugins/AI/MassAI/Source/MassAIMovement/Private/MassZoneGraphMovementProcessors.cpp b/Engine/Plugins/AI/MassAI/Source/MassAIMovement/Private/MassZoneGraphMovementProcessors.cpp
--- a/Engine/Plugins/AI/MassAI/Source/MassAIMovement/Private/MassZoneGraphMovementProcessors.cpp
+++ b/Engine/Plugins/AI/MassAI/Source/MassAIMovement/Private/MassZoneGraphMovementProcessors.cpp
@@ -90,7 +90,7 @@ void UMassZoneGraphLocationInitializer::Initialize(UObject& Owner)
 	SignalSubsystem = UWorld::GetSubsystem<UMassSignalSubsystem>(Owner.GetWorld());
 }
 
-void UMassZoneGraphLocationInitializer::Execute(UEntitySubsystem& EntitySubsystem, FLWComponentSystemExecutionContext& Context)
+void UMassZoneGraphLocationInitializer::Execute(UMassEntitySubsystem& EntitySubsystem, FLWComponentSystemExecutionContext& Context)
 {
 	const UMassMovementSettings* Settings = GetDefault<UMassMovementSettings>();
 	if (!ZoneGraphSubsystem || !Settings || !SignalSubsystem)
@@ -195,7 +195,7 @@ void UMassZoneGraphPathFollowProcessor::ConfigureQueries()
 	EntityQuery_Conditional.SetChunkFilter(&FMassSimulationVariableTickChunkFragment::ShouldTickChunkThisFrame);
 }
 
-void UMassZoneGraphPathFollowProcessor::Execute(UEntitySubsystem& EntitySubsystem, FLWComponentSystemExecutionContext& Context)
+void UMassZoneGraphPathFollowProcessor::Execute(UMassEntitySubsystem& EntitySubsystem, FLWComponentSystemExecutionContext& Context)
 {
 	if (!SignalSubsystem || !ZoneGraphSubsystem)
 	{
@@ -485,7 +485,7 @@ void UMassZoneGraphSteeringProcessor::ConfigureQueries()
 	EntityQuery.AddTagRequirement<FMassOffLODTag>(ELWComponentPresence::None);
 }
 
-void UMassZoneGraphSteeringProcessor::Execute(UEntitySubsystem& EntitySubsystem, FLWComponentSystemExecutionContext& Context)
+void UMassZoneGraphSteeringProcessor::Execute(UMassEntitySubsystem& EntitySubsystem, FLWComponentSystemExecutionContext& Context)
 {
 	const UMassMovementSettings* Settings = GetDefault<UMassMovementSettings>();
 	check(Settings);
diff --git a/Engine/Plugins/AI/MassAI/Source/MassAIMovement/Private/MovementProcessor.cpp b/Engine/Plugins/AI/MassAI/Source/MassAIMovement/Private/MovementProcessor.cpp
--- a/Engine/Plugins/AI/MassAI/Source/MassAIMovement/Private/MovementProcessor.cpp
+++ b/Engine/Plugins/AI/MassAI/Source/MassAIMovement/Private/MovementProcessor.cpp
@@ -118,7 +118,7 @@ void UMassProcessor_Movement::ConfigureQueries()
 	EntityQuery.AddRequirement<FMassVelocityFragment>(ELWComponentAccess::ReadWrite);
 }
 
-void UMassProcessor_Movement::Execute(UEntitySubsystem& EntitySubsystem, FLWComponentSystemExecutionContext& Context)
+void UMassProcessor_Movement::Execute(UMassEntitySubsystem& EntitySubsystem, FLWComponentSystemExecutionContext& Context)
 {
 	QUICK_SCOPE_CYCLE_COUNTER(MovementProcessor_Run);
 
@@ -176,7 +176,7 @@ void UMassAvoidanceObstacleRemoverFragmentDestructor::Initialize(UObject& Owner)
 	WeakMovementSubsystem = UWorld::GetSubsystem<UMassMovementSubsystem>(Owner.GetWorld());
 }
 
-void UMassAvoidanceObstacleRemoverFragmentDestructor::Execute(UEntitySubsystem& EntitySubsystem, FLWComponentSystemExecutionContext& Context)
+void UMassAvoidanceObstacleRemoverFragmentDestructor::Execute(UMassEntitySubsystem& EntitySubsystem, FLWComponentSystemExecutionContext& Context)
 {
 	UMassMovementSubsystem* MovementSubsystem = WeakMovementSubsystem.Get();
 	if (!MovementSubsystem)
@@ -212,7 +212,7 @@ void UVelocityFragmentInitializer::ConfigureQueries()
 	EntityQuery.AddRequirement<FMassVelocityFragment>(ELWComponentAccess::ReadWrite);
 }
 
-void UVelocityFragmentInitializer::Execute(UEntitySubsystem& EntitySubsystem, FLWComponentSystemExecutionContext& Context)
+void UVelocityFragmentInitializer::Execute(UMassEntitySubsystem& EntitySubsystem, FLWComponentSystemExecutionContext& Context)
 {
 	EntityQuery.ForEachEntityChunk(EntitySubsystem, Context, [this](FLWComponentSystemExecutionContext& Context)
 		{
diff --git a/Engine/Plugins/AI/MassAI/Source/MassAIMovement/Private/NavigationProcessor.cpp b/Engine/Plugins/AI/MassAI/Source/MassAIMovement/Private/NavigationProcessor.cpp
--- a/Engine/Plugins/AI/MassAI/Source/MassAIMovement/Private/NavigationProcessor.cpp
+++ b/Engine/Plugins/AI/MassAI/Source/MassAIMovement/Private/NavigationProcessor.cpp
@@ -98,7 +98,7 @@ void UMassApplyVelocityMoveTargetProcessor::ConfigureQueries()
 	LowResEntityQuery_Conditional.SetChunkFilter(&FMassSimulationVariableTickChunkFragment::ShouldTickChunkThisFrame);
 }
 
-void UMassApplyVelocityMoveTargetProcessor::Execute(UEntitySubsystem& EntitySubsystem,
+void UMassApplyVelocityMoveTargetProcessor::Execute(UMassEntitySubsystem& EntitySubsystem,
 													FLWComponentSystemExecutionContext& Context)
 {
 	// Clamp max delta time to avoid force explosion on large time steps (i.e. during initialization).
@@ -241,7 +241,7 @@ void UMassDynamicObstacleProcessor::ConfigureQueries()
 	EntityQuery_Conditional.SetChunkFilter(&FMassSimulationVariableTickChunkFragment::ShouldTickChunkThisFrame);
 }
 
-void UMassDynamicObstacleProcessor::Execute(UEntitySubsystem& EntitySubsystem, FLWComponentSystemExecutionContext& Context)
+void UMassDynamicObstacleProcessor::Execute(UMassEntitySubsystem& EntitySubsystem, FLWComponentSystemExecutionContext& Context)
 {
 	EntityQuery_Conditional.ForEachEntityChunk(EntitySubsystem, Context, [this](FLWComponentSystemExecutionContext& Context)
 		{
diff --git a/Engine/Plugins/AI/MassAI/Source/MassAIMovement/Private/SnapToNavigationProcessor.cpp b/Engine/Plugins/AI/MassAI/Source/MassAIMovement/Private/SnapToNavigationProcessor.cpp
--- a/Engine/Plugins/AI/MassAI/Source/MassAIMovement/Private/SnapToNavigationProcessor.cpp
+++ b/Engine/Plugins/AI/MassAI/Source/MassAIMovement/Private/SnapToNavigationProcessor.cpp
@@ -23,7 +23,7 @@ void USnapToNavigationProcessor::ConfigureQueries()
 	EntityQuery.AddRequirement<FDataFragment_NavLocation>(ELWComponentAccess::ReadWrite);
 }
 
-void USnapToNavigationProcessor::Execute(UEntitySubsystem& EntitySubsystem, FLWComponentSystemExecutionContext& Context)
+void USnapToNavigationProcessor::Execute(UMassEntitySubsystem& EntitySubsystem, FLWComponentSystemExecutionContext& Context)
 {
 	ANavigationData* NavData = WeakNavData.Get();
 	if (NavData == nullptr)
diff --git a/Engine/Plugins/AI/MassAI/Source/MassAIMovement/Public/AssignRandomNavLocationProcessor.h b/Engine/Plugins/AI/MassAI/Source/MassAIMovement/Public/AssignRandomNavLocationProcessor.h
--- a/Engine/Plugins/AI/MassAI/Source/MassAIMovement/Public/AssignRandomNavLocationProcessor.h
+++ b/Engine/Plugins/AI/MassAI/Source/MassAIMovement/Public/AssignRandomNavLocationProcessor.h
@@ -19,7 +19,7 @@ public:
 
 protected:
 	virtual void ConfigureQueries() override;
-	virtual void Execute(UEntitySubsystem& EntitySubsystem, FLWComponentSystemExecutionContext& Context) override;
+	virtual void Execute(UMassEntitySubsystem& EntitySubsystem, FLWComponentSystemExecutionContext& Context) override;
 	virtual void Initialize(UObject& InOwner) override;
 
 	UPROPERTY()
diff --git a/Engine/Plugins/AI/MassAI/Source/MassAIMovement/Public/MassAvoidanceProcessors.h b/Engine/Plugins/AI/MassAI/Source/MassAIMovement/Public/MassAvoidanceProcessors.h
--- a/Engine/Plugins/AI/MassAI/Source/MassAIMovement/Public/MassAvoidanceProcessors.h
+++ b/Engine/Plugins/AI/MassAI/Source/MassAIMovement/Public/MassAvoidanceProcessors.h
@@ -66,7 +66,7 @@ public:
 protected:
 	virtual void ConfigureQueries() override;
 	virtual void Initialize(UObject& Owner) override;
-	virtual void Execute(UEntitySubsystem& EntitySubsystem, FLWComponentSystemExecutionContext& Context) override;
+	virtual void Execute(UMassEntitySubsystem& EntitySubsystem, FLWComponentSystemExecutionContext& Context) override;
 
 private:
 	TWeakObjectPtr<UWorld> WeakWorld;
@@ -86,7 +86,7 @@ public:
 protected:
 	virtual void ConfigureQueries() override;
 	virtual void Initialize(UObject& Owner) override;
-	virtual void Execute(UEntitySubsystem& EntitySubsystem, FLWComponentSystemExecutionContext& Context) override;
+	virtual void Execute(UMassEntitySubsystem& EntitySubsystem, FLWComponentSystemExecutionContext& Context) override;
 
 private:
 	TWeakObjectPtr<UWorld> WeakWorld;
@@ -113,7 +113,7 @@ public:
 protected:
 	virtual void ConfigureQueries() override;
 	virtual void Initialize(UObject& Owner) override;
-	virtual void Execute(UEntitySubsystem& EntitySubsystem, FLWComponentSystemExecutionContext& Context) override;
+	virtual void Execute(UMassEntitySubsystem& EntitySubsystem, FLWComponentSystemExecutionContext& Context) override;
 
 private:
 	TWeakObjectPtr<UMassMovementSubsystem> WeakMovementSubsystem;
@@ -134,7 +134,7 @@ public:
 protected:
 	virtual void ConfigureQueries() override;
 	virtual void Initialize(UObject& Owner) override;
-	virtual void Execute(UEntitySubsystem& EntitySubsystem, FLWComponentSystemExecutionContext& Context) override;
+	virtual void Execute(UMassEntitySubsystem& EntitySubsystem, FLWComponentSystemExecutionContext& Context) override;
 
 private:
 	TWeakObjectPtr<UWorld> WeakWorld;
@@ -171,7 +171,7 @@ public:
 protected:
 	virtual void ConfigureQueries() override;
 	virtual void Initialize(UObject& Owner) override;
-	virtual void Execute(UEntitySubsystem& EntitySubsystem, FLWComponentSystemExecutionContext& Context) override;
+	virtual void Execute(UMassEntitySubsystem& EntitySubsystem, FLWComponentSystemExecutionContext& Context) override;
 
 private:
 	TWeakObjectPtr<UWorld> WeakWorld;
@@ -204,7 +204,7 @@ public:
 protected:
 	virtual void ConfigureQueries() override;
 	virtual void Initialize(UObject& Owner) override;
-	virtual void Execute(UEntitySubsystem& EntitySubsystem, FLWComponentSystemExecutionContext& Context) override;
+	virtual void Execute(UMassEntitySubsystem& EntitySubsystem, FLWComponentSystemExecutionContext& Context) override;
 
 private:
 	TWeakObjectPtr<UWorld> WeakWorld;
diff --git a/Engine/Plugins/AI/MassAI/Source/MassAIMovement/Public/MassMovementSubsystem.h b/Engine/Plugins/AI/MassAI/Source/MassAIMovement/Public/MassMovementSubsystem.h
--- a/Engine/Plugins/AI/MassAI/Source/MassAIMovement/Public/MassMovementSubsystem.h
+++ b/Engine/Plugins/AI/MassAI/Source/MassAIMovement/Public/MassMovementSubsystem.h
@@ -2,7 +2,7 @@
 
 #pragma once
 
-#include "MassEntitySystem.h"
+#include "MassEntitySubsystem.h"
 #include "HierarchicalHashGrid2D.h"
 #include "MassMovementTypes.h"
 #include "Subsystems/WorldSubsystem.h"
@@ -26,7 +26,7 @@ public:
 protected:
 
 	UPROPERTY(Transient)
-	UPipeEntitySubsystem* EntitySubsystem;
+	UMassEntitySubsystem* EntitySubsystem;
 
 	UPROPERTY(Transient)
 	UZoneGraphSubsystem* ZoneGraphSubsystem;
diff --git a/Engine/Plugins/AI/MassAI/Source/MassAIMovement/Public/MassZoneGraphMovementProcessors.h b/Engine/Plugins/AI/MassAI/Source/MassAIMovement/Public/MassZoneGraphMovementProcessors.h
--- a/Engine/Plugins/AI/MassAI/Source/MassAIMovement/Public/MassZoneGraphMovementProcessors.h
+++ b/Engine/Plugins/AI/MassAI/Source/MassAIMovement/Public/MassZoneGraphMovementProcessors.h
@@ -23,7 +23,7 @@ public:
 protected:
 	virtual void Initialize(UObject& Owner) override;
 	virtual void ConfigureQueries() override;
-	virtual void Execute(UEntitySubsystem& EntitySubsystem, FLWComponentSystemExecutionContext& Context) override;
+	virtual void Execute(UMassEntitySubsystem& EntitySubsystem, FLWComponentSystemExecutionContext& Context) override;
 
 	FLWComponentQuery EntityQuery;
 
@@ -47,7 +47,7 @@ protected:
 	
 	virtual void Initialize(UObject& Owner) override;
 	virtual void ConfigureQueries() override;
-	virtual void Execute(UEntitySubsystem& EntitySubsystem, FLWComponentSystemExecutionContext& Context) override;
+	virtual void Execute(UMassEntitySubsystem& EntitySubsystem, FLWComponentSystemExecutionContext& Context) override;
 
 	FLWComponentQuery EntityQuery_Conditional;
 
@@ -71,7 +71,7 @@ protected:
 	
 	virtual void Initialize(UObject& Owner) override;
 	virtual void ConfigureQueries() override;
-	virtual void Execute(UEntitySubsystem& EntitySubsystem, FLWComponentSystemExecutionContext& Context) override;
+	virtual void Execute(UMassEntitySubsystem& EntitySubsystem, FLWComponentSystemExecutionContext& Context) override;
 
 	FLWComponentQuery EntityQuery;
 
diff --git a/Engine/Plugins/AI/MassAI/Source/MassAIMovement/Public/MovementProcessor.h b/Engine/Plugins/AI/MassAI/Source/MassAIMovement/Public/MovementProcessor.h
--- a/Engine/Plugins/AI/MassAI/Source/MassAIMovement/Public/MovementProcessor.h
+++ b/Engine/Plugins/AI/MassAI/Source/MassAIMovement/Public/MovementProcessor.h
@@ -18,7 +18,7 @@ public:
 
 protected:
 	virtual void ConfigureQueries() override;
-	virtual void Execute(UEntitySubsystem& EntitySubsystem, FLWComponentSystemExecutionContext& Context) override;
+	virtual void Execute(UMassEntitySubsystem& EntitySubsystem, FLWComponentSystemExecutionContext& Context) override;
 
 	FLWComponentQuery EntityQuery;
 };
@@ -47,7 +47,7 @@ class MASSAIMOVEMENT_API UMassAvoidanceObstacleRemoverFragmentDestructor : publi
 protected:
 	virtual void ConfigureQueries() override;
 	virtual void Initialize(UObject& Owner) override;
-	virtual void Execute(UEntitySubsystem& EntitySubsystem, FLWComponentSystemExecutionContext& Context) override;
+	virtual void Execute(UMassEntitySubsystem& EntitySubsystem, FLWComponentSystemExecutionContext& Context) override;
 
 	FLWComponentQuery EntityQuery;
 };
@@ -61,7 +61,7 @@ public:
 
 protected:
 	virtual void ConfigureQueries() override;
-	virtual void Execute(UEntitySubsystem& EntitySubsystem, FLWComponentSystemExecutionContext& Context) override;
+	virtual void Execute(UMassEntitySubsystem& EntitySubsystem, FLWComponentSystemExecutionContext& Context) override;
 
 	FLWComponentQuery EntityQuery;
 };
diff --git a/Engine/Plugins/AI/MassAI/Source/MassAIMovement/Public/NavigationProcessor.h b/Engine/Plugins/AI/MassAI/Source/MassAIMovement/Public/NavigationProcessor.h
--- a/Engine/Plugins/AI/MassAI/Source/MassAIMovement/Public/NavigationProcessor.h
+++ b/Engine/Plugins/AI/MassAI/Source/MassAIMovement/Public/NavigationProcessor.h
@@ -24,7 +24,7 @@ public:
 
 protected:
 	virtual void ConfigureQueries() override;
-	virtual void Execute(UEntitySubsystem& EntitySubsystem, FLWComponentSystemExecutionContext& Context) override;
+	virtual void Execute(UMassEntitySubsystem& EntitySubsystem, FLWComponentSystemExecutionContext& Context) override;
 
 private:
 	FLWComponentQuery HighResEntityQuery;
@@ -50,7 +50,7 @@ public:
 	
 protected:
 	virtual void ConfigureQueries() override;
-	virtual void Execute(UEntitySubsystem& EntitySubsystem, FLWComponentSystemExecutionContext& Context) override;
+	virtual void Execute(UMassEntitySubsystem& EntitySubsystem, FLWComponentSystemExecutionContext& Context) override;
 	virtual void OnStop(FMassDynamicObstacleFragment& OutObstacle, const float BlockingRadius) {}
 	virtual void OnMove(FMassDynamicObstacleFragment& OutObstacle) {}
 
diff --git a/Engine/Plugins/AI/MassAI/Source/MassAIMovement/Public/SnapToNavigationProcessor.h b/Engine/Plugins/AI/MassAI/Source/MassAIMovement/Public/SnapToNavigationProcessor.h
--- a/Engine/Plugins/AI/MassAI/Source/MassAIMovement/Public/SnapToNavigationProcessor.h
+++ b/Engine/Plugins/AI/MassAI/Source/MassAIMovement/Public/SnapToNavigationProcessor.h
@@ -16,7 +16,7 @@ public:
 
 protected:	
 	virtual void ConfigureQueries() override;
-	virtual void Execute(UEntitySubsystem& EntitySubsystem, FLWComponentSystemExecutionContext& Context) override;
+	virtual void Execute(UMassEntitySubsystem& EntitySubsystem, FLWComponentSystemExecutionContext& Context) override;
 	virtual void Initialize(UObject& InOwner) override;
 
 	TWeakObjectPtr<ANavigationData> WeakNavData;
diff --git a/Engine/Plugins/AI/MassCrowd/Source/MassCrowd/Private/MassCrowdBubble.cpp b/Engine/Plugins/AI/MassCrowd/Source/MassCrowd/Private/MassCrowdBubble.cpp
--- a/Engine/Plugins/AI/MassCrowd/Source/MassCrowd/Private/MassCrowdBubble.cpp
+++ b/Engine/Plugins/AI/MassCrowd/Source/MassCrowd/Private/MassCrowdBubble.cpp
@@ -17,7 +17,7 @@ namespace UE::Mass::Crowd
 	// @todo provide a better way of selecting agents to debug
 	constexpr int32 MaxAgentsDraw = 300;
 
-	void DebugDrawReplicatedAgent(FLWEntity Entity, const UPipeEntitySubsystem& EntitySystem)
+	void DebugDrawReplicatedAgent(FLWEntity Entity, const UMassEntitySubsystem& EntitySystem)
 	{
 		static const FVector DebugCylinderHeight = FVector(0.f, 0.f, 200.f);
 		static constexpr float DebugCylinderRadius = 50.f;
@@ -70,7 +70,7 @@ void FMassCrowdClientBubbleHandler::DebugValidateBubbleOnServer()
 
 	if (UE::Mass::Crowd::bDebugReplicationPositions)
 	{
-		UPipeEntitySubsystem* EntitySystem = Serializer->GetEntitySystem();
+		UMassEntitySubsystem* EntitySystem = Serializer->GetEntitySystem();
 		check(EntitySystem);
 
 		// @todo cap at MaxAgentsDraw for now
@@ -97,7 +97,7 @@ void FMassCrowdClientBubbleHandler::DebugValidateBubbleOnClient()
 
 	if (UE::Mass::Crowd::bDebugReplicationPositions)
 	{
-		UPipeEntitySubsystem* EntitySystem = Serializer->GetEntitySystem();
+		UMassEntitySubsystem* EntitySystem = Serializer->GetEntitySystem();
 		check(EntitySystem);
 
 		UMassReplicationManager* ReplicationManager = Serializer->GetReplicationManager();
diff --git a/Engine/Plugins/AI/MassCrowd/Source/MassCrowd/Private/MassCrowdNavigationProcessor.cpp b/Engine/Plugins/AI/MassCrowd/Source/MassCrowd/Private/MassCrowdNavigationProcessor.cpp
--- a/Engine/Plugins/AI/MassCrowd/Source/MassCrowd/Private/MassCrowdNavigationProcessor.cpp
+++ b/Engine/Plugins/AI/MassCrowd/Source/MassCrowd/Private/MassCrowdNavigationProcessor.cpp
@@ -38,7 +38,7 @@ void UMassCrowdLaneTrackingSignalProcessor::Initialize(UObject& Owner)
 	SubscribeToSignal(UE::Mass::Signals::CurrentLaneChanged);
 }
 
-void UMassCrowdLaneTrackingSignalProcessor::SignalEntities(UEntitySubsystem& EntitySubsystem, FLWComponentSystemExecutionContext& Context, FMassSignalNameLookup& EntitySignals)
+void UMassCrowdLaneTrackingSignalProcessor::SignalEntities(UMassEntitySubsystem& EntitySubsystem, FLWComponentSystemExecutionContext& Context, FMassSignalNameLookup& EntitySignals)
 {
 	if (!MassCrowdSubsystem)
 	{
@@ -86,7 +86,7 @@ void UMassCrowdLaneTrackingDestructor::ConfigureQueries()
 	EntityQuery.AddRequirement<FMassCrowdLaneTrackingFragment>(ELWComponentAccess::ReadOnly);
 }
 
-void UMassCrowdLaneTrackingDestructor::Execute(UEntitySubsystem& EntitySubsystem, FLWComponentSystemExecutionContext& Context)
+void UMassCrowdLaneTrackingDestructor::Execute(UMassEntitySubsystem& EntitySubsystem, FLWComponentSystemExecutionContext& Context)
 {
 	EntityQuery.ForEachEntityChunk(EntitySubsystem, Context, [this](const FLWComponentSystemExecutionContext& Context)
 	{
diff --git a/Engine/Plugins/AI/MassCrowd/Source/MassCrowd/Private/MassCrowdReplicationProcessor.cpp b/Engine/Plugins/AI/MassCrowd/Source/MassCrowd/Private/MassCrowdReplicationProcessor.cpp
--- a/Engine/Plugins/AI/MassCrowd/Source/MassCrowd/Private/MassCrowdReplicationProcessor.cpp
+++ b/Engine/Plugins/AI/MassCrowd/Source/MassCrowd/Private/MassCrowdReplicationProcessor.cpp
@@ -48,7 +48,7 @@ void UMassCrowdReplicationProcessor::Initialize(UObject& Owner)
 	BubbleInfoClassHandle = ReplicationManager->GetBubbleInfoClassHandle(AMassCrowdClientBubbleInfo::StaticClass());
 }
 
-void UMassCrowdReplicationProcessor::Execute(UEntitySubsystem& EntitySubsystem, FLWComponentSystemExecutionContext& Context)
+void UMassCrowdReplicationProcessor::Execute(UMassEntitySubsystem& EntitySubsystem, FLWComponentSystemExecutionContext& Context)
 {
 	Super::Execute(EntitySubsystem, Context);
 
@@ -68,7 +68,7 @@ void UMassCrowdReplicationProcessor::Execute(UEntitySubsystem& EntitySubsystem,
 #endif // WITH_MASS_DEBUG
 }
 
-void UMassCrowdReplicationProcessor::ProcessClientReplication(UEntitySubsystem& EntitySubsystem, FLWComponentSystemExecutionContext& Context)
+void UMassCrowdReplicationProcessor::ProcessClientReplication(UMassEntitySubsystem& EntitySubsystem, FLWComponentSystemExecutionContext& Context)
 {
 #if UE_REPLICATION_COMPILE_SERVER_CODE
 
diff --git a/Engine/Plugins/AI/MassCrowd/Source/MassCrowd/Private/MassCrowdRepresentationProcessor.cpp b/Engine/Plugins/AI/MassCrowd/Source/MassCrowd/Private/MassCrowdRepresentationProcessor.cpp
--- a/Engine/Plugins/AI/MassCrowd/Source/MassCrowd/Private/MassCrowdRepresentationProcessor.cpp
+++ b/Engine/Plugins/AI/MassCrowd/Source/MassCrowd/Private/MassCrowdRepresentationProcessor.cpp
@@ -38,7 +38,7 @@ void UMassCrowdRepresentationProcessor::ConfigureQueries()
 	CharacterMovementEntitiesQuery_Conditional.SetChunkFilter(&FMassVisualizationChunkFragment::AreAnyEntitiesVisibleInChunk);
 }
 
-void UMassCrowdRepresentationProcessor::InitializeVelocity(UEntitySubsystem& EntitySubsystem, FLWComponentSystemExecutionContext& Context)
+void UMassCrowdRepresentationProcessor::InitializeVelocity(UMassEntitySubsystem& EntitySubsystem, FLWComponentSystemExecutionContext& Context)
 {
 	CharacterMovementEntitiesQuery_Conditional.ForEachEntityChunk(EntitySubsystem, Context, [this](FLWComponentSystemExecutionContext& Context)
 		{
diff --git a/Engine/Plugins/AI/MassCrowd/Source/MassCrowd/Private/MassCrowdServerRepresentationLODProcessor.cpp b/Engine/Plugins/AI/MassCrowd/Source/MassCrowd/Private/MassCrowdServerRepresentationLODProcessor.cpp
--- a/Engine/Plugins/AI/MassCrowd/Source/MassCrowd/Private/MassCrowdServerRepresentationLODProcessor.cpp
+++ b/Engine/Plugins/AI/MassCrowd/Source/MassCrowd/Private/MassCrowdServerRepresentationLODProcessor.cpp
@@ -43,7 +43,7 @@ void UMassCrowdServerRepresentationLODProcessor::Initialize(UObject& InOwner)
 	Super::Initialize(InOwner);
 }
 
-void UMassCrowdServerRepresentationLODProcessor::Execute(UEntitySubsystem& EntitySubsystem, FLWComponentSystemExecutionContext& Context)
+void UMassCrowdServerRepresentationLODProcessor::Execute(UMassEntitySubsystem& EntitySubsystem, FLWComponentSystemExecutionContext& Context)
 {
 	TRACE_CPUPROFILER_EVENT_SCOPE(TEXT("CrowdServerRepresentationLOD"))
 
diff --git a/Engine/Plugins/AI/MassCrowd/Source/MassCrowd/Private/MassCrowdServerRepresentationProcessor.cpp b/Engine/Plugins/AI/MassCrowd/Source/MassCrowd/Private/MassCrowdServerRepresentationProcessor.cpp
--- a/Engine/Plugins/AI/MassCrowd/Source/MassCrowd/Private/MassCrowdServerRepresentationProcessor.cpp
+++ b/Engine/Plugins/AI/MassCrowd/Source/MassCrowd/Private/MassCrowdServerRepresentationProcessor.cpp
@@ -16,7 +16,7 @@ UMassCrowdServerRepresentationProcessor::UMassCrowdServerRepresentationProcessor
 	LODRepresentation[EMassLOD::Off] = ERepresentationType::None;
 }
 
-void UMassCrowdServerRepresentationProcessor::Execute(UEntitySubsystem& EntitySubsystem, FLWComponentSystemExecutionContext& Context)
+void UMassCrowdServerRepresentationProcessor::Execute(UMassEntitySubsystem& EntitySubsystem, FLWComponentSystemExecutionContext& Context)
 {
 	EntityQuery.ForEachEntityChunk(EntitySubsystem, Context, [this](FLWComponentSystemExecutionContext& Context)
 	{
diff --git a/Engine/Plugins/AI/MassCrowd/Source/MassCrowd/Private/MassCrowdVisualizationLODProcessor.cpp b/Engine/Plugins/AI/MassCrowd/Source/MassCrowd/Private/MassCrowdVisualizationLODProcessor.cpp
--- a/Engine/Plugins/AI/MassCrowd/Source/MassCrowd/Private/MassCrowdVisualizationLODProcessor.cpp
+++ b/Engine/Plugins/AI/MassCrowd/Source/MassCrowd/Private/MassCrowdVisualizationLODProcessor.cpp
@@ -68,7 +68,7 @@ void UMassCrowdVisualizationLODProcessor::ConfigureQueries()
 	FarEntityQuery.AddChunkRequirement<FMassVisualizationChunkFragment>(ELWComponentAccess::ReadOnly);
 }
 
-void UMassCrowdVisualizationLODProcessor::Execute(UEntitySubsystem& EntitySubsystem, FLWComponentSystemExecutionContext& Context)
+void UMassCrowdVisualizationLODProcessor::Execute(UMassEntitySubsystem& EntitySubsystem, FLWComponentSystemExecutionContext& Context)
 {
 	ForceOffLOD((bool)UE::MassCrowd::GCrowdTurnOffVisualization);
 
diff --git a/Engine/Plugins/AI/MassCrowd/Source/MassCrowd/Private/MassCrowdVisualizationProcessor.cpp b/Engine/Plugins/AI/MassCrowd/Source/MassCrowd/Private/MassCrowdVisualizationProcessor.cpp
--- a/Engine/Plugins/AI/MassCrowd/Source/MassCrowd/Private/MassCrowdVisualizationProcessor.cpp
+++ b/Engine/Plugins/AI/MassCrowd/Source/MassCrowd/Private/MassCrowdVisualizationProcessor.cpp
@@ -35,7 +35,7 @@ void UMassCrowdVisualizationProcessor::ConfigureQueries()
 	Super::ConfigureQueries();
 }
 
-void UMassCrowdVisualizationProcessor::Execute(UEntitySubsystem& EntitySubsystem, FLWComponentSystemExecutionContext& Context)
+void UMassCrowdVisualizationProcessor::Execute(UMassEntitySubsystem& EntitySubsystem, FLWComponentSystemExecutionContext& Context)
 {
 	EntityQuery.ForEachEntityChunk(EntitySubsystem, Context, [this](FLWComponentSystemExecutionContext& Context)
 	{
diff --git a/Engine/Plugins/AI/MassCrowd/Source/MassCrowd/Public/MassCrowdNavigationProcessor.h b/Engine/Plugins/AI/MassCrowd/Source/MassCrowd/Public/MassCrowdNavigationProcessor.h
--- a/Engine/Plugins/AI/MassCrowd/Source/MassCrowd/Public/MassCrowdNavigationProcessor.h
+++ b/Engine/Plugins/AI/MassCrowd/Source/MassCrowd/Public/MassCrowdNavigationProcessor.h
@@ -23,7 +23,7 @@ public:
 protected:
 	virtual void Initialize(UObject& Owner) override;
 	virtual void ConfigureQueries() override;
-	virtual void SignalEntities(UEntitySubsystem& EntitySubsystem, FLWComponentSystemExecutionContext& Context, FMassSignalNameLookup& EntitySignals) override;
+	virtual void SignalEntities(UMassEntitySubsystem& EntitySubsystem, FLWComponentSystemExecutionContext& Context, FMassSignalNameLookup& EntitySignals) override;
 
 	UPROPERTY(Transient)
 	UMassCrowdSubsystem* MassCrowdSubsystem;
@@ -41,7 +41,7 @@ public:
 protected:
 	virtual void Initialize(UObject& Owner) override;
 	virtual void ConfigureQueries() override;
-	virtual void Execute(UEntitySubsystem& EntitySubsystem, FLWComponentSystemExecutionContext& Context) override;
+	virtual void Execute(UMassEntitySubsystem& EntitySubsystem, FLWComponentSystemExecutionContext& Context) override;
 
 	UPROPERTY(Transient)
 	UMassCrowdSubsystem* MassCrowdSubsystem;
diff --git a/Engine/Plugins/AI/MassCrowd/Source/MassCrowd/Public/MassCrowdReplicationProcessor.h b/Engine/Plugins/AI/MassCrowd/Source/MassCrowd/Public/MassCrowdReplicationProcessor.h
--- a/Engine/Plugins/AI/MassCrowd/Source/MassCrowd/Public/MassCrowdReplicationProcessor.h
+++ b/Engine/Plugins/AI/MassCrowd/Source/MassCrowd/Public/MassCrowdReplicationProcessor.h
@@ -22,7 +22,7 @@ public:
 protected:
 	virtual void ConfigureQueries() override;
 	virtual void Initialize(UObject& Owner) override;
-	virtual void Execute(UEntitySubsystem& EntitySubsystem, FLWComponentSystemExecutionContext& Context) override;
+	virtual void Execute(UMassEntitySubsystem& EntitySubsystem, FLWComponentSystemExecutionContext& Context) override;
 
-	void ProcessClientReplication(UEntitySubsystem& EntitySubsystem, FLWComponentSystemExecutionContext& Context);
+	void ProcessClientReplication(UMassEntitySubsystem& EntitySubsystem, FLWComponentSystemExecutionContext& Context);
 };
\ No newline at end of file
diff --git a/Engine/Plugins/AI/MassCrowd/Source/MassCrowd/Public/MassCrowdRepresentationProcessor.h b/Engine/Plugins/AI/MassCrowd/Source/MassCrowd/Public/MassCrowdRepresentationProcessor.h
--- a/Engine/Plugins/AI/MassCrowd/Source/MassCrowd/Public/MassCrowdRepresentationProcessor.h
+++ b/Engine/Plugins/AI/MassCrowd/Source/MassCrowd/Public/MassCrowdRepresentationProcessor.h
@@ -65,7 +65,7 @@ protected:
 	 * @param EntitySubsystem The subsystem in which all entities must be initialized
 	 * @param Context of the execution from the entity sub system
 	 */
-	void InitializeVelocity(UEntitySubsystem& EntitySubsystem, FLWComponentSystemExecutionContext& Context);
+	void InitializeVelocity(UMassEntitySubsystem& EntitySubsystem, FLWComponentSystemExecutionContext& Context);
 
 	/** 
 	 * A dedicated query for processing entities owning a FDataFragment_CharacterMovementComponentWrapper
diff --git a/Engine/Plugins/AI/MassCrowd/Source/MassCrowd/Public/MassCrowdServerRepresentationLODProcessor.h b/Engine/Plugins/AI/MassCrowd/Source/MassCrowd/Public/MassCrowdServerRepresentationLODProcessor.h
--- a/Engine/Plugins/AI/MassCrowd/Source/MassCrowd/Public/MassCrowdServerRepresentationLODProcessor.h
+++ b/Engine/Plugins/AI/MassCrowd/Source/MassCrowd/Public/MassCrowdServerRepresentationLODProcessor.h
@@ -22,7 +22,7 @@ protected:
 
 	virtual void ConfigureQueries() override;
 	virtual void Initialize(UObject& InOwner) override;
-	virtual void Execute(UEntitySubsystem& EntitySubsystem, FLWComponentSystemExecutionContext& Context) override;
+	virtual void Execute(UMassEntitySubsystem& EntitySubsystem, FLWComponentSystemExecutionContext& Context) override;
 
 	/** Distance where each LOD becomes relevant */
 	UPROPERTY(EditAnywhere, Category = "Mass|LOD", config)
diff --git a/Engine/Plugins/AI/MassCrowd/Source/MassCrowd/Public/MassCrowdServerRepresentationProcessor.h b/Engine/Plugins/AI/MassCrowd/Source/MassCrowd/Public/MassCrowdServerRepresentationProcessor.h
--- a/Engine/Plugins/AI/MassCrowd/Source/MassCrowd/Public/MassCrowdServerRepresentationProcessor.h
+++ b/Engine/Plugins/AI/MassCrowd/Source/MassCrowd/Public/MassCrowdServerRepresentationProcessor.h
@@ -25,5 +25,5 @@ protected:
 	 * @param EntitySubsystem is the system to execute the lambdas on each entity chunk
 	 * @param Context is the execution context to be passed when executing the lambdas
 	 */
-	virtual void Execute(UEntitySubsystem& EntitySubsystem, FLWComponentSystemExecutionContext& Context) override;
+	virtual void Execute(UMassEntitySubsystem& EntitySubsystem, FLWComponentSystemExecutionContext& Context) override;
 };
\ No newline at end of file
diff --git a/Engine/Plugins/AI/MassCrowd/Source/MassCrowd/Public/MassCrowdVisualizationLODProcessor.h b/Engine/Plugins/AI/MassCrowd/Source/MassCrowd/Public/MassCrowdVisualizationLODProcessor.h
--- a/Engine/Plugins/AI/MassCrowd/Source/MassCrowd/Public/MassCrowdVisualizationLODProcessor.h
+++ b/Engine/Plugins/AI/MassCrowd/Source/MassCrowd/Public/MassCrowdVisualizationLODProcessor.h
@@ -18,5 +18,5 @@ public:
 
 protected:
 	virtual void ConfigureQueries() override;
-	virtual void Execute(UEntitySubsystem& EntitySubsystem, FLWComponentSystemExecutionContext& Context) override;
+	virtual void Execute(UMassEntitySubsystem& EntitySubsystem, FLWComponentSystemExecutionContext& Context) override;
 };
diff --git a/Engine/Plugins/AI/MassCrowd/Source/MassCrowd/Public/MassCrowdVisualizationProcessor.h b/Engine/Plugins/AI/MassCrowd/Source/MassCrowd/Public/MassCrowdVisualizationProcessor.h
--- a/Engine/Plugins/AI/MassCrowd/Source/MassCrowd/Public/MassCrowdVisualizationProcessor.h
+++ b/Engine/Plugins/AI/MassCrowd/Source/MassCrowd/Public/MassCrowdVisualizationProcessor.h
@@ -26,7 +26,7 @@ protected:
 	 * @param EntitySubsystem is the system to execute the lambdas on each entity chunk
 	 * @param Context is the execution context to be passed when executing the lambdas
 	 */
-	virtual void Execute(UEntitySubsystem& EntitySubsystem, FLWComponentSystemExecutionContext& Context) override;
+	virtual void Execute(UMassEntitySubsystem& EntitySubsystem, FLWComponentSystemExecutionContext& Context) override;
 };
 
 UCLASS()
diff --git a/Engine/Plugins/Runtime/MassEntity/Source/MassEntity/Private/ArchetypeData.h b/Engine/Plugins/Runtime/MassEntity/Source/MassEntity/Private/ArchetypeData.h
--- a/Engine/Plugins/Runtime/MassEntity/Source/MassEntity/Private/ArchetypeData.h
+++ b/Engine/Plugins/Runtime/MassEntity/Source/MassEntity/Private/ArchetypeData.h
@@ -2,7 +2,7 @@
 
 #pragma once
 
-#include "EntitySubsystem.h"
+#include "MassEntitySubsystem.h"
 #include "ArchetypeTypes.h"
 
 struct FLWComponentQuery;
diff --git a/Engine/Plugins/Runtime/MassEntity/Source/MassEntity/Private/ArchetypeTypes.cpp b/Engine/Plugins/Runtime/MassEntity/Source/MassEntity/Private/ArchetypeTypes.cpp
--- a/Engine/Plugins/Runtime/MassEntity/Source/MassEntity/Private/ArchetypeTypes.cpp
+++ b/Engine/Plugins/Runtime/MassEntity/Source/MassEntity/Private/ArchetypeTypes.cpp
@@ -1,7 +1,7 @@
 // Copyright Epic Games, Inc. All Rights Reserved.
 
 #include "ArchetypeTypes.h"
-#include "EntitySubsystem.h"
+#include "MassEntitySubsystem.h"
 #include "ArchetypeData.h"
 #include "LWCCommandBuffer.h"
 
diff --git a/Engine/Plugins/Runtime/MassEntity/Source/MassEntity/Private/EntityQuery.cpp b/Engine/Plugins/Runtime/MassEntity/Source/MassEntity/Private/EntityQuery.cpp
--- a/Engine/Plugins/Runtime/MassEntity/Source/MassEntity/Private/EntityQuery.cpp
+++ b/Engine/Plugins/Runtime/MassEntity/Source/MassEntity/Private/EntityQuery.cpp
@@ -1,7 +1,7 @@
 // Copyright Epic Games, Inc. All Rights Reserved.
 
 #include "EntityQuery.h"
-#include "EntitySubsystem.h"
+#include "MassEntitySubsystem.h"
 #include "ArchetypeData.h"
 #include "LWCCommandBuffer.h"
 #include "VisualLogger/VisualLogger.h"
@@ -53,7 +53,7 @@ void FLWComponentQuery::SortRequirements()
 	ChunkRequirements.Sort(FLWComponentSorterOperator<FLWComponentRequirement>());
 }
 
-void FLWComponentQuery::CacheArchetypes(UEntitySubsystem& InEntitySubsystem)
+void FLWComponentQuery::CacheArchetypes(UMassEntitySubsystem& InEntitySubsystem)
 {
 	const uint32 InEntitySubsystemHash = PointerHash(&InEntitySubsystem);
 	if (EntitySubsystemHash != InEntitySubsystemHash || InEntitySubsystem.GetArchetypeDataVersion() != ArchetypeDataVersion)
@@ -117,7 +117,7 @@ bool FLWComponentQuery::DoesArchetypeMatchRequirements(const FArchetypeHandle& A
 	return true;
 }
 
-void FLWComponentQuery::ForEachEntityChunk(const FArchetypeChunkCollection& Chunks, UEntitySubsystem& EntitySubsystem, FLWComponentSystemExecutionContext& ExecutionContext, const FLWComponentSystemExecuteFunction& ExecuteFunction)
+void FLWComponentQuery::ForEachEntityChunk(const FArchetypeChunkCollection& Chunks, UMassEntitySubsystem& EntitySubsystem, FLWComponentSystemExecutionContext& ExecutionContext, const FLWComponentSystemExecuteFunction& ExecuteFunction)
 {
 	// mz@todo I don't like that we're copying data here.
 	ExecutionContext.SetChunkCollection(Chunks);
@@ -125,7 +125,7 @@ void FLWComponentQuery::ForEachEntityChunk(const FArchetypeChunkCollection& Chun
 	ExecutionContext.ClearChunkCollection();
 }
 
-void FLWComponentQuery::ForEachEntityChunk(UEntitySubsystem& EntitySubsystem, FLWComponentSystemExecutionContext& ExecutionContext, const FLWComponentSystemExecuteFunction& ExecuteFunction)
+void FLWComponentQuery::ForEachEntityChunk(UMassEntitySubsystem& EntitySubsystem, FLWComponentSystemExecutionContext& ExecutionContext, const FLWComponentSystemExecuteFunction& ExecuteFunction)
 {
 #if WITH_AGGREGATETICKING_DEBUG
 	int32 NumEntitiesToProcess = 0;
@@ -175,7 +175,7 @@ void FLWComponentQuery::ForEachEntityChunk(UEntitySubsystem& EntitySubsystem, FL
 	ExecutionContext.FlushDeferred(EntitySubsystem);
 }
 
-void FLWComponentQuery::ParallelForEachEntityChunk(UEntitySubsystem& EntitySubsystem, FLWComponentSystemExecutionContext& ExecutionContext, const FLWComponentSystemExecuteFunction& ExecuteFunction)
+void FLWComponentQuery::ParallelForEachEntityChunk(UMassEntitySubsystem& EntitySubsystem, FLWComponentSystemExecutionContext& ExecutionContext, const FLWComponentSystemExecuteFunction& ExecuteFunction)
 {
 	if (bAllowParallelExecution == false)
 	{
@@ -241,7 +241,7 @@ void FLWComponentQuery::ParallelForEachEntityChunk(UEntitySubsystem& EntitySubsy
 	ExecutionContext.FlushDeferred(EntitySubsystem);
 }
 
-int32 FLWComponentQuery::GetMatchingEntitiesNum(UEntitySubsystem& InEntitySubsystem)
+int32 FLWComponentQuery::GetMatchingEntitiesNum(UMassEntitySubsystem& InEntitySubsystem)
 {
 	CacheArchetypes(InEntitySubsystem);
 	int32 TotalEntities = 0;
@@ -255,7 +255,7 @@ int32 FLWComponentQuery::GetMatchingEntitiesNum(UEntitySubsystem& InEntitySubsys
 	return TotalEntities;
 }
 
-bool FLWComponentQuery::HasMatchingEntities(UEntitySubsystem& InEntitySubsystem)
+bool FLWComponentQuery::HasMatchingEntities(UMassEntitySubsystem& InEntitySubsystem)
 {
 	CacheArchetypes(InEntitySubsystem);
 
diff --git a/Engine/Plugins/Runtime/MassEntity/Source/MassEntity/Private/EntitySubsystem.cpp b/Engine/Plugins/Runtime/MassEntity/Source/MassEntity/Private/EntitySubsystem.cpp
deleted file mode 100644
diff --git a/Engine/Plugins/Runtime/MassEntity/Source/MassEntity/Private/EntityView.cpp b/Engine/Plugins/Runtime/MassEntity/Source/MassEntity/Private/EntityView.cpp
--- a/Engine/Plugins/Runtime/MassEntity/Source/MassEntity/Private/EntityView.cpp
+++ b/Engine/Plugins/Runtime/MassEntity/Source/MassEntity/Private/EntityView.cpp
@@ -1,7 +1,7 @@
 // Copyright Epic Games, Inc. All Rights Reserved.
 
 #include "EntityView.h"
-#include "EntitySubsystem.h"
+#include "MassEntitySubsystem.h"
 #include "ArchetypeData.h"
 
 
@@ -15,7 +15,7 @@ FEntityView::FEntityView(const FArchetypeHandle& ArchetypeHandle, FLWEntity InEn
 	EntityHandle = Archetype->MakeEntityHandle(Entity);
 }
 
-FEntityView::FEntityView(const UEntitySubsystem& EntitySubsystem, FLWEntity InEntity)
+FEntityView::FEntityView(const UMassEntitySubsystem& EntitySubsystem, FLWEntity InEntity)
 {
 	Entity = InEntity;
 	const FArchetypeHandle ArchetypeHandle = EntitySubsystem.GetArchetypeForEntity(Entity);
diff --git a/Engine/Plugins/Runtime/MassEntity/Source/MassEntity/Private/LWCCommandBuffer.cpp b/Engine/Plugins/Runtime/MassEntity/Source/MassEntity/Private/LWCCommandBuffer.cpp
--- a/Engine/Plugins/Runtime/MassEntity/Source/MassEntity/Private/LWCCommandBuffer.cpp
+++ b/Engine/Plugins/Runtime/MassEntity/Source/MassEntity/Private/LWCCommandBuffer.cpp
@@ -66,7 +66,7 @@ void GetCommandStatNames(FStructView Entry, FString& OutName, ANSIName*& OutANSI
 #endif
 } // UE::FLWCCommand
 
-void FLWCCommandBuffer::ReplayBufferAgainstSystem(UEntitySubsystem* System)
+void FLWCCommandBuffer::ReplayBufferAgainstSystem(UMassEntitySubsystem* System)
 {
 	check(System);
 
@@ -119,12 +119,12 @@ void FLWCCommandBuffer::MoveAppend(FLWCCommandBuffer& Other)
 //////////////////////////////////////////////////////////////////////
 // Command implementations
 
-void FCommandAddComponentInstance::Execute(UEntitySubsystem& System) const
+void FCommandAddComponentInstance::Execute(UMassEntitySubsystem& System) const
 {
 	System.AddComponentInstanceListToEntity(TargetEntity, MakeArrayView(&Struct, 1));
 }
 
-void FCommandSwapTags::Execute(UEntitySubsystem& System) const
+void FCommandSwapTags::Execute(UMassEntitySubsystem& System) const
 {
 	if (System.IsEntityValid(TargetEntity) == false)
 	{
@@ -145,37 +145,37 @@ void FCommandSwapTags::Execute(UEntitySubsystem& System) const
 	}
 }
 
-void FBuildEntityFromComponentInstance::Execute(UEntitySubsystem& System) const
+void FBuildEntityFromComponentInstance::Execute(UMassEntitySubsystem& System) const
 {
 	System.BuildEntity(TargetEntity, MakeArrayView(&Struct, 1));
 }
 
-void FBuildEntityFromComponentInstances::Execute(UEntitySubsystem& System) const
+void FBuildEntityFromComponentInstances::Execute(UMassEntitySubsystem& System) const
 {
 	System.BuildEntity(TargetEntity, Instances);
 }
 
-void FCommandAddComponent::Execute(UEntitySubsystem& System) const
+void FCommandAddComponent::Execute(UMassEntitySubsystem& System) const
 {
 	System.AddComponentToEntity(TargetEntity, StructParam);
 }
 
-void FCommandRemoveComponent::Execute(UEntitySubsystem& System) const
+void FCommandRemoveComponent::Execute(UMassEntitySubsystem& System) const
 {
 	System.RemoveComponentFromEntity(TargetEntity, StructParam);
 }
 
-void FCommandAddComponentList::Execute(UEntitySubsystem& System) const
+void FCommandAddComponentList::Execute(UMassEntitySubsystem& System) const
 {
 	System.AddComponentListToEntity(TargetEntity, ComponentList);
 }
 
-void FCommandRemoveComponentList::Execute(UEntitySubsystem& System) const
+void FCommandRemoveComponentList::Execute(UMassEntitySubsystem& System) const
 {
 	System.RemoveComponentListFromEntity(TargetEntity, ComponentList);
 }
 
-void FCommandAddTag::Execute(UEntitySubsystem & System) const 
+void FCommandAddTag::Execute(UMassEntitySubsystem & System) const 
 {
 	if (System.IsEntityValid(TargetEntity) == false)
 {
@@ -185,12 +185,12 @@ void FCommandAddTag::Execute(UEntitySubsystem & System) const
 	System.AddTagToEntity(TargetEntity, StructParam);
 }
 
-void FCommandRemoveTag::Execute(UEntitySubsystem& System) const
+void FCommandRemoveTag::Execute(UMassEntitySubsystem& System) const
 {
 	System.RemoveTagFromEntity(TargetEntity, StructParam);
 }
 
-void FCommandRemoveComposition::Execute(UEntitySubsystem& System) const
+void FCommandRemoveComposition::Execute(UMassEntitySubsystem& System) const
 {
 	System.RemoveCompositionFromEntity(TargetEntity, Descriptor);
 }
diff --git a/Engine/Plugins/Runtime/MassEntity/Source/MassEntity/Private/MassEntityDebug.cpp b/Engine/Plugins/Runtime/MassEntity/Source/MassEntity/Private/MassEntityDebug.cpp
--- a/Engine/Plugins/Runtime/MassEntity/Source/MassEntity/Private/MassEntityDebug.cpp
+++ b/Engine/Plugins/Runtime/MassEntity/Source/MassEntity/Private/MassEntityDebug.cpp
@@ -7,7 +7,8 @@
 
 DEFINE_ENUM_TO_STRING(EPipeProcessingPhase);
 
-namespace UE { namespace Pipe { namespace Debug {
+namespace UE::Pipe::Debug
+{
 
 void DebugOutputDescription(TConstArrayView<UPipeProcessor*> Processors, FOutputDevice& Ar)
 {
@@ -30,5 +31,100 @@ void DebugOutputDescription(TConstArrayView<UPipeProcessor*> Processors, FOutput
 #endif // WITH_PIPE_DEBUG
 }
 
-}}} // namespace UE::Pipe::Debug
+#if WITH_PIPE_DEBUG && WITH_AGGREGATETICKING_DEBUG
+FAutoConsoleCommandWithWorldArgsAndOutputDevice PrintEntityFragmentsCmd(
+	TEXT("pipe.PrintEntityFragments"),
+	TEXT("Prints all fragment types and values (uproperties) for the specified Entity index"),
+	FConsoleCommandWithWorldArgsAndOutputDeviceDelegate::CreateLambda(
+		[](const TArray<FString>& Params, UWorld* World, FOutputDevice& Ar)
+		{
+			if (UMassEntitySubsystem* EntitySystem = UMassEntitySubsystem::GetCurrent(World))
+			{
+				int32 Index = INDEX_NONE;
+				if (LexTryParseString<int32>(Index, *Params[0]))
+				{
+					EntitySystem->DebugPrintEntity(Index, Ar);
+				}
+				else
+				{
+					Ar.Logf(ELogVerbosity::Error, TEXT("Entity index parameter must be an integer"));
+				}
+			}
+			else
+			{
+				Ar.Logf(ELogVerbosity::Error, TEXT("Failed to find PipeEntitySubsystem for world %s"), *GetPathNameSafe(World));
+			}
+		})
+);
+
+FAutoConsoleCommandWithWorldArgsAndOutputDevice LogArchetypesCmd(
+	TEXT("pipe.LogArchetypes"),
+	TEXT("Dumps description of archetypes to log. Optional parameter controls whether to include or exclude non-occupied archetypes. Defaults to 'include'."),
+	FConsoleCommandWithWorldArgsAndOutputDeviceDelegate::CreateLambda([](const TArray<FString>& Params, UWorld*, FOutputDevice& Ar)
+	{
+		const TIndirectArray<FWorldContext>& WorldContexts = GEngine->GetWorldContexts();
+		for (const FWorldContext& Context : WorldContexts)
+		{
+			UWorld* World = Context.World();
+			if (World == nullptr || World->IsPreviewWorld())
+			{
+				continue;
+			}
+
+			Ar.Logf(ELogVerbosity::Log, TEXT("Dumping description of archetypes for world: %s (%s - %s)"),
+				*GetPathNameSafe(World),
+				LexToString(World->WorldType),
+				*ToString(World->GetNetMode()));
+
+			if (UMassEntitySubsystem* EntitySystem = UMassEntitySubsystem::GetCurrent(World))
+			{
+				bool bIncludeEmpty = true;
+				if (Params.Num())
+				{
+					LexTryParseString(bIncludeEmpty, *Params[0]);
+				}
+				Ar.Logf(ELogVerbosity::Log, TEXT("Include empty archetypes: %s"), bIncludeEmpty ? TEXT("TRUE") : TEXT("FALSE"));
+				EntitySystem->DebugGetArchetypesStringDetails(Ar, bIncludeEmpty);
+			}
+			else
+			{
+				Ar.Logf(ELogVerbosity::Error, TEXT("Failed to find PipeEntitySubsystem for world: %s (%s - %s)"),
+					*GetPathNameSafe(World),
+					LexToString(World->WorldType),
+					*ToString(World->GetNetMode()));
+			}
+		}
+	})
+);
+
+// @todo these console commands will be reparented to "massentities" domain once we rename and shuffle the modules around 
+FAutoConsoleCommandWithWorld RecacheQueries(
+	TEXT("pipe.RecacheQueries"),
+	TEXT("Forces EntityQueries to recache their valid archetypes"),
+	FConsoleCommandWithWorldDelegate::CreateLambda([](UWorld* InWorld)
+		{
+			if (UMassEntitySubsystem* System = UMassEntitySubsystem::GetCurrent(InWorld))
+			{
+				System->DebugForceArchetypeDataVersionBump();
+			}
+		}
+	));
+
+FAutoConsoleCommandWithWorldArgsAndOutputDevice LogFragmentSizes(
+	TEXT("pipe.LogFragmentSizes"),
+	TEXT("Logs all the fragment types being used along with their sizes."),
+	FConsoleCommandWithWorldArgsAndOutputDeviceDelegate::CreateLambda([](const TArray<FString>& Params, UWorld* World, FOutputDevice& Ar)
+		{
+			for (const TWeakObjectPtr<const UScriptStruct>& WeakStruct : FLWComponentBitSet::DebugGetAllStructTypes())
+			{
+				if (const UScriptStruct* StructType = WeakStruct.Get())
+				{
+					Ar.Logf(ELogVerbosity::Log, TEXT("%s, size: %d"), *StructType->GetName(), StructType->GetStructureSize());
+				}
+			}
+		})
+	);
+#endif // WITH_PIPE_DEBUG && WITH_AGGREGATETICKING_DEBUG
+
+} // namespace UE::Pipe::Debug
 
diff --git a/Engine/Plugins/Runtime/MassEntity/Source/MassEntity/Private/MassEntitySubsystem.cpp b/Engine/Plugins/Runtime/MassEntity/Source/MassEntity/Private/MassEntitySubsystem.cpp
new file mode 100644
--- /dev/null
+++ b/Engine/Plugins/Runtime/MassEntity/Source/MassEntity/Private/MassEntitySubsystem.cpp
@@ -0,0 +1,1105 @@
+// Copyright Epic Games, Inc. All Rights Reserved.
+
+#include "MassEntitySubsystem.h"
+#include "ArchetypeData.h"
+#include "LWCCommandBuffer.h"
+#include "HAL/IConsoleManager.h"
+#include "Engine/World.h"
+#include "UObject/UObjectIterator.h"
+
+const FLWEntity UMassEntitySubsystem::InvalidEntity;
+
+DEFINE_LOG_CATEGORY(LogAggregateTicking);
+
+namespace UE { namespace AggregateTicking {
+
+FString DebugGetComponentAccessString(ELWComponentAccess Access)
+{
+#if WITH_AGGREGATETICKING_DEBUG
+	switch (Access)
+	{
+	case ELWComponentAccess::None:	return TEXT("--");
+	case ELWComponentAccess::ReadOnly:	return TEXT("RO");
+	case ELWComponentAccess::ReadWrite:	return TEXT("RW");
+	default:
+		ensureMsgf(false, TEXT("Missing string conversion for ELWComponentAccess=%d"), Access);
+		break;
+	}
+#endif // WITH_AGGREGATETICKING_DEBUG
+	return TEXT("Missing string conversion");
+}
+
+}} // UE::AggregateTicking
+
+//@TODO: Everything still alive leaks at shutdown
+//@TODO: No re-entrance safety while running a system (e.g., preventing someone from adding/removing entities or altering archetypes, etc...)
+//@TODO: Do we allow GCable types?  If so, need to implement AddReferencedObjects
+//@TODO: Do a UObject stat scope around the system Run call
+//@TODO: Tag components end up taking up space per-entity when they shouldn't
+
+//////////////////////////////////////////////////////////////////////
+// UMassEntitySubsystem
+
+UMassEntitySubsystem::UMassEntitySubsystem()
+{
+}
+
+void UMassEntitySubsystem::GetResourceSizeEx(FResourceSizeEx& CumulativeResourceSize)
+{
+	Super::GetResourceSizeEx(CumulativeResourceSize);
+
+	const SIZE_T MyExtraSize = Entities.GetAllocatedSize() + 
+		EntityFreeIndexList.GetAllocatedSize() +
+		DeferredCommandBuffer->GetAllocatedSize() +
+		ComponentHashToArchetypeMap.GetAllocatedSize() +
+		ComponentTypeToArchetypeMap.GetAllocatedSize();
+	CumulativeResourceSize.AddDedicatedSystemMemoryBytes(MyExtraSize);
+
+	for (const auto& KVP : ComponentHashToArchetypeMap)
+	{
+		for (const TSharedPtr<FArchetypeData>& ArchetypePtr : KVP.Value)
+		{
+			CumulativeResourceSize.AddDedicatedSystemMemoryBytes(ArchetypePtr->GetAllocatedSize());
+		}
+	}
+}
+
+void UMassEntitySubsystem::Initialize(FSubsystemCollectionBase& Collection)
+{
+	// Index 0 is reserved so we can treat that index as an invalid entity handle
+	Entities.Add();
+	SerialNumberGenerator.fetch_add(FMath::Max(1,NumReservedEntities));
+
+	DeferredCommandBuffer = MakeShareable(new FLWCCommandBuffer());
+
+	// initialize the bitsets
+
+	FLWComponentBitSet Components;
+	FLWTagBitSet Tags;
+	FLWChunkComponentBitSet ChunkComponents;
+
+	for (TObjectIterator<UStruct> StructIt; StructIt; ++StructIt)
+	{
+		if (StructIt->IsChildOf(FLWComponentData::StaticStruct()))
+		{
+			Components.Add(*CastChecked<UScriptStruct>(*StructIt));
+		}
+		else if (StructIt->IsChildOf(FComponentTag::StaticStruct()))
+		{
+			Tags.Add(*CastChecked<UScriptStruct>(*StructIt));
+		}
+		else if (StructIt->IsChildOf(FLWChunkComponent::StaticStruct()))
+		{
+			ChunkComponents.Add(*CastChecked<UScriptStruct>(*StructIt));
+		}
+	}
+}
+
+void UMassEntitySubsystem::Deinitialize()
+{
+	//@TODO: Should actually do this...
+}
+
+FArchetypeHandle UMassEntitySubsystem::CreateArchetype(TConstArrayView<const UScriptStruct*> ComponentsAngTagsList)
+{
+	FLWChunkComponentBitSet ChunkComponents;
+	FLWTagBitSet Tags;
+	TArray<const UScriptStruct*, TInlineAllocator<16>> ComponentList;
+	ComponentList.Reserve(ComponentsAngTagsList.Num());
+
+	for (const UScriptStruct* Type : ComponentsAngTagsList)
+	{
+		if (Type->IsChildOf(FLWComponentData::StaticStruct()))
+		{
+			ComponentList.Add(Type);
+		}
+		else if (Type->IsChildOf(FComponentTag::StaticStruct()))
+		{
+			Tags.Add(*Type);
+		}
+		else if (Type->IsChildOf(FLWChunkComponent::StaticStruct()))
+		{
+			ChunkComponents.Add(*Type);
+		}
+		else
+		{
+			UE_LOG(LogAggregateTicking, Warning, TEXT("%s: %s is not a valid component nor tag type. Ignoring.")
+				, ANSI_TO_TCHAR(__FUNCTION__), *GetNameSafe(Type));
+		}
+	}
+
+	return CreateArchetype(FLWCompositionDescriptor(FLWComponentBitSet(ComponentList), Tags, ChunkComponents));
+}
+
+FArchetypeHandle UMassEntitySubsystem::CreateArchetype(TConstArrayView<const UScriptStruct*> ComponentList, const FLWTagBitSet& Tags, const FLWChunkComponentBitSet& ChunkComponents)
+{
+	return CreateArchetype(FLWCompositionDescriptor(ComponentList, Tags, ChunkComponents));
+}
+
+FArchetypeHandle UMassEntitySubsystem::CreateArchetype(const TSharedPtr<FArchetypeData>& SourceArchetype, const FLWComponentBitSet& NewComponents)
+{
+	check(SourceArchetype.IsValid());
+	checkf(NewComponents.IsEmpty() == false, TEXT("%s Adding an empty component list to an archetype is not supported."), ANSI_TO_TCHAR(__FUNCTION__));
+
+	return CreateArchetype(FLWCompositionDescriptor(NewComponents + SourceArchetype->GetComponentBitSet(), SourceArchetype->GetTagBitSet(), SourceArchetype->GetChunkComponentBitSet()));
+}
+
+FArchetypeHandle UMassEntitySubsystem::CreateArchetype(const FLWCompositionDescriptor& Composition)
+{
+	const uint32 TypeHash = Composition.CalculateHash();
+
+	TArray<TSharedPtr<FArchetypeData>>& HashRow = ComponentHashToArchetypeMap.FindOrAdd(TypeHash);
+
+	FArchetypeHandle Result;
+	for (TSharedPtr<FArchetypeData>& Ptr : HashRow)
+	{
+		if (Ptr->IsEquivalent(Composition))
+		{
+			Result.DataPtr = Ptr;
+			break;
+		}
+	}
+
+	if (!Result.DataPtr.IsValid())
+	{
+		// Create a new archetype
+		TSharedPtr<FArchetypeData> NewArchetype = MakeShareable(new FArchetypeData);
+		NewArchetype->Initialize(Composition.Components, Composition.Tags, Composition.ChunkComponents);
+		HashRow.Add(NewArchetype);
+
+		for (const FArchetypeComponentConfig& ComponentConfig : NewArchetype->GetComponentConfigs())
+		{
+			checkSlow(ComponentConfig.ComponentType)
+			ComponentTypeToArchetypeMap.FindOrAdd(ComponentConfig.ComponentType).Add(NewArchetype);
+		}
+
+		Result.DataPtr = NewArchetype;
+
+		++ArchetypeDataVersion;
+	}
+
+	return MoveTemp(Result);
+}
+
+FArchetypeHandle UMassEntitySubsystem::InternalCreateSiblingArchetype(const TSharedPtr<FArchetypeData>& SourceArchetype, const FLWTagBitSet& OverrideTags)
+{
+	checkSlow(SourceArchetype.IsValid());
+	const FArchetypeData& SourceArchetypeRef = *SourceArchetype.Get();
+	const uint32 TypeHash = FLWCompositionDescriptor::CalculateHash(SourceArchetypeRef.GetComponentBitSet(), OverrideTags, SourceArchetypeRef.GetChunkComponentBitSet());
+
+	TArray<TSharedPtr<FArchetypeData>>& HashRow = ComponentHashToArchetypeMap.FindOrAdd(TypeHash);
+
+	FArchetypeHandle Result;
+	for (TSharedPtr<FArchetypeData>& Ptr : HashRow)
+	{
+		if (Ptr->IsEquivalent(SourceArchetypeRef.GetComponentBitSet(), OverrideTags, SourceArchetypeRef.GetChunkComponentBitSet()))
+		{
+			Result.DataPtr = Ptr;
+			break;
+		}
+	}
+
+	if (!Result.DataPtr.IsValid())
+	{
+		// Create a new archetype
+		TSharedPtr<FArchetypeData> NewArchetype = MakeShareable(new FArchetypeData);
+		NewArchetype->InitializeWithSibling(SourceArchetypeRef, OverrideTags);
+		HashRow.Add(NewArchetype);
+
+		for (const FArchetypeComponentConfig& ComponentConfig : NewArchetype->GetComponentConfigs())
+		{
+			checkSlow(ComponentConfig.ComponentType)
+			ComponentTypeToArchetypeMap.FindOrAdd(ComponentConfig.ComponentType).Add(NewArchetype);
+		}
+
+		Result.DataPtr = NewArchetype;
+
+		++ArchetypeDataVersion;
+	}
+
+	return MoveTemp(Result);
+}
+
+FArchetypeHandle UMassEntitySubsystem::GetArchetypeForEntity(FLWEntity Entity) const
+{
+	FArchetypeHandle Result;
+	if (IsEntityValid(Entity))
+	{
+		Result.DataPtr = Entities[Entity.Index].CurrentArchetype;
+	}
+	return Result;
+}
+
+void UMassEntitySubsystem::ForEachArchetypeComponentType(const FArchetypeHandle Archetype, TFunction< void(const UScriptStruct* /*ComponentType*/)> Function)
+{
+	check(Archetype.DataPtr.IsValid());
+	const FArchetypeData& ArchetypeData = *Archetype.DataPtr.Get();
+	ArchetypeData.ForEachComponentType(Function);
+}
+
+void UMassEntitySubsystem::SetDefaultChunkComponentValue(const FArchetypeHandle Archetype, FConstStructView InstancedStruct)
+{
+	check(Archetype.DataPtr.IsValid());
+	FArchetypeData& ArchetypeData = *Archetype.DataPtr.Get();
+	ArchetypeData.SetDefaultChunkComponentValue(InstancedStruct);
+}
+
+void UMassEntitySubsystem::DoEntityCompaction(const double TimeAllowed)
+{
+	const double TimeAllowedEnd = FPlatformTime::Seconds() + TimeAllowed;
+
+	bool bReachedTimeLimit = false;
+	for (const auto& KVP : ComponentHashToArchetypeMap)
+	{
+		for (const TSharedPtr<FArchetypeData>& ArchetypePtr : KVP.Value)
+		{
+			const double TimeAllowedLeft = TimeAllowedEnd - FPlatformTime::Seconds();
+			bReachedTimeLimit = TimeAllowedLeft <= 0.0;
+			if (bReachedTimeLimit)
+			{
+				break;
+			}
+			ArchetypePtr->CompactEntities(TimeAllowedLeft);
+		}
+		if (bReachedTimeLimit)
+		{
+			break;
+		}
+	}
+}
+
+FLWEntity UMassEntitySubsystem::CreateEntity(const FArchetypeHandle Archetype)
+{
+	check(Archetype.DataPtr.IsValid());
+
+	FLWEntity Entity = ReserveEntity();
+	InternalBuildEntity(Entity, Archetype);
+	return Entity;
+}
+
+FLWEntity UMassEntitySubsystem::CreateEntity(TConstArrayView<FInstancedStruct> ComponentInstanceList)
+{
+	check(ComponentInstanceList.Num() > 0);
+
+	FArchetypeHandle Archetype = CreateArchetype(FLWCompositionDescriptor(ComponentInstanceList, FLWTagBitSet(), FLWChunkComponentBitSet()));
+	check(Archetype.DataPtr.IsValid());
+
+	FLWEntity Entity = ReserveEntity();
+	InternalBuildEntity(Entity, Archetype);
+
+	FEntityData& EntityData = Entities[Entity.Index];
+	EntityData.CurrentArchetype->SetComponentsData(Entity, ComponentInstanceList);
+
+	return Entity;
+}
+
+FLWEntity UMassEntitySubsystem::ReserveEntity()
+{
+	// @todo: Need to add thread safety to the reservation of an entity
+	FLWEntity Result;
+	Result.Index = (EntityFreeIndexList.Num() > 0) ? EntityFreeIndexList.Pop(/*bAllowShrinking=*/ false) : Entities.Add();
+	Result.SerialNumber = SerialNumberGenerator.fetch_add(1);
+	Entities[Result.Index].SerialNumber = Result.SerialNumber;
+
+	return Result;
+}
+
+void UMassEntitySubsystem::ReleaseReservedEntity(FLWEntity Entity)
+{
+	checkf(!IsEntityBuilt(Entity), TEXT("Entity is already built, use DestroyEntity() instead"));
+
+	InternalReleaseEntity(Entity);
+}
+
+void UMassEntitySubsystem::BuildEntity(FLWEntity Entity, FArchetypeHandle Archetype)
+{
+	checkf(!IsEntityBuilt(Entity), TEXT("Expecting an entity that is not already built"));
+	check(Archetype.DataPtr.IsValid());
+
+	InternalBuildEntity(Entity, Archetype);
+}
+
+void UMassEntitySubsystem::BuildEntity(FLWEntity Entity, TConstArrayView<FInstancedStruct> ComponentInstanceList)
+{
+	check(ComponentInstanceList.Num() > 0);
+	checkf(!IsEntityBuilt(Entity), TEXT("Expecting an entity that is not already built"));
+
+	FArchetypeHandle Archetype = CreateArchetype(FLWCompositionDescriptor(ComponentInstanceList, FLWTagBitSet(), FLWChunkComponentBitSet()));
+	check(Archetype.DataPtr.IsValid());
+
+	InternalBuildEntity(Entity, Archetype);
+
+	FEntityData& EntityData = Entities[Entity.Index];
+	EntityData.CurrentArchetype->SetComponentsData(Entity, ComponentInstanceList);
+}
+
+int32 UMassEntitySubsystem::BatchCreateEntities(const FArchetypeHandle Archetype, const int32 Count, TArray<FLWEntity>& OutEntities)
+{
+	FArchetypeData* ArchetypePtr = Archetype.DataPtr.Get();
+	check(ArchetypePtr);
+	check(Count >= 0);
+	
+	int32 Index = OutEntities.Num();
+	OutEntities.AddDefaulted(Count);
+
+	// @todo optimize
+	for (; Index < OutEntities.Num(); ++Index)
+	{
+		FLWEntity& Result = OutEntities[Index];
+		Result.Index = (EntityFreeIndexList.Num() > 0) ? EntityFreeIndexList.Pop(/*bAllowShrinking=*/ false) : Entities.Add();
+		Result.SerialNumber = SerialNumberGenerator++;
+
+		FEntityData& EntityData = Entities[Result.Index];
+		EntityData.CurrentArchetype = Archetype.DataPtr;
+		EntityData.SerialNumber = Result.SerialNumber;
+		
+		ArchetypePtr->AddEntity(Result);
+	}
+
+	return Count;
+}
+
+void UMassEntitySubsystem::DestroyEntity(FLWEntity Entity)
+{
+	checkf(IsProcessing() == false, TEXT("Synchronous API function %s called during mass processing. Use asynchronous API instead."), ANSI_TO_TCHAR(__FUNCTION__));
+	
+	CheckIfEntityIsActive(Entity);
+
+	FEntityData& EntityData = Entities[Entity.Index];
+	FArchetypeData* Archetype = EntityData.CurrentArchetype.Get();
+	check(Archetype);
+	Archetype->RemoveEntity(Entity);
+
+	InternalReleaseEntity(Entity);
+
+	if (Archetype->GetNumEntities() == 0)
+	{
+		//@TODO: Drop our reference to the archetype?
+	}
+}
+
+void UMassEntitySubsystem::BatchDestroyEntities(TConstArrayView<FLWEntity> InEntities)
+{
+	TRACE_CPUPROFILER_EVENT_SCOPE_STR("BatchDestroyEntities");
+
+	checkf(IsProcessing() == false, TEXT("Synchronous API function %s called during mass processing. Use asynchronous API instead."), ANSI_TO_TCHAR(__FUNCTION__));
+	
+	EntityFreeIndexList.Reserve(EntityFreeIndexList.Num() + InEntities.Num());
+
+	// @todo optimize, we can make savings by implementing Archetype->RemoveEntities()
+	for (const FLWEntity Entity : InEntities)
+	{
+		if (Entities.IsValidIndex(Entity.Index) == false)
+		{
+			continue;
+		}
+
+		FEntityData& EntityData = Entities[Entity.Index];
+		if (EntityData.SerialNumber != Entity.SerialNumber)
+		{
+			continue;
+		}
+
+		FArchetypeData* Archetype = EntityData.CurrentArchetype.Get();
+		check(Archetype);
+		Archetype->RemoveEntity(Entity);
+
+		EntityData.Reset();
+		EntityFreeIndexList.Add(Entity.Index);
+	}
+}
+
+void UMassEntitySubsystem::BatchDestroyEntityChunks(const FArchetypeChunkCollection& Chunks)
+{
+	TRACE_CPUPROFILER_EVENT_SCOPE_STR("BatchDestroyEntityChunks");
+
+	checkf(IsProcessing() == false, TEXT("Synchronous API function %s called during mass processing. Use asynchronous API instead."), ANSI_TO_TCHAR(__FUNCTION__));
+
+	check(Chunks.GetArchetype().IsValid());
+	TArray<FLWEntity> EntitiesRemoved;
+	Chunks.GetArchetype().DataPtr->BatchDestroyEntityChunks(Chunks, EntitiesRemoved);
+	
+	EntityFreeIndexList.Reserve(EntityFreeIndexList.Num() + EntitiesRemoved.Num());
+
+	for (const FLWEntity& Entity : EntitiesRemoved)
+	{
+		check(Entities.IsValidIndex(Entity.Index));
+
+		FEntityData& EntityData = Entities[Entity.Index];
+		if (EntityData.SerialNumber == Entity.SerialNumber)
+		{
+			EntityData.Reset();
+			EntityFreeIndexList.Add(Entity.Index);
+		}
+	}
+}
+
+void UMassEntitySubsystem::AddComponentToEntity(FLWEntity Entity, const UScriptStruct* ComponentType)
+{
+	checkf(ComponentType, TEXT("Null component type passed in to %s"), ANSI_TO_TCHAR(__FUNCTION__));
+	checkf(IsProcessing() == false, TEXT("Synchronous API function %s called during mass processing. Use asynchronous API instead."), ANSI_TO_TCHAR(__FUNCTION__));
+
+	CheckIfEntityIsActive(Entity);
+
+	InternalAddComponentListToEntityChecked(Entity, FLWComponentBitSet(*ComponentType));
+}
+
+void UMassEntitySubsystem::AddComponentListToEntity(FLWEntity Entity, TConstArrayView<const UScriptStruct*> ComponentList)
+{
+	CheckIfEntityIsActive(Entity);
+
+	InternalAddComponentListToEntityChecked(Entity, FLWComponentBitSet(ComponentList));
+}
+
+void UMassEntitySubsystem::AddCompositionToEntity_GetDelta(FLWEntity Entity, FLWCompositionDescriptor& Composition)
+{
+	CheckIfEntityIsActive(Entity);
+
+	FEntityData& EntityData = Entities[Entity.Index];
+	FArchetypeData* OldArchetype = EntityData.CurrentArchetype.Get();
+	check(OldArchetype);
+
+	Composition.Components -= OldArchetype->GetCompositionDescriptor().Components;
+	Composition.Tags -= OldArchetype->GetCompositionDescriptor().Tags;
+
+	if (Composition.IsEmpty() == false)
+	{
+		FLWCompositionDescriptor NewDescriptor = OldArchetype->GetCompositionDescriptor();
+		NewDescriptor.Components += Composition.Components;
+		NewDescriptor.Tags+= Composition.Tags;
+
+		const FArchetypeHandle NewArchetypeHandle = CreateArchetype(NewDescriptor);
+
+		if (ensure(NewArchetypeHandle.DataPtr != EntityData.CurrentArchetype))
+		{
+			// Move the entity over
+			FArchetypeData* NewArchetype = NewArchetypeHandle.DataPtr.Get();
+			check(NewArchetype);
+			EntityData.CurrentArchetype->MoveEntityToAnotherArchetype(Entity, *NewArchetype);
+			EntityData.CurrentArchetype = NewArchetypeHandle.DataPtr;
+		}
+	}
+}
+
+void UMassEntitySubsystem::RemoveCompositionFromEntity(FLWEntity Entity, const FLWCompositionDescriptor& InDescriptor)
+{
+	CheckIfEntityIsActive(Entity);
+
+	FEntityData& EntityData = Entities[Entity.Index];
+	FArchetypeData* OldArchetype = EntityData.CurrentArchetype.Get();
+	check(OldArchetype);
+
+	FLWCompositionDescriptor NewDescriptor = OldArchetype->GetCompositionDescriptor();
+	NewDescriptor.Components -= InDescriptor.Components;
+	NewDescriptor.Tags -= InDescriptor.Tags;
+
+	if (NewDescriptor.IsEmpty() == false)
+	{
+		const FArchetypeHandle NewArchetypeHandle = CreateArchetype(NewDescriptor);
+
+		if (ensure(NewArchetypeHandle.DataPtr != EntityData.CurrentArchetype))
+		{
+			// Move the entity over
+			FArchetypeData* NewArchetype = NewArchetypeHandle.DataPtr.Get();
+			check(NewArchetype);
+			EntityData.CurrentArchetype->MoveEntityToAnotherArchetype(Entity, *NewArchetype);
+			EntityData.CurrentArchetype = NewArchetypeHandle.DataPtr;
+		}
+	}
+}
+
+void UMassEntitySubsystem::InternalBuildEntity(FLWEntity Entity, const FArchetypeHandle Archetype)
+{
+	FEntityData& EntityData = Entities[Entity.Index];
+	EntityData.CurrentArchetype = Archetype.DataPtr;
+	EntityData.CurrentArchetype->AddEntity(Entity);
+}
+
+void UMassEntitySubsystem::InternalReleaseEntity(FLWEntity Entity)
+{
+	FEntityData& EntityData = Entities[Entity.Index];
+	EntityData.Reset();
+	EntityFreeIndexList.Add(Entity.Index);
+}
+
+void UMassEntitySubsystem::InternalAddComponentListToEntityChecked(FLWEntity Entity, const FLWComponentBitSet& InComponents)
+{
+	FEntityData& EntityData = Entities[Entity.Index];
+	FArchetypeData* OldArchetype = EntityData.CurrentArchetype.Get();
+	check(OldArchetype);
+	
+	ensureMsgf(OldArchetype->GetComponentBitSet().HasAny(InComponents) == false, TEXT("Trying to add a new component type to an entity, but it already has some of them."));
+
+	const FLWComponentBitSet NewComponents = InComponents - OldArchetype->GetComponentBitSet();
+	if (NewComponents.IsEmpty() == false)
+	{
+		InternalAddComponentListToEntity(Entity, NewComponents);
+	}
+}
+
+void UMassEntitySubsystem::InternalAddComponentListToEntity(FLWEntity Entity, const FLWComponentBitSet& NewComponents)
+{
+	checkf(NewComponents.IsEmpty() == false, TEXT("%s is intended for internal calls with non empty NewComponents parameter"), ANSI_TO_TCHAR(__FUNCTION__));
+	check(Entities.IsValidIndex(Entity.Index));
+	FEntityData& EntityData = Entities[Entity.Index];
+	check(EntityData.CurrentArchetype.IsValid());
+
+	// fetch or create the new archetype
+	const FArchetypeHandle NewArchetypeHandle = CreateArchetype(EntityData.CurrentArchetype, NewComponents);
+
+	if (NewArchetypeHandle.DataPtr != EntityData.CurrentArchetype)
+	{
+		// Move the entity over
+		FArchetypeData* NewArchetype = NewArchetypeHandle.DataPtr.Get();
+		check(NewArchetype);
+		EntityData.CurrentArchetype->MoveEntityToAnotherArchetype(Entity, *NewArchetype);
+		EntityData.CurrentArchetype = NewArchetypeHandle.DataPtr;
+	}
+}
+
+void UMassEntitySubsystem::AddComponentInstanceListToEntity(FLWEntity Entity, TConstArrayView<FInstancedStruct> ComponentInstanceList)
+{
+	checkf(IsProcessing() == false, TEXT("Synchronous API function %s called during mass processing. Use asynchronous API instead."), ANSI_TO_TCHAR(__FUNCTION__));
+
+	CheckIfEntityIsActive(Entity);
+	checkf(ComponentInstanceList.Num() > 0, TEXT("Need to specify at least one component instances for this operation"));
+
+	InternalAddComponentListToEntityChecked(Entity, FLWComponentBitSet(ComponentInstanceList));
+
+	FEntityData& EntityData = Entities[Entity.Index];
+	EntityData.CurrentArchetype->SetComponentsData(Entity, ComponentInstanceList);
+}
+
+void UMassEntitySubsystem::RemoveComponentFromEntity(FLWEntity Entity, const UScriptStruct* ComponentType)
+{
+	RemoveComponentListFromEntity(Entity, MakeArrayView(&ComponentType, 1));
+}
+
+void UMassEntitySubsystem::RemoveComponentListFromEntity(FLWEntity Entity, TConstArrayView<const UScriptStruct*> ComponentList)
+{
+	checkf(IsProcessing() == false, TEXT("Synchronous API function %s called during mass processing. Use asynchronous API instead."), ANSI_TO_TCHAR(__FUNCTION__));
+
+	CheckIfEntityIsActive(Entity);
+	
+	FEntityData& EntityData = Entities[Entity.Index];
+	FArchetypeData* OldArchetype = EntityData.CurrentArchetype.Get();
+	check(OldArchetype);
+
+	const FLWComponentBitSet ComponentsToRemove(ComponentList);
+
+	if (ensureMsgf(OldArchetype->GetComponentBitSet().HasAny(ComponentsToRemove), TEXT("Trying to remove a list of components from an entity but none of the components given was found.")))
+	{
+		// If all the components got removed this will result in fetching of the empty archetype
+		const FArchetypeHandle NewArchetypeHandle = CreateArchetype(FLWCompositionDescriptor(OldArchetype->GetComponentBitSet() - ComponentsToRemove, OldArchetype->GetTagBitSet(), OldArchetype->GetChunkComponentBitSet()));
+
+		// Move the entity over
+		FArchetypeData* NewArchetype = NewArchetypeHandle.DataPtr.Get();
+		check(NewArchetype);
+		OldArchetype->MoveEntityToAnotherArchetype(Entity, *NewArchetype);
+		EntityData.CurrentArchetype = NewArchetypeHandle.DataPtr;
+	}
+}
+
+void UMassEntitySubsystem::SwapTagsForEntity(FLWEntity Entity, const UScriptStruct* OldTagType, const UScriptStruct* NewTagType)
+{
+	checkf(IsProcessing() == false, TEXT("Synchronous API function %s called during mass processing. Use asynchronous API instead."), ANSI_TO_TCHAR(__FUNCTION__));
+
+	CheckIfEntityIsActive(Entity);
+
+	checkf((OldTagType != nullptr) && OldTagType->IsChildOf(FComponentTag::StaticStruct()), TEXT("%s works only with tags while '%s' is not one."), ANSI_TO_TCHAR(__FUNCTION__), *GetPathNameSafe(OldTagType));
+	checkf((NewTagType != nullptr) && NewTagType->IsChildOf(FComponentTag::StaticStruct()), TEXT("%s works only with tags while '%s' is not one."), ANSI_TO_TCHAR(__FUNCTION__), *GetPathNameSafe(NewTagType));
+
+	FEntityData& EntityData = Entities[Entity.Index];
+	FArchetypeData* CurrentArchetype = EntityData.CurrentArchetype.Get();
+	check(CurrentArchetype);
+
+	FLWTagBitSet NewTagBitSet = CurrentArchetype->GetTagBitSet();
+	NewTagBitSet.Remove(*OldTagType);
+	NewTagBitSet.Add(*NewTagType);
+	
+	if (NewTagBitSet != CurrentArchetype->GetTagBitSet())
+	{
+		FArchetypeHandle NewArchetypeHandle = InternalCreateSiblingArchetype(EntityData.CurrentArchetype, NewTagBitSet);
+		checkSlow(NewArchetypeHandle.IsValid());
+
+		// Move the entity over
+		EntityData.CurrentArchetype->MoveEntityToAnotherArchetype(Entity, *NewArchetypeHandle.DataPtr.Get());
+		EntityData.CurrentArchetype = NewArchetypeHandle.DataPtr;
+	}
+}
+
+void UMassEntitySubsystem::AddTagToEntity(FLWEntity Entity, const UScriptStruct* TagType)
+{
+	checkf((TagType != nullptr) && TagType->IsChildOf(FComponentTag::StaticStruct()), TEXT("%s works only with tags while '%s' is not one."), ANSI_TO_TCHAR(__FUNCTION__), *GetPathNameSafe(TagType));
+
+	CheckIfEntityIsActive(Entity);
+
+	FEntityData& EntityData = Entities[Entity.Index];
+	FArchetypeData* CurrentArchetype = EntityData.CurrentArchetype.Get();
+	check(CurrentArchetype);
+
+	if (CurrentArchetype->HasTagType(TagType) == false)
+	{
+		//FLWTagBitSet NewTags = CurrentArchetype->GetTagBitSet() - *TagType;
+		FLWTagBitSet NewTags = CurrentArchetype->GetTagBitSet();
+		NewTags.Add(*TagType);
+		FArchetypeHandle NewArchetypeHandle = InternalCreateSiblingArchetype(EntityData.CurrentArchetype, NewTags);
+		checkSlow(NewArchetypeHandle.IsValid());
+
+		// Move the entity over
+		EntityData.CurrentArchetype->MoveEntityToAnotherArchetype(Entity, *NewArchetypeHandle.DataPtr.Get());
+		EntityData.CurrentArchetype = NewArchetypeHandle.DataPtr;
+	}
+}
+	
+void UMassEntitySubsystem::RemoveTagFromEntity(FLWEntity Entity, const UScriptStruct* TagType)
+{
+	checkf((TagType != nullptr) && TagType->IsChildOf(FComponentTag::StaticStruct()), TEXT("%s works only with tags while '%s' is not one."), ANSI_TO_TCHAR(__FUNCTION__), *GetPathNameSafe(TagType));
+
+	CheckIfEntityIsActive(Entity);
+
+	FEntityData& EntityData = Entities[Entity.Index];
+	FArchetypeData* CurrentArchetype = EntityData.CurrentArchetype.Get();
+	check(CurrentArchetype);
+
+	if (CurrentArchetype->HasTagType(TagType))
+	{
+		// CurrentArchetype->GetTagBitSet() -  *TagType
+		FLWTagBitSet NewTags = CurrentArchetype->GetTagBitSet();
+		NewTags.Remove(*TagType);
+		FArchetypeHandle NewArchetypeHandle = InternalCreateSiblingArchetype(EntityData.CurrentArchetype, NewTags);
+		checkSlow(NewArchetypeHandle.IsValid());
+
+		// Move the entity over
+		EntityData.CurrentArchetype->MoveEntityToAnotherArchetype(Entity, *NewArchetypeHandle.DataPtr.Get());
+		EntityData.CurrentArchetype = NewArchetypeHandle.DataPtr;
+	}
+}
+
+
+void UMassEntitySubsystem::MoveEntityToAnotherArchetype(FLWEntity Entity, FArchetypeHandle NewArchetypeHandle)
+{
+	CheckIfEntityIsActive(Entity);
+
+	FArchetypeData* NewArchetype = NewArchetypeHandle.DataPtr.Get();
+	check(NewArchetype);
+
+	// Move the entity over
+	FEntityData& EntityData = Entities[Entity.Index];
+	EntityData.CurrentArchetype->MoveEntityToAnotherArchetype(Entity, *NewArchetype);
+	EntityData.CurrentArchetype = NewArchetypeHandle.DataPtr;
+}
+
+void UMassEntitySubsystem::SetEntityComponentsValues(FLWEntity Entity, TArrayView<const FInstancedStruct> ComponentInstanceList)
+{
+	CheckIfEntityIsActive(Entity);
+
+	FEntityData& EntityData = Entities[Entity.Index];
+	EntityData.CurrentArchetype->SetComponentsData(Entity, ComponentInstanceList);
+}
+
+void UMassEntitySubsystem::BatchSetEntityComponentsValues(const FArchetypeChunkCollection& SparseEntities, TArrayView<const FInstancedStruct> ComponentInstanceList)
+{
+	FArchetypeData* Archetype = SparseEntities.GetArchetype().DataPtr.Get();
+	check(Archetype);
+
+	for (const FInstancedStruct& ComponentTemplate : ComponentInstanceList)
+	{
+		Archetype->SetComponentData(SparseEntities, ComponentTemplate);
+	}
+}
+
+void* UMassEntitySubsystem::InternalGetComponentDataChecked(FLWEntity Entity, const UScriptStruct* ComponentType) const
+{
+	CheckIfEntityIsActive(Entity);
+
+	checkf((ComponentType != nullptr) && ComponentType->IsChildOf(FLWComponentData::StaticStruct()), TEXT("InternalGetComponentDataChecked called with an invalid component type '%s'"), *GetPathNameSafe(ComponentType));
+	const FEntityData& EntityData = Entities[Entity.Index];
+	return EntityData.CurrentArchetype->GetComponentDataForEntityChecked(ComponentType, Entity.Index);
+}
+
+void* UMassEntitySubsystem::InternalGetComponentDataPtr(FLWEntity Entity, const UScriptStruct* ComponentType) const
+{
+	CheckIfEntityIsActive(Entity);
+	checkf((ComponentType != nullptr) && ComponentType->IsChildOf(FLWComponentData::StaticStruct()), TEXT("InternalGetComponentData called with an invalid component type '%s'"), *GetPathNameSafe(ComponentType));
+	const FEntityData& EntityData = Entities[Entity.Index];
+	return EntityData.CurrentArchetype->GetComponentDataForEntity(ComponentType, Entity.Index);
+}
+
+bool UMassEntitySubsystem::IsEntityValid(FLWEntity Entity) const
+{
+	return (Entity.Index > 0) && Entities.IsValidIndex(Entity.Index) && (Entities[Entity.Index].SerialNumber == Entity.SerialNumber);
+}
+
+bool UMassEntitySubsystem::IsEntityBuilt(FLWEntity Entity) const
+{
+	CheckIfEntityIsValid(Entity);
+	return Entities[Entity.Index].CurrentArchetype.IsValid();
+}
+
+void UMassEntitySubsystem::CheckIfEntityIsValid(FLWEntity Entity) const
+{
+	checkf(IsEntityValid(Entity), TEXT("Invalid entity (ID: %d, SN:%d, %s)"), Entity.Index, Entity.SerialNumber,
+		   (Entity.Index == 0) ? TEXT("was never initialized") : TEXT("already destroyed"));
+}
+
+void UMassEntitySubsystem::CheckIfEntityIsActive(FLWEntity Entity) const
+{
+	checkf(IsEntityBuilt(Entity), TEXT("Entity not yet created(ID: %d, SN:%d)"));
+}
+
+void UMassEntitySubsystem::GetValidArchetypes(const FLWComponentQuery& Query, TArray<FArchetypeHandle>& OutValidArchetypes)
+{
+	//@TODO: Not optimized yet, but we call this rarely now, so not a big deal.
+
+	// First get set of all archetypes that contain *any* component
+	TSet<TSharedPtr<FArchetypeData>> AnyArchetypes;
+	for (const FLWComponentRequirement& Requirement : Query.GetRequirements())
+	{
+		check(Requirement.StructType);
+		if (Requirement.Presence != ELWComponentPresence::None)
+		{
+			if (TArray<TSharedPtr<FArchetypeData>>* pData = ComponentTypeToArchetypeMap.Find(Requirement.StructType))
+			{
+				AnyArchetypes.Append(*pData);
+			}
+		}
+	}
+
+	// Then verify that they contain *all* required components
+	for (TSharedPtr<FArchetypeData>& ArchetypePtr : AnyArchetypes)
+	{
+		FArchetypeData& Archetype = *(ArchetypePtr.Get());
+
+		if (Archetype.GetTagBitSet().HasAll(Query.GetRequiredAllTags()) == false)
+		{
+			// missing some required tags, skip.
+#if WITH_AGGREGATETICKING_DEBUG
+			const FLWTagBitSet UnsatisfiedTags = Query.GetRequiredAllTags() - Archetype.GetTagBitSet();
+			FStringOutputDevice Description;
+			UnsatisfiedTags.DebugGetStringDesc(Description);
+			UE_LOG(LogAggregateTicking, VeryVerbose, TEXT("Archetype did not match due to missing tags: %s")
+				, *Description);
+#endif // WITH_AGGREGATETICKING_DEBUG
+			continue;
+		}
+
+		if (Archetype.GetTagBitSet().HasNone(Query.GetRequiredNoneTags()) == false)
+		{
+			// has some tags required to be absent
+#if WITH_AGGREGATETICKING_DEBUG
+			const FLWTagBitSet UnwantedTags = Query.GetRequiredAllTags().GetOverlap(Archetype.GetTagBitSet());
+			FStringOutputDevice Description;
+			UnwantedTags.DebugGetStringDesc(Description);
+			UE_LOG(LogAggregateTicking, VeryVerbose, TEXT("Archetype has tags required absent: %s")
+				, *Description);
+#endif // WITH_AGGREGATETICKING_DEBUG
+			continue;
+		}
+
+		if (Query.GetRequiredAnyTags().IsEmpty() == false 
+			&& Archetype.GetTagBitSet().HasAny(Query.GetRequiredAnyTags()) == false)
+		{
+#if WITH_AGGREGATETICKING_DEBUG
+			FStringOutputDevice Description;
+			Query.GetRequiredAnyTags().DebugGetStringDesc(Description);
+			UE_LOG(LogAggregateTicking, VeryVerbose, TEXT("Archetype did not match due to missing \'any\' tags: %s")
+				, *Description);
+#endif // WITH_AGGREGATETICKING_DEBUG
+			continue;
+		}
+		
+		if (Archetype.GetComponentBitSet().HasAll(Query.GetRequiredAllComponents()) == false)
+		{
+			// missing some required components, skip.
+#if WITH_AGGREGATETICKING_DEBUG
+			const FLWComponentBitSet UnsatisfiedComponents = Query.GetRequiredAllComponents() - Archetype.GetComponentBitSet();
+			FStringOutputDevice Description;
+			UnsatisfiedComponents.DebugGetStringDesc(Description);
+			UE_LOG(LogAggregateTicking, VeryVerbose, TEXT("Archetype did not match due to missing Components: %s")
+				, *Description);
+#endif // WITH_AGGREGATETICKING_DEBUG
+			continue;
+		}
+
+		if (Archetype.GetComponentBitSet().HasNone(Query.GetRequiredNoneComponents()) == false)
+		{
+			// has some Components required to be absent
+#if WITH_AGGREGATETICKING_DEBUG
+			const FLWComponentBitSet UnwantedComponents = Query.GetRequiredAllComponents().GetOverlap(Archetype.GetComponentBitSet());
+			FStringOutputDevice Description;
+			UnwantedComponents.DebugGetStringDesc(Description);
+			UE_LOG(LogAggregateTicking, VeryVerbose, TEXT("Archetype has Components required absent: %s")
+				, *Description);
+#endif // WITH_AGGREGATETICKING_DEBUG
+			continue;
+		}
+
+		if (Query.GetRequiredAnyComponents().IsEmpty() == false 
+			&& Archetype.GetComponentBitSet().HasAny(Query.GetRequiredAnyComponents()) == false)
+		{
+#if WITH_AGGREGATETICKING_DEBUG
+			FStringOutputDevice Description;
+			Query.GetRequiredAnyComponents().DebugGetStringDesc(Description);
+			UE_LOG(LogAggregateTicking, VeryVerbose, TEXT("Archetype did not match due to missing \'any\' components: %s")
+				, *Description);
+#endif // WITH_AGGREGATETICKING_DEBUG
+			continue;
+		}
+
+		if (Archetype.GetChunkComponentBitSet().HasAll(Query.GetRequiredAllChunkComponents()) == false)
+		{
+			// missing some required components, skip.
+#if WITH_AGGREGATETICKING_DEBUG
+			const FLWChunkComponentBitSet UnsatisfiedComponents = Query.GetRequiredAllChunkComponents() - Archetype.GetChunkComponentBitSet();
+			FStringOutputDevice Description;
+			UnsatisfiedComponents.DebugGetStringDesc(Description);
+			UE_LOG(LogAggregateTicking, VeryVerbose, TEXT("Archetype did not match due to missing Chunk Components: %s")
+				, *Description);
+#endif // WITH_AGGREGATETICKING_DEBUG
+			continue;
+		}
+
+		if (Archetype.GetChunkComponentBitSet().HasNone(Query.GetRequiredNoneChunkComponents()) == false)
+		{
+			// has some Components required to be absent
+#if WITH_AGGREGATETICKING_DEBUG
+			const FLWChunkComponentBitSet UnwantedComponents = Query.GetRequiredNoneChunkComponents().GetOverlap(Archetype.GetChunkComponentBitSet());
+			FStringOutputDevice Description;
+			UnwantedComponents.DebugGetStringDesc(Description);
+			UE_LOG(LogAggregateTicking, VeryVerbose, TEXT("Archetype has Chunk Components required absent: %s")
+				, *Description);
+#endif // WITH_AGGREGATETICKING_DEBUG
+			continue;
+		}
+
+		OutValidArchetypes.Add(ArchetypePtr);
+	}
+}
+
+FLWComponentSystemExecutionContext UMassEntitySubsystem::CreateExecutionContext(const float DeltaSeconds) const
+{
+	FLWComponentSystemExecutionContext ExecutionContext(DeltaSeconds);
+	ExecutionContext.SetDeferredCommandBuffer(DeferredCommandBuffer);
+	return MoveTemp(ExecutionContext);
+}
+
+#if WITH_AGGREGATETICKING_DEBUG
+void UMassEntitySubsystem::DebugPrintArchetypes(FOutputDevice& Ar) const
+{
+	Ar.Logf(ELogVerbosity::Log, TEXT("Listing archetypes contained in %s"), *GetPathNameSafe(this));
+
+	int32 NumBuckets = 0;
+	int32 NumArchetypes = 0;
+	int32 LongestArchetypeBucket = 0;
+	for (const auto& KVP : ComponentHashToArchetypeMap)
+	{
+		for (const TSharedPtr<FArchetypeData>& ArchetypePtr : KVP.Value)
+		{
+			ArchetypePtr->DebugPrintArchetype(Ar);
+		}
+
+		const int32 NumArchetypesInBucket = KVP.Value.Num();
+		LongestArchetypeBucket = FMath::Max(LongestArchetypeBucket, NumArchetypesInBucket);
+		NumArchetypes += NumArchetypesInBucket;
+		++NumBuckets;
+	}
+
+	Ar.Logf(ELogVerbosity::Log, TEXT("ComponentHashToArchetypeMap: %d archetypes across %d buckets, longest bucket is %d"),
+		NumArchetypes, NumBuckets, LongestArchetypeBucket);
+}
+
+void UMassEntitySubsystem::DebugPrintEntity(int32 Index, FOutputDevice& Ar, const TCHAR* InPrefix) const
+{
+	if (Index >= Entities.Num())
+	{
+		Ar.Logf(ELogVerbosity::Log, TEXT("Unable to list components values for out of range index in %s"), *GetPathNameSafe(this));
+		return;
+	}
+
+	const FEntityData& EntityData = Entities[Index];
+	if (!EntityData.IsValid())
+	{
+		Ar.Logf(ELogVerbosity::Log, TEXT("Unable to list components values for invalid entity in %s"), *GetPathNameSafe(this));
+	}
+
+	FLWEntity Entity;
+	Entity.Index = Index;
+	Entity.SerialNumber = EntityData.SerialNumber;
+	DebugPrintEntity(Entity, Ar, InPrefix);
+}
+
+void UMassEntitySubsystem::DebugPrintEntity(FLWEntity Entity, FOutputDevice& Ar, const TCHAR* InPrefix) const
+{
+	if (!IsEntityActive(Entity))
+	{
+		Ar.Logf(ELogVerbosity::Log, TEXT("Unable to list components values for invalid entity in %s"), *GetPathNameSafe(this));
+	}
+
+	Ar.Logf(ELogVerbosity::Log, TEXT("Listing components values for Entity[%s] in %s"), *Entity.DebugGetDescription(), *GetPathNameSafe(this));
+
+	const FEntityData& EntityData = Entities[Entity.Index];
+	FArchetypeData* Archetype = EntityData.CurrentArchetype.Get();
+	if (Archetype == nullptr)
+	{
+		Ar.Logf(ELogVerbosity::Log, TEXT("Unable to list components values for invalid entity in %s"), *GetPathNameSafe(this));
+	}
+	else
+	{
+		Archetype->DebugPrintEntity(Entity, Ar, InPrefix);
+	}
+}
+
+void UMassEntitySubsystem::DebugGetStringDesc(const FArchetypeHandle& Archetype, FOutputDevice& Ar)
+{
+	Ar.Logf(TEXT("%s"), Archetype.IsValid() ? *Archetype.DataPtr->DebugGetDescription() : TEXT("INVALID"));
+}
+
+void UMassEntitySubsystem::DebugGetArchetypesStringDetails(FOutputDevice& Ar, const bool bIncludeEmpty)
+{
+#if WITH_AGGREGATETICKING_DEBUG
+	Ar.SetAutoEmitLineTerminator(true);
+	for (auto Pair : ComponentHashToArchetypeMap)
+	{
+		Ar.Logf(ELogVerbosity::Log, TEXT("\n-----------------------------------\nHash: %u"), Pair.Key);
+		for (TSharedPtr<FArchetypeData> Archetype : Pair.Value)
+		{
+			if (Archetype.IsValid() && (bIncludeEmpty == true || Archetype->GetChunkCount() > 0))
+			{
+				Archetype->DebugPrintArchetype(Ar);
+				Ar.Logf(ELogVerbosity::Log, TEXT("+++++++++++++++++++++++++\n"));
+			}
+		}
+	}
+#endif // WITH_AGGREGATETICKING_DEBUG
+}
+
+void UMassEntitySubsystem::DebugGetArchetypeComponentTypes(const FArchetypeHandle& Archetype, TArray<const UScriptStruct*>& InOutComponentList) const
+{
+	if (Archetype.DataPtr.IsValid())
+	{
+		Archetype.DataPtr->GetCompositionDescriptor().Components.DebugGetStructTypes(InOutComponentList);
+	}
+}
+
+int32 UMassEntitySubsystem::DebugGetArchetypeEntitiesCount(const FArchetypeHandle& Archetype) const
+{
+	return Archetype.DataPtr.IsValid() ? Archetype.DataPtr->GetNumEntities() : 0;
+}
+
+void UMassEntitySubsystem::DebugRemoveAllEntities()
+{
+	for (int EntityIndex = NumReservedEntities; EntityIndex < Entities.Num(); ++EntityIndex)
+	{
+		FEntityData& EntityData = Entities[EntityIndex];
+		if (EntityData.IsValid() == false)
+		{
+			// already dead
+			continue;
+		}
+		const TSharedPtr<FArchetypeData>& Archetype = EntityData.CurrentArchetype;
+		FLWEntity Entity;
+		Entity.Index = EntityIndex;
+		Entity.SerialNumber = EntityData.SerialNumber;
+		Archetype->RemoveEntity(Entity);
+
+		EntityData.Reset();
+		EntityFreeIndexList.Add(EntityIndex);
+	}
+}
+
+void UMassEntitySubsystem::DebugGetArchetypeStrings(const FArchetypeHandle& Archetype, TArray<FName>& OutComponentNames, TArray<FName>& OutTagNames)
+{
+	if (Archetype.IsValid() == false)
+	{
+		return;
+	}
+
+	FArchetypeData& ArchetypeRef = *Archetype.DataPtr.Get();
+	
+	OutComponentNames.Reserve(ArchetypeRef.GetComponentConfigs().Num());
+	for (const FArchetypeComponentConfig& ComponentConfig : ArchetypeRef.GetComponentConfigs())
+	{
+		checkSlow(ComponentConfig.ComponentType);
+		OutComponentNames.Add(ComponentConfig.ComponentType->GetFName());
+	}
+
+	ArchetypeRef.GetTagBitSet().DebugGetIndividualNames(OutTagNames);
+}
+
+//////////////////////////////////////////////////////////////////////
+// Debug commands
+
+FAutoConsoleCommandWithWorldArgsAndOutputDevice GPrintArchetypesCmd(
+	TEXT("EntityManager.PrintArchetypes"),
+	TEXT("Prints information about all archetypes in the current world"),
+	FConsoleCommandWithWorldArgsAndOutputDeviceDelegate::CreateStatic(
+		[](const TArray<FString>& Params, UWorld* World, FOutputDevice& Ar)
+{
+	if (UMassEntitySubsystem* EntitySystem = World ? World->GetSubsystem<UMassEntitySubsystem>() : nullptr)
+	{
+		EntitySystem->DebugPrintArchetypes(Ar);
+	}
+	else
+	{
+		Ar.Logf(ELogVerbosity::Error, TEXT("Failed to find Entity Subsystem for world %s"), *GetPathNameSafe(World));
+	}
+}));
+#endif // WITH_AGGREGATETICKING_DEBUG
+
+//////////////////////////////////////////////////////////////////////
+// FLWComponentSystemExecutionContext
+
+void FLWComponentSystemExecutionContext::FlushDeferred(UMassEntitySubsystem& EntitySystem) const
+{
+	if (bFlushDeferredCommands && DeferredCommandBuffer)
+	{
+		DeferredCommandBuffer->ReplayBufferAgainstSystem(&EntitySystem);
+	}
+}
+
+void FLWComponentSystemExecutionContext::ClearExecutionData()
+{
+	ComponentViews.Reset();
+	CurrentArchetypesTagBitSet.Reset();
+	ChunkSerialModificationNumber = -1;
+}
+
+void FLWComponentSystemExecutionContext::SetCurrentArchetypeData(FArchetypeData& ArchetypeData)
+{
+	CurrentArchetypesTagBitSet = ArchetypeData.GetTagBitSet();
+}
+
+void FLWComponentSystemExecutionContext::SetChunkCollection(const FArchetypeChunkCollection& InChunkCollection)
+{
+	check(ChunkCollection.IsEmpty());
+	ChunkCollection = InChunkCollection;
+}
+
+void FLWComponentSystemExecutionContext::SetChunkCollection(FArchetypeChunkCollection&& InChunkCollection)
+{
+	check(ChunkCollection.IsEmpty());
+	ChunkCollection = MoveTemp(InChunkCollection);
+}
+
+void FLWComponentSystemExecutionContext::SetRequirements(TConstArrayView<FLWComponentRequirement> InRequirements, TConstArrayView<FLWComponentRequirement> InChunkRequirements)
+{ 
+	ComponentViews.Reset();
+	for (const FLWComponentRequirement& Requirement : InRequirements)
+	{
+		if (Requirement.RequiresBinding())
+		{
+			ComponentViews.Add(FComponentView(Requirement));
+		}
+	}
+
+	ChunkComponents.Reset();
+	for (const FLWComponentRequirement& Requirement : InChunkRequirements)
+	{
+		if (Requirement.RequiresBinding())
+		{
+			ChunkComponents.Add(FChunkComponentView(Requirement));
+		}
+	}
+}
+
diff --git a/Engine/Plugins/Runtime/MassEntity/Source/MassEntity/Private/MassEntitySystem.cpp b/Engine/Plugins/Runtime/MassEntity/Source/MassEntity/Private/MassEntitySystem.cpp
deleted file mode 100644
diff --git a/Engine/Plugins/Runtime/MassEntity/Source/MassEntity/Private/MassEntityTypes.cpp b/Engine/Plugins/Runtime/MassEntity/Source/MassEntity/Private/MassEntityTypes.cpp
--- a/Engine/Plugins/Runtime/MassEntity/Source/MassEntity/Private/MassEntityTypes.cpp
+++ b/Engine/Plugins/Runtime/MassEntity/Source/MassEntity/Private/MassEntityTypes.cpp
@@ -14,7 +14,7 @@ DEFINE_LOG_CATEGORY(LogPipe);
 //----------------------------------------------------------------------//
 //  FPipeContext
 //----------------------------------------------------------------------//
-FPipeContext::FPipeContext(UEntitySubsystem& InEntities, const float InDeltaSeconds)
+FPipeContext::FPipeContext(UMassEntitySubsystem& InEntities, const float InDeltaSeconds)
 	: EntitySubsystem(&InEntities), DeltaSeconds(InDeltaSeconds)
 {
 
diff --git a/Engine/Plugins/Runtime/MassEntity/Source/MassEntity/Private/MassExecutor.cpp b/Engine/Plugins/Runtime/MassEntity/Source/MassEntity/Private/MassExecutor.cpp
--- a/Engine/Plugins/Runtime/MassEntity/Source/MassEntity/Private/MassExecutor.cpp
+++ b/Engine/Plugins/Runtime/MassEntity/Source/MassEntity/Private/MassExecutor.cpp
@@ -97,7 +97,7 @@ void RunProcessorsView(TArrayView<UPipeProcessor*> Processors, FPipeContext& Pip
 	{
 		TRACE_CPUPROFILER_EVENT_SCOPE_STR("Execute Processors")
 		
-		UEntitySubsystem::FScopedProcessing ProcessingScope = PipeContext.EntitySubsystem->NewProcessingScope();
+		UMassEntitySubsystem::FScopedProcessing ProcessingScope = PipeContext.EntitySubsystem->NewProcessingScope();
 
 		for (UPipeProcessor* Proc : Processors)
 		{
@@ -117,7 +117,7 @@ void RunProcessorsView(TArrayView<UPipeProcessor*> Processors, FPipeContext& Pip
 
 struct FPipeExecutorDoneTask
 {
-	FPipeExecutorDoneTask(const FLWComponentSystemExecutionContext& InExecutionContext, UEntitySubsystem& InEntitySubsystem, TFunction<void()> InOnDoneNotification, const FString& InDebugName)
+	FPipeExecutorDoneTask(const FLWComponentSystemExecutionContext& InExecutionContext, UMassEntitySubsystem& InEntitySubsystem, TFunction<void()> InOnDoneNotification, const FString& InDebugName)
 		: ExecutionContext(InExecutionContext)
 		, EntitySubsystem(InEntitySubsystem)
 		, OnDoneNotification(InOnDoneNotification)
@@ -148,7 +148,7 @@ struct FPipeExecutorDoneTask
 	}
 private:
 	FLWComponentSystemExecutionContext ExecutionContext;
-	UEntitySubsystem& EntitySubsystem;
+	UMassEntitySubsystem& EntitySubsystem;
 	TFunction<void()> OnDoneNotification;
 	FString DebugName;
 };
@@ -165,7 +165,7 @@ FGraphEventRef TriggerParallelTasks(UPipeProcessor& Processor, FPipeContext& Pip
 
 	TRACE_CPUPROFILER_EVENT_SCOPE_STR("PipeExecutor RunParallel")
 
-	// not going through UEntitySubsystem::CreateExecutionContext on purpose - we do need a separate command buffer
+	// not going through UMassEntitySubsystem::CreateExecutionContext on purpose - we do need a separate command buffer
 	FLWComponentSystemExecutionContext ExecutionContext(PipeContext.DeltaSeconds);
 	ExecutionContext.SetDeferredCommandBuffer(MakeShareable(new FLWCCommandBuffer()));
 	ExecutionContext.SetFlushDeferredCommands(false);
diff --git a/Engine/Plugins/Runtime/MassEntity/Source/MassEntity/Private/MassProcessingPhase.cpp b/Engine/Plugins/Runtime/MassEntity/Source/MassEntity/Private/MassProcessingPhase.cpp
--- a/Engine/Plugins/Runtime/MassEntity/Source/MassEntity/Private/MassProcessingPhase.cpp
+++ b/Engine/Plugins/Runtime/MassEntity/Source/MassEntity/Private/MassProcessingPhase.cpp
@@ -6,7 +6,7 @@
 #include "MassEntitySettings.h"
 #include "MassProcessor.h"
 #include "MassExecutor.h"
-#include "MassEntitySystem.h"
+#include "MassEntitySubsystem.h"
 #include "VisualLogger/VisualLogger.h"
 #include "Engine/World.h"
 #include "LWCCommandBuffer.h"
@@ -172,7 +172,7 @@ void UPipeProcessingPhaseManager::InitializePhases(UObject& InProcessorOwner)
 
 void UPipeProcessingPhaseManager::Start(UWorld& World)
 {
-	EntitySubsystem = UPipeEntitySubsystem::GetCurrent(&World);
+	EntitySubsystem = UMassEntitySubsystem::GetCurrent(&World);
 
 	if (ensure(EntitySubsystem))
 	{
@@ -184,7 +184,7 @@ void UPipeProcessingPhaseManager::Start(UWorld& World)
 	}
 }
 
-void UPipeProcessingPhaseManager::Start(UPipeEntitySubsystem& InEntitySubsystem)
+void UPipeProcessingPhaseManager::Start(UMassEntitySubsystem& InEntitySubsystem)
 {
 	UWorld* World = InEntitySubsystem.GetWorld();
 	check(World);
diff --git a/Engine/Plugins/Runtime/MassEntity/Source/MassEntity/Private/MassProcessor.cpp b/Engine/Plugins/Runtime/MassEntity/Source/MassEntity/Private/MassProcessor.cpp
--- a/Engine/Plugins/Runtime/MassEntity/Source/MassEntity/Private/MassProcessor.cpp
+++ b/Engine/Plugins/Runtime/MassEntity/Source/MassEntity/Private/MassProcessor.cpp
@@ -49,7 +49,7 @@ namespace FPipeTweakables
 class FPipeProcessorTask
 {
 public:
-	FPipeProcessorTask(UEntitySubsystem& InEntitySubsystem, const FLWComponentSystemExecutionContext& InExecutionContext, UPipeProcessor& InProc, bool bInManageCommandBuffer = true)
+	FPipeProcessorTask(UMassEntitySubsystem& InEntitySubsystem, const FLWComponentSystemExecutionContext& InExecutionContext, UPipeProcessor& InProc, bool bInManageCommandBuffer = true)
 		: EntitySubsystem(&InEntitySubsystem)
 		, ExecutionContext(InExecutionContext)
 		, Processor(&InProc)
@@ -75,7 +75,7 @@ public:
 
 		PROCESSOR_LOG(TEXT("+--+ Task %s started on %u"), *Processor->GetProcessorName(), FPlatformTLS::GetCurrentThreadId());
 
-		UEntitySubsystem::FScopedProcessing ProcessingScope = EntitySubsystem->NewProcessingScope();
+		UMassEntitySubsystem::FScopedProcessing ProcessingScope = EntitySubsystem->NewProcessingScope();
 
 		TRACE_CPUPROFILER_EVENT_SCOPE_STR("Pipe Processor Task");
 		
@@ -94,7 +94,7 @@ public:
 	}
 
 private:
-	UEntitySubsystem* EntitySubsystem = nullptr;
+	UMassEntitySubsystem* EntitySubsystem = nullptr;
 	FLWComponentSystemExecutionContext ExecutionContext;
 	UPipeProcessor* Processor = nullptr;
 	/** 
@@ -107,7 +107,7 @@ private:
 class FPipeProcessorsTask_GameThread : public FPipeProcessorTask
 {
 public:
-	FPipeProcessorsTask_GameThread(UEntitySubsystem& InEntitySubsystem, const FLWComponentSystemExecutionContext& InExecutionContext, UPipeProcessor& InProc)
+	FPipeProcessorsTask_GameThread(UMassEntitySubsystem& InEntitySubsystem, const FLWComponentSystemExecutionContext& InExecutionContext, UPipeProcessor& InProc)
 		: FPipeProcessorTask(InEntitySubsystem, InExecutionContext, InProc)
 	{}
 
@@ -148,7 +148,7 @@ void UPipeProcessor::PostInitProperties()
 #endif
 }
 
-void UPipeProcessor::CallExecute(UEntitySubsystem& EntitySubsystem, FLWComponentSystemExecutionContext& Context)
+void UPipeProcessor::CallExecute(UMassEntitySubsystem& EntitySubsystem, FLWComponentSystemExecutionContext& Context)
 {
 	TRACE_CPUPROFILER_EVENT_SCOPE_TEXT(*StatId);
 #if WITH_PIPE_DEBUG
@@ -157,7 +157,7 @@ void UPipeProcessor::CallExecute(UEntitySubsystem& EntitySubsystem, FLWComponent
 	Execute(EntitySubsystem, Context);
 }
 
-FGraphEventRef UPipeProcessor::DispatchProcessorTasks(UEntitySubsystem& EntitySubsystem, FLWComponentSystemExecutionContext& ExecutionContext, const FGraphEventArray& Prerequisites)
+FGraphEventRef UPipeProcessor::DispatchProcessorTasks(UMassEntitySubsystem& EntitySubsystem, FLWComponentSystemExecutionContext& ExecutionContext, const FGraphEventArray& Prerequisites)
 {
 	FGraphEventRef ReturnVal;
 	if (bRequiresGameThreadExecution)
@@ -216,7 +216,7 @@ void UPipeCompositeProcessor::ConfigureQueries()
 	// via their individual PostInitProperties call
 }
 
-FGraphEventRef UPipeCompositeProcessor::DispatchProcessorTasks(UEntitySubsystem& EntitySubsystem, FLWComponentSystemExecutionContext& ExecutionContext, const FGraphEventArray& InPrerequisites)
+FGraphEventRef UPipeCompositeProcessor::DispatchProcessorTasks(UMassEntitySubsystem& EntitySubsystem, FLWComponentSystemExecutionContext& ExecutionContext, const FGraphEventArray& InPrerequisites)
 {
 	FGraphEventArray Events;
 	Events.Reserve(ProcessingFlatGraph.Num());
@@ -273,7 +273,7 @@ FGraphEventRef UPipeCompositeProcessor::DispatchProcessorTasks(UEntitySubsystem&
 	return CompletionEvent;
 }
 
-void UPipeCompositeProcessor::Execute(UEntitySubsystem& EntitySubsystem, FLWComponentSystemExecutionContext& Context)
+void UPipeCompositeProcessor::Execute(UMassEntitySubsystem& EntitySubsystem, FLWComponentSystemExecutionContext& Context)
 {
 #if PARALLELIZED_TRAFFIC_HACK
 	if (FPipeTweakables::bParallelGroups == false
diff --git a/Engine/Plugins/Runtime/MassEntity/Source/MassEntity/Public/ArchetypeTypes.h b/Engine/Plugins/Runtime/MassEntity/Source/MassEntity/Public/ArchetypeTypes.h
--- a/Engine/Plugins/Runtime/MassEntity/Source/MassEntity/Public/ArchetypeTypes.h
+++ b/Engine/Plugins/Runtime/MassEntity/Source/MassEntity/Public/ArchetypeTypes.h
@@ -7,7 +7,7 @@
 #include "Containers/UnrealString.h"
 #include "LWComponentTypes.h"
 
-class UEntitySubsystem;
+class UMassEntitySubsystem;
 struct FArchetypeData;
 struct FLWComponentSystemExecutionContext;
 struct FLWComponentData;
@@ -40,7 +40,7 @@ private:
 	{}
 	TSharedPtr<FArchetypeData> DataPtr;
 
-	friend UEntitySubsystem;
+	friend UMassEntitySubsystem;
 	friend FArchetypeChunkCollection;
 	friend FLWComponentQuery;
 	friend FEntityView;
diff --git a/Engine/Plugins/Runtime/MassEntity/Source/MassEntity/Public/EntityQuery.h b/Engine/Plugins/Runtime/MassEntity/Source/MassEntity/Public/EntityQuery.h
--- a/Engine/Plugins/Runtime/MassEntity/Source/MassEntity/Public/EntityQuery.h
+++ b/Engine/Plugins/Runtime/MassEntity/Source/MassEntity/Public/EntityQuery.h
@@ -10,7 +10,7 @@
 #include "EntityQuery.generated.h"
 
 
-class UEntitySubsystem;
+class UMassEntitySubsystem;
 struct FArchetypeData;
 struct FLWComponentSystemExecutionContext;
 struct FLWComponentData;
@@ -103,20 +103,20 @@ public:
 	FLWComponentQuery(TConstArrayView<const UScriptStruct*> InitList);
 
 	/** Runs ExecuteFunction on all entities matching Requirements */
-	void ForEachEntityChunk(UEntitySubsystem& EntitySubsystem, FLWComponentSystemExecutionContext& ExecutionContext, const FLWComponentSystemExecuteFunction& ExecuteFunction);
+	void ForEachEntityChunk(UMassEntitySubsystem& EntitySubsystem, FLWComponentSystemExecutionContext& ExecutionContext, const FLWComponentSystemExecuteFunction& ExecuteFunction);
 	
 	/** Will first verify that the archetype given with Chunks matches the query's requirements, and if so will run the other, more generic ForEachEntityChunk implementation */
-	void ForEachEntityChunk(const FArchetypeChunkCollection& Chunks, UEntitySubsystem& EntitySubsystem, FLWComponentSystemExecutionContext& ExecutionContext, const FLWComponentSystemExecuteFunction& ExecuteFunction);
+	void ForEachEntityChunk(const FArchetypeChunkCollection& Chunks, UMassEntitySubsystem& EntitySubsystem, FLWComponentSystemExecutionContext& ExecutionContext, const FLWComponentSystemExecuteFunction& ExecuteFunction);
 
 	/**
 	 * Attempts to process every chunk of every affected archetype in parallel.
 	 */
-	void ParallelForEachEntityChunk(UEntitySubsystem& EntitySubsystem, FLWComponentSystemExecutionContext& ExecutionContext, const FLWComponentSystemExecuteFunction& ExecuteFunction);
+	void ParallelForEachEntityChunk(UMassEntitySubsystem& EntitySubsystem, FLWComponentSystemExecutionContext& ExecutionContext, const FLWComponentSystemExecuteFunction& ExecuteFunction);
 
 	/** Will gather all archetypes from InEntitySubsystem matching this->Requirements.
 	 *  Note that no work will be done if the cached data is up to date (as tracked by EntitySubsystemHash and 
 	 *	ArchetypeDataVersion properties). */
-	void CacheArchetypes(UEntitySubsystem& InEntitySubsystem);
+	void CacheArchetypes(UMassEntitySubsystem& InEntitySubsystem);
 
 	FLWComponentQuery& AddRequirement(const UScriptStruct* ComponentType, const ELWComponentAccess AccessMode, const ELWComponentPresence Presence = ELWComponentPresence::All)
 	{
@@ -293,14 +293,14 @@ public:
 	 * Note that the function is not const because calling it can result in re-caching of ValidArchetypes 
 	 * @return the number of entities this given query would process if called "now"
 	 */
-	int32 GetMatchingEntitiesNum(UEntitySubsystem& InEntitySubsystem);
+	int32 GetMatchingEntitiesNum(UMassEntitySubsystem& InEntitySubsystem);
 
 	/**
 	 * Checks if any of ValidArchetypes has any entities.
 	 * Note that the function is not const because calling it can result in re-caching of ValidArchetypes
 	 * @return "true" if any of the ValidArchetypes has any entities, "false" otherwise
 	 */
-	bool HasMatchingEntities(UEntitySubsystem& InEntitySubsystem);
+	bool HasMatchingEntities(UMassEntitySubsystem& InEntitySubsystem);
 
 	/** 
 	 * Sets a chunk filter condition that will applied to each chunk of all valid archetypes. Note 
diff --git a/Engine/Plugins/Runtime/MassEntity/Source/MassEntity/Public/EntitySubsystem.h b/Engine/Plugins/Runtime/MassEntity/Source/MassEntity/Public/EntitySubsystem.h
deleted file mode 100644
diff --git a/Engine/Plugins/Runtime/MassEntity/Source/MassEntity/Public/EntityView.h b/Engine/Plugins/Runtime/MassEntity/Source/MassEntity/Public/EntityView.h
--- a/Engine/Plugins/Runtime/MassEntity/Source/MassEntity/Public/EntityView.h
+++ b/Engine/Plugins/Runtime/MassEntity/Source/MassEntity/Public/EntityView.h
@@ -7,7 +7,7 @@
 #include "EntityView.generated.h"
 
 
-class UEntitySubsystem;
+class UMassEntitySubsystem;
 struct FArchetypeData;
 struct FArchetypeHandle;
 struct FArchetypeComponentConfig;
@@ -26,7 +26,7 @@ struct MASSENTITY_API FEntityView
 	/** 
 	 *  Resolves Entity against ArchetypeHandle. Note that this approach requires the caller to ensure that Entity
 	 *  indeed belongs to ArchetypeHandle. If not the call will fail a check. As a remedy calling the 
-	 *  UEntitySubsystem-flavored constructor is recommended since it will first find the appropriate archetype for
+	 *  UMassEntitySubsystem-flavored constructor is recommended since it will first find the appropriate archetype for
 	 *  Entity. 
 	 */
 	FEntityView(const FArchetypeHandle& ArchetypeHandle, FLWEntity Entity);
@@ -35,7 +35,7 @@ struct MASSENTITY_API FEntityView
 	 *  Finds the archetype Entity belongs to and then resolves against it. The caller is responsible for ensuring
 	 *  that the given Entity is in fact a valid ID tied to any of the archetypes 
 	 */
-	FEntityView(const UEntitySubsystem& EntitySubsystem, FLWEntity Entity);
+	FEntityView(const UMassEntitySubsystem& EntitySubsystem, FLWEntity Entity);
 
 	FLWEntity GetEntity() const	{ return Entity; }
 
diff --git a/Engine/Plugins/Runtime/MassEntity/Source/MassEntity/Public/LWCCommandBuffer.h b/Engine/Plugins/Runtime/MassEntity/Source/MassEntity/Public/LWCCommandBuffer.h
--- a/Engine/Plugins/Runtime/MassEntity/Source/MassEntity/Public/LWCCommandBuffer.h
+++ b/Engine/Plugins/Runtime/MassEntity/Source/MassEntity/Public/LWCCommandBuffer.h
@@ -3,7 +3,7 @@
 #pragma once
 
 #include "LWComponentTypes.h"
-#include "EntitySubsystem.h"
+#include "MassEntitySubsystem.h"
 #include "InstancedStructStream.h"
 #include "Misc/MTAccessDetector.h"
 
@@ -27,7 +27,7 @@ struct MASSENTITY_API FCommandBufferEntryBase
 		: TargetEntity(InEntity)
 	{}
 
-	virtual void Execute(UEntitySubsystem& System) const PURE_VIRTUAL(FCommandBufferEntryBase::Execute, );
+	virtual void Execute(UMassEntitySubsystem& System) const PURE_VIRTUAL(FCommandBufferEntryBase::Execute, );
 };
 template<> struct TStructOpsTypeTraits<FCommandBufferEntryBase> : public TStructOpsTypeTraitsBase2<FCommandBufferEntryBase> { enum { WithPureVirtual = true, }; };
 
@@ -41,17 +41,17 @@ struct MASSENTITY_API FDeferredCommand : public FCommandBufferEntryBase
 
 	FDeferredCommand() = default;
 
-	FDeferredCommand(const TFunction<void(UEntitySubsystem& System)>& InDeferredFunction)
+	FDeferredCommand(const TFunction<void(UMassEntitySubsystem& System)>& InDeferredFunction)
 		: DeferredFunction(InDeferredFunction)
 	{}
 
-	virtual void Execute(UEntitySubsystem& System) const override
+	virtual void Execute(UMassEntitySubsystem& System) const override
 	{
 		DeferredFunction(System);
 	}
 
 protected:
-	TFunction<void(UEntitySubsystem& System)> DeferredFunction;
+	TFunction<void(UMassEntitySubsystem& System)> DeferredFunction;
 };
 
 /**
@@ -69,7 +69,7 @@ struct MASSENTITY_API FBuildEntityFromComponentInstance : public FCommandBufferE
 	{}
 
 protected:
-	virtual void Execute(UEntitySubsystem& System) const override;
+	virtual void Execute(UMassEntitySubsystem& System) const override;
 
 	FInstancedStruct Struct;
 };
@@ -89,7 +89,7 @@ struct MASSENTITY_API FBuildEntityFromComponentInstances : public FCommandBuffer
 	{}
 
 protected:
-	virtual void Execute(UEntitySubsystem& System) const override;
+	virtual void Execute(UMassEntitySubsystem& System) const override;
 
 	TArray<FInstancedStruct> Instances;
 };
@@ -109,7 +109,7 @@ struct MASSENTITY_API FCommandAddComponent : public FCommandBufferEntryBase
 	{}
 
 protected:
-	virtual void Execute(UEntitySubsystem& System) const override;
+	virtual void Execute(UMassEntitySubsystem& System) const override;
 
 	UScriptStruct* StructParam = nullptr;
 };
@@ -129,7 +129,7 @@ struct MASSENTITY_API FCommandAddComponentInstance : public FCommandBufferEntryB
 	{}
 
 protected:
-	virtual void Execute(UEntitySubsystem& System) const override;
+	virtual void Execute(UMassEntitySubsystem& System) const override;
 
 	FInstancedStruct Struct;
 };
@@ -149,7 +149,7 @@ struct MASSENTITY_API FCommandRemoveComponent : public FCommandBufferEntryBase
 	{}
 
 protected:
-	virtual void Execute(UEntitySubsystem& System) const override;
+	virtual void Execute(UMassEntitySubsystem& System) const override;
 
 	UScriptStruct* StructParam = nullptr;
 };
@@ -169,7 +169,7 @@ struct MASSENTITY_API FCommandAddComponentList : public FCommandBufferEntryBase
 	{}
 
 protected:
-	virtual void Execute(UEntitySubsystem& System) const override;
+	virtual void Execute(UMassEntitySubsystem& System) const override;
 
 	TArray<const UScriptStruct*> ComponentList;
 };
@@ -189,7 +189,7 @@ struct MASSENTITY_API FCommandRemoveComponentList : public FCommandBufferEntryBa
 	{}
 
 protected:
-	virtual void Execute(UEntitySubsystem& System) const override;
+	virtual void Execute(UMassEntitySubsystem& System) const override;
 
 	TArray<const UScriptStruct*> ComponentList;
 };
@@ -209,7 +209,7 @@ struct MASSENTITY_API FCommandAddTag: public FCommandBufferEntryBase
 	{}
 
 protected:
-	virtual void Execute(UEntitySubsystem& System) const override;
+	virtual void Execute(UMassEntitySubsystem& System) const override;
 
 	UScriptStruct* StructParam = nullptr;
 };
@@ -229,7 +229,7 @@ struct MASSENTITY_API FCommandRemoveTag : public FCommandBufferEntryBase
 	{}
 
 protected:
-	virtual void Execute(UEntitySubsystem& System) const override;
+	virtual void Execute(UMassEntitySubsystem& System) const override;
 
 	UScriptStruct* StructParam = nullptr;
 };
@@ -255,7 +255,7 @@ struct MASSENTITY_API FCommandSwapTags : public FCommandBufferEntryBase
 	}
 
 protected:
-	virtual void Execute(UEntitySubsystem& System) const override;
+	virtual void Execute(UMassEntitySubsystem& System) const override;
 
 	const UScriptStruct* OldTagType = nullptr;
 	const UScriptStruct* NewTagType = nullptr;
@@ -276,7 +276,7 @@ struct MASSENTITY_API FCommandRemoveComposition : public FCommandBufferEntryBase
 	{}
 
 protected:
-	virtual void Execute(UEntitySubsystem& System) const override;
+	virtual void Execute(UMassEntitySubsystem& System) const override;
 
 	FLWCompositionDescriptor Descriptor;
 };
@@ -346,7 +346,7 @@ public:
 		EntitiesToDestroy.Append(InEntitiesToDestroy);
 	}
 
-	void ReplayBufferAgainstSystem(UEntitySubsystem* System);
+	void ReplayBufferAgainstSystem(UMassEntitySubsystem* System);
 
 	SIZE_T GetAllocatedSize() const { return PendingCommands.GetAllocatedSize(); }
 
diff --git a/Engine/Plugins/Runtime/MassEntity/Source/MassEntity/Public/LWComponentTypes.h b/Engine/Plugins/Runtime/MassEntity/Source/MassEntity/Public/LWComponentTypes.h
--- a/Engine/Plugins/Runtime/MassEntity/Source/MassEntity/Public/LWComponentTypes.h
+++ b/Engine/Plugins/Runtime/MassEntity/Source/MassEntity/Public/LWComponentTypes.h
@@ -35,7 +35,7 @@ struct FLWChunkComponent
 	FLWChunkComponent() {}
 };
 
-// A handle to a lightweight entity.  An entity is used in conjunction with the UEntitySubsystem
+// A handle to a lightweight entity.  An entity is used in conjunction with the UMassEntitySubsystem
 // for the current world and can contain lightweight components.
 USTRUCT()
 struct FLWEntity
diff --git a/Engine/Plugins/Runtime/MassEntity/Source/MassEntity/Public/MassEntitySubsystem.h b/Engine/Plugins/Runtime/MassEntity/Source/MassEntity/Public/MassEntitySubsystem.h
new file mode 100644
--- /dev/null
+++ b/Engine/Plugins/Runtime/MassEntity/Source/MassEntity/Public/MassEntitySubsystem.h
@@ -0,0 +1,526 @@
+// Copyright Epic Games, Inc. All Rights Reserved.
+
+#pragma once
+
+#include "Subsystems/WorldSubsystem.h"
+#include "LWComponentTypes.h"
+#include "InstancedStruct.h"
+#include "EntityQuery.h"
+#include "MassEntitySubsystem.generated.h"
+
+MASSENTITY_API DECLARE_LOG_CATEGORY_EXTERN(LogAggregateTicking, Warning, All);
+
+class UMassEntitySubsystem;
+struct FLWComponentQuery;
+struct FLWComponentSystemExecutionContext;
+struct FArchetypeData;
+struct FLWCCommandBuffer;
+struct FArchetypeChunkCollection;
+struct FArchetypeChunk;
+class FOutputDevice;
+enum class ELWComponentAccess : uint8;
+
+namespace UE { namespace AggregateTicking {
+FString DebugGetComponentAccessString(ELWComponentAccess Access);
+}}
+
+//@TODO: Comment this guy
+UCLASS()
+class MASSENTITY_API UMassEntitySubsystem : public UWorldSubsystem
+{
+	GENERATED_BODY()
+
+	friend struct FLWComponentQuery;
+private:
+	// Index 0 is reserved so we can treat that index as an invalid entity handle
+	constexpr static int32 NumReservedEntities = 1;
+
+	struct FEntityData
+	{
+		TSharedPtr<FArchetypeData> CurrentArchetype;
+		int32 SerialNumber = 0;
+
+		void Reset()
+		{
+			CurrentArchetype.Reset();
+			SerialNumber = 0;
+		}
+
+		bool IsValid() const
+		{
+			return SerialNumber != 0 && CurrentArchetype.IsValid();
+		}
+	};
+	
+public:
+	struct FScopedProcessing
+	{
+		explicit FScopedProcessing(std::atomic<int32>& InProcessingScopeCount) : ScopedProcessingCount(InProcessingScopeCount)
+		{
+			++ScopedProcessingCount;
+		}
+		~FScopedProcessing()
+		{
+			--ScopedProcessingCount;
+		}
+	private:
+		std::atomic<int32>& ScopedProcessingCount;
+	};
+
+	const static FLWEntity InvalidEntity;
+
+	UMassEntitySubsystem();
+
+	//~UObject interface
+	virtual void GetResourceSizeEx(FResourceSizeEx& CumulativeResourceSize) override;
+	//~End of UObject interface
+
+	//~USubsystem interface
+	virtual void Initialize(FSubsystemCollectionBase& Collection) override;
+	virtual void Deinitialize() override;
+	//~End of USubsystem interface
+
+	static FORCEINLINE UMassEntitySubsystem* GetCurrent(UWorld* World)
+	{
+		return World ? World->GetSubsystem<UMassEntitySubsystem>() : nullptr;
+	}
+	
+	/** 
+	 * A special, relaxed but slower version of CreateArchetype functions that allows ComponentAngTagsList to contain 
+	 * both components and tags. 
+	 */
+	FArchetypeHandle CreateArchetype(TConstArrayView<const UScriptStruct*> ComponentsAngTagsList);
+
+	FArchetypeHandle CreateArchetype(TConstArrayView<const UScriptStruct*> ComponentList, const FLWTagBitSet& Tags, const FLWChunkComponentBitSet& ChunkComponents);
+
+	/** 
+	 *  Creates an archetype like SourceArchetype + NewComponentList. 
+	 *  @param SourceArchetype the archetype used to initially populate the list of components of the archetype being created. 
+	 *  @param NewComponentList list of unique components to add to components fetched from SourceArchetype. Note that 
+	 *   adding an empty list is not supported and doing so will result in failing a `check`
+	 *  @return a handle of a new archetype
+	 *  @note it's caller's responsibility to ensure that NewComponentList is not empty and contains only component
+	 *   types that SourceArchetype doesn't already have. If the caller cannot guarantee it use of AddComponent functions
+	 *   family is recommended.
+	 */
+	FArchetypeHandle CreateArchetype(const TSharedPtr<FArchetypeData>& SourceArchetype, const FLWComponentBitSet& NewComponentList);
+
+	FArchetypeHandle CreateArchetype(const FLWCompositionDescriptor& Descriptor);
+
+
+	FArchetypeHandle GetArchetypeForEntity(FLWEntity Entity) const;
+	/** Method to iterate on all the component types of an archetype */
+	static void ForEachArchetypeComponentType(const FArchetypeHandle Archetype, TFunction< void(const UScriptStruct* /*ComponentType*/)> Function);
+
+	void SetDefaultChunkComponentValue(const FArchetypeHandle Archetype, FConstStructView InstancedStruct);
+
+	/**
+	 * Go through all archetypes and compact entities
+	 * @param TimeAllowed to do entity compaction, once it reach that time it will stop and return
+	 */
+	void DoEntityCompaction(const double TimeAllowed);
+
+	/**
+	 * Creates fully built entity ready to be used by the subsystem
+	 * @param Archetype you want this entity to be
+	 * @return FLWEntity id of the newly created entity */
+	FLWEntity CreateEntity(const FArchetypeHandle Archetype);
+
+	/**
+	 * Creates fully built entity ready to be used by the subsystem
+	 * @param ComponentInstanceList is the components to create the entity from and initialize values
+	 * @return FLWEntity id of the newly created entity */
+	FLWEntity CreateEntity(TConstArrayView<FInstancedStruct> ComponentInstanceList);
+
+	/** A version of CreateEntity that's creating a number of entities (Count) at one go
+	 *  @param Archetype you want this entity to be
+	 *  @param Count number of entities to create
+	 *  @param OutEntities the newly created entities are appended to given array, i.e. the pre-existing content of OutEntities won't be affected by the call
+	 *  @return number of entities created */
+	int32 BatchCreateEntities(const FArchetypeHandle Archetype, const int32 Count, TArray<FLWEntity>& OutEntities);
+
+	/**
+	 * Destroys a fully built entity, use ReleaseReservedEntity if entity was not yet built.
+	 * @param Entity to destroy */
+	void DestroyEntity(FLWEntity Entity);
+
+	/**
+	 * Reserves an entity in the subsystem, the entity is still not ready to be used by the subsystem, need to call BuildEntity()
+	 * @return FLWEntity id of the reserved entity */
+	FLWEntity ReserveEntity();
+
+	/**
+	 * Builds an entity for it to be ready to be used by the subsystem
+	 * @param Entity to build which was retrieved with ReserveEntity() method
+	 * @param Archetype you want this entity to be*/
+	void BuildEntity(FLWEntity Entity, const FArchetypeHandle Archetype);
+
+	/**
+	 * Builds an entity for it to be ready to be used by the subsystem
+	 * @param Entity to build which was retrieved with ReserveEntity() method
+	 * @param ComponentInstanceList is the components to create the entity from and initialize values*/
+	void BuildEntity(FLWEntity Entity, TConstArrayView<FInstancedStruct> ComponentInstanceList);
+
+	/*
+	 * Releases a previously reserved entity that was not yet built, otherwise call DestroyEntity
+	 * @param Entity to release */
+	void ReleaseReservedEntity(FLWEntity Entity);
+
+	/**
+	 * Destroys all the entity in the provided array of entities
+	 * @param InEntities to destroy
+	 */
+	void BatchDestroyEntities(TConstArrayView<FLWEntity> InEntities);
+
+	void BatchDestroyEntityChunks(const FArchetypeChunkCollection& Chunks);
+
+	void AddComponentToEntity(FLWEntity Entity, const UScriptStruct* ComponentType);
+
+	/** 
+	 *  Ensures that only unique components are added. 
+	 *  @note It's caller's responsibility to ensure Entity's and ComponentList's validity. 
+	 */
+	void AddComponentListToEntity(FLWEntity Entity, TConstArrayView<const UScriptStruct*> ComponentList);
+
+	void AddComponentInstanceListToEntity(FLWEntity Entity, TConstArrayView<FInstancedStruct> ComponentInstanceList);
+	void RemoveComponentFromEntity(FLWEntity Entity, const UScriptStruct* ComponentType);
+	void RemoveComponentListFromEntity(FLWEntity Entity, TConstArrayView<const UScriptStruct*> ComponentList);
+
+	void AddTagToEntity(FLWEntity Entity, const UScriptStruct* TagType);
+	void RemoveTagFromEntity(FLWEntity Entity, const UScriptStruct* TagType);
+	void SwapTagsForEntity(FLWEntity Entity, const UScriptStruct* FromComponentType, const UScriptStruct* ToComponentType);
+
+
+	/**
+	 * Adds fragments and tags indicated by InOutDescriptor to the Entity. The function also figures out which elements
+	 * in InOutDescriptor are missing from the current composition of the given entity and then returns the resulting 
+	 * delta via InOutDescriptor.
+	 */
+	void AddCompositionToEntity_GetDelta(FLWEntity Entity, FLWCompositionDescriptor& InOutDescriptor);
+	void RemoveCompositionFromEntity(FLWEntity Entity, const FLWCompositionDescriptor& InDescriptor);
+
+	/** 
+	 * Moves an entity over to a new archetype by copying over components common to both archetypes
+	 * @param Entity is the entity to move 
+	 * @param NewArchetypeHandle the handle to the new archetype
+	 */
+	void MoveEntityToAnotherArchetype(FLWEntity Entity, FArchetypeHandle NewArchetypeHandle);
+
+	/** Copies values from ComponentInstanceList over to Entity's component. Caller is responsible for ensuring that 
+	 *  the given entity does have given components. Failing this assumption will cause a check-fail.*/
+	void SetEntityComponentsValues(FLWEntity Entity, TArrayView<const FInstancedStruct> ComponentInstanceList);
+
+	/** Copies values from ComponentInstanceList over to components of given entities collection. The caller is responsible 
+	 *  for ensuring that the given entity archetype (FArchetypeChunkCollection.Archetype) does have given components. 
+	 *  Failing this assumption will cause a check-fail. */
+	static void BatchSetEntityComponentsValues(const FArchetypeChunkCollection& SparseEntities, TArrayView<const FInstancedStruct> ComponentInstanceList);
+
+	// Return true if it is an valid built entity
+	bool IsEntityActive(FLWEntity Entity) const 
+	{
+		return IsEntityValid(Entity) && IsEntityBuilt(Entity);
+	}
+
+	// Returns true if Entity is valid
+	bool IsEntityValid(FLWEntity Entity) const;
+
+	// Returns true if Entity is has been fully built (expecting a valid Entity)
+	bool IsEntityBuilt(FLWEntity Entity) const;
+
+	// Asserts that IsEntityValid
+	void CheckIfEntityIsValid(FLWEntity Entity) const;
+
+	// Asserts that IsEntityBuilt
+	void CheckIfEntityIsActive(FLWEntity Entity) const;
+
+	template <typename ComponentType>
+	ComponentType& GetComponentDataChecked(FLWEntity Entity) const
+	{
+		return *((ComponentType*)InternalGetComponentDataChecked(Entity, ComponentType::StaticStruct()));
+	}
+
+	template <typename ComponentType>
+	ComponentType* GetComponentDataPtr(FLWEntity Entity) const
+	{
+		return (ComponentType*)InternalGetComponentDataPtr(Entity, ComponentType::StaticStruct());
+	}
+
+	FStructView GetComponentDataStruct(FLWEntity Entity, const UScriptStruct* ComponentType) const
+	{
+		return FStructView(ComponentType, static_cast<uint8*>(InternalGetComponentDataPtr(Entity, ComponentType)));
+	}
+
+	uint32 GetArchetypeDataVersion() const { return ArchetypeDataVersion; }
+
+	/**
+	 * Creates and initializes a FLWComponentSystemExecutionContext instance.
+	 */
+	FLWComponentSystemExecutionContext CreateExecutionContext(const float DeltaSeconds) const;
+
+	FScopedProcessing NewProcessingScope() { return FScopedProcessing(ProcessingScopeCount); }
+	bool IsProcessing() const { return ProcessingScopeCount > 0; }
+	FLWCCommandBuffer& Defer() { return *DeferredCommandBuffer.Get(); }
+
+#if WITH_AGGREGATETICKING_DEBUG
+	void DebugPrintEntity(int32 Index, FOutputDevice& Ar, const TCHAR* InPrefix = TEXT("")) const;
+	void DebugPrintEntity(FLWEntity Entity, FOutputDevice& Ar, const TCHAR* InPrefix = TEXT("")) const;
+	void DebugPrintArchetypes(FOutputDevice& Ar) const;
+	static void DebugGetStringDesc(const FArchetypeHandle& Archetype, FOutputDevice& Ar);
+	void DebugGetArchetypesStringDetails(FOutputDevice& Ar, const bool bIncludeEmpty = true);
+	void DebugGetArchetypeComponentTypes(const FArchetypeHandle& Archetype, TArray<const UScriptStruct*>& InOutComponentList) const;
+	int32 DebugGetArchetypeEntitiesCount(const FArchetypeHandle& Archetype) const;	
+	int32 DebugGetEntityCount() const { return Entities.Num() - NumReservedEntities - EntityFreeIndexList.Num(); }
+	int32 DebugGetArchetypesCount() const { return ComponentHashToArchetypeMap.Num(); }
+	void DebugRemoveAllEntities();
+	void DebugForceArchetypeDataVersionBump() { ++ArchetypeDataVersion; }
+	void DebugGetArchetypeStrings(const FArchetypeHandle& Archetype, TArray<FName>& OutComponentNames, TArray<FName>& OutTagNames);
+#endif // WITH_AGGREGATETICKING_DEBUG
+
+protected:
+	void GetValidArchetypes(const FLWComponentQuery& Query, TArray<FArchetypeHandle>& OutValidArchetypes);
+	
+	FArchetypeHandle InternalCreateSiblingArchetype(const TSharedPtr<FArchetypeData>& SourceArchetype, const FLWTagBitSet& OverrideTags);
+
+private:
+	void InternalBuildEntity(FLWEntity Entity, const FArchetypeHandle Archetype);
+	void InternalReleaseEntity(FLWEntity Entity);
+
+	/** 
+	 *  Adds components in ComponentList to Entity. Only the unique components will be added.
+	 */
+	void InternalAddComponentListToEntityChecked(FLWEntity Entity, const FLWComponentBitSet& InComponents);
+
+	/** 
+	 *  Similar to InternalAddComponentListToEntity but expects NewComponentList not overlapping with current entity's
+	 *  component list. It's callers responsibility to ensure that's true. Failing this will cause a `check` fail.
+	 */
+	void InternalAddComponentListToEntity(FLWEntity Entity, const FLWComponentBitSet& NewComponents);
+	void* InternalGetComponentDataChecked(FLWEntity Entity, const UScriptStruct* ComponentType) const;
+	void* InternalGetComponentDataPtr(FLWEntity Entity, const UScriptStruct* ComponentType) const;
+
+	//@TODO: Not optimized or a great API, super WIP
+	struct MASSENTITY_API FChunkIteratorByComponent
+	{
+		FChunkIteratorByComponent(UMassEntitySubsystem* InSystem, const UScriptStruct* InComponentType);
+
+		void GetCurrentChunkBase(void*& OutChunkBasePtr, int32& OutChunkEntityCount) const;
+
+		void Advance();
+		bool IsDone() const;
+		FArchetypeData* GetCurrentArchetype() const;
+
+		UMassEntitySubsystem* System;
+		const UScriptStruct* ComponentType;
+
+		TArray<TSharedPtr<FArchetypeData>>* pPendingArchetypes;
+		int32 ArchetypeIndex = 0;
+		int32 ChunkIndex = 0;
+	};
+
+	friend FChunkIteratorByComponent;
+
+private:
+	TChunkedArray<FEntityData> Entities;
+	TArray<int32> EntityFreeIndexList;
+
+	TSharedPtr<FLWCCommandBuffer> DeferredCommandBuffer;
+	std::atomic<int32> SerialNumberGenerator;
+	std::atomic<int32> ProcessingScopeCount;
+
+	// the "version" number increased every time an archetype gets added
+	uint32 ArchetypeDataVersion = 0;
+
+	// Map of hash of sorted component list to archetypes with that hash
+	TMap<uint32, TArray<TSharedPtr<FArchetypeData>>> ComponentHashToArchetypeMap;
+
+	// Map to list of archetypes that contain the specified component type
+	TMap<const UScriptStruct*, TArray<TSharedPtr<FArchetypeData>>> ComponentTypeToArchetypeMap;
+};
+
+
+//////////////////////////////////////////////////////////////////////
+//
+
+struct MASSENTITY_API FLWComponentSystemExecutionContext
+{
+private:
+	struct FComponentView 
+	{
+		FLWComponentRequirement Requirement;
+		TArrayView<FLWComponentData> ComponentView;
+
+		FComponentView() {}
+		explicit FComponentView(const FLWComponentRequirement& InRequirement) : Requirement(InRequirement) {}
+
+		bool operator==(const UScriptStruct* ComponentType) const { return Requirement.StructType == ComponentType; }
+	};
+	TArray<FComponentView, TInlineAllocator<8>> ComponentViews;
+
+	struct FChunkComponentView
+	{
+		FLWComponentRequirement Requirement;
+		FStructView ChunkComponentView;
+
+		FChunkComponentView() {}
+		explicit FChunkComponentView(const FLWComponentRequirement& InRequirement) : Requirement(InRequirement)	{}
+
+		bool operator==(const UScriptStruct* ComponentType) const { return Requirement.StructType == ComponentType; }
+	};
+	TArray<FChunkComponentView, TInlineAllocator<4>> ChunkComponents;
+
+	// mz@todo make this shared ptr thread-safe and never auto-flush in MT environment. 
+	TSharedPtr<FLWCCommandBuffer> DeferredCommandBuffer;
+	TArrayView<FLWEntity> EntityListView;
+	
+	/** If set this indicates the exact archetype and its chunks to be processed. 
+	 *  @todo this data should live somewhere else, preferably be just a parameter to Query.ForEachEntityChunk function */
+	FArchetypeChunkCollection ChunkCollection;
+	
+	FInstancedStruct AuxData;
+	float DeltaTimeSeconds = 0.0f;
+	int32 ChunkSerialModificationNumber = -1;
+	FLWTagBitSet CurrentArchetypesTagBitSet;
+
+#if WITH_AGGREGATETICKING_DEBUG
+	FString DebugExecutionDescription;
+#endif
+	
+	/** @todo update comment, and look if we If true the EntitySystem will flush the deferred commands stored in DeferredCommandBuffer just after executing 
+	 *  the given system. If False then the party calling UEntitySubsystem::ExecuteSystem is responsible for manually
+	 *  calling FLWCCommandBuffer.ReplayBufferAgainstSystem() */
+	bool bFlushDeferredCommands = true;
+
+	TArrayView<FComponentView> GetMutableRequirements() { return ComponentViews; }
+	TArrayView<FChunkComponentView> GetMutableChunkRequirements() { return ChunkComponents; }
+	
+	friend FArchetypeData;
+	friend FLWComponentQuery;
+
+public:
+	FLWComponentSystemExecutionContext() = default;
+	explicit FLWComponentSystemExecutionContext(const float InDeltaTimeSeconds, const bool bInFlushDeferredCommands = true)
+		: DeltaTimeSeconds(InDeltaTimeSeconds)
+		, bFlushDeferredCommands(bInFlushDeferredCommands)
+	{}
+
+#if WITH_AGGREGATETICKING_DEBUG
+	const FString& DebugGetExecutionDesc() const { return DebugExecutionDescription; }
+	void DebugSetExecutionDesc(const FString& Description) { DebugExecutionDescription = Description; }
+#endif
+
+	/** Sets bFlushDeferredCommands. Note that setting to True while the system is being executed doesn't result in
+	 *  immediate commands flushing */
+	void SetFlushDeferredCommands(const bool bNewFlushDeferredCommands) { bFlushDeferredCommands = bNewFlushDeferredCommands; } 
+	void SetDeferredCommandBuffer(const TSharedPtr<FLWCCommandBuffer>& InDeferredCommandBuffer) { DeferredCommandBuffer = InDeferredCommandBuffer; }
+	void SetChunkCollection(const FArchetypeChunkCollection& InChunkCollection);
+	void SetChunkCollection(FArchetypeChunkCollection&& InChunkCollection);
+	void ClearChunkCollection() { ChunkCollection.Reset(); }
+	void SetAuxData(const FInstancedStruct& InAuxData) { AuxData = InAuxData; }
+
+	float GetDeltaTimeSeconds() const
+	{
+		return DeltaTimeSeconds;
+	}
+
+	TSharedPtr<FLWCCommandBuffer> GetSharedDeferredCommandBuffer() const { return DeferredCommandBuffer; }
+	FLWCCommandBuffer& Defer() { checkSlow(DeferredCommandBuffer.IsValid()); return *DeferredCommandBuffer.Get(); }
+
+	TConstArrayView<FLWEntity> GetEntities() const { return EntityListView; }
+	int32 GetEntitiesNum() const { return EntityListView.Num(); }
+
+	FLWEntity GetEntity(const int32 Index) const
+	{
+		return EntityListView[Index];
+	}
+
+	template<typename T>
+	bool DoesArchetypeHaveTag() const
+	{
+		static_assert(TIsDerivedFrom<T, FComponentTag>::IsDerived, "Given struct is not of a valid fragment type.");
+		return CurrentArchetypesTagBitSet.Contains<T>();
+	}
+
+	/** Chunk related operation */
+	void SetCurrentChunkSerialModificationNumber(const int32 SerialModificationNumber) { ChunkSerialModificationNumber = SerialModificationNumber; }
+	int32 GetChunkSerialModificationNumber() const { return ChunkSerialModificationNumber; }
+
+	template<typename T>
+	T& GetMutableChunkComponent()
+	{
+		static_assert(TIsDerivedFrom<T, FLWChunkComponent>::IsDerived, "Given struct doesn't represent a valid chunk fragment type. Make sure to inherit from FLWChunkComponent or one of its child-types.");
+
+		const UScriptStruct* Type = T::StaticStruct();
+		const FChunkComponentView* FoundChunkComponentData = ChunkComponents.FindByPredicate([Type](const FChunkComponentView& Element) { return Element.Requirement.StructType == Type; });
+		checkf(FoundChunkComponentData, TEXT("Chunk Component requirement not found: %s"), *T::StaticStruct()->GetName());
+		return FoundChunkComponentData->ChunkComponentView.template GetMutable<T>();
+	}
+
+	template<typename T>
+	const T& GetChunkComponent() const
+	{
+		static_assert(TIsDerivedFrom<T, FLWChunkComponent>::IsDerived, "Given struct doesn't represent a valid chunk fragment type. Make sure to inherit from FLWChunkComponent or one of its child-types.");
+
+		const UScriptStruct* Type = T::StaticStruct();
+		const FChunkComponentView* FoundChunkComponentData = ChunkComponents.FindByPredicate([Type](const FChunkComponentView& Element) { return Element.Requirement.StructType == Type; } );
+		checkf(FoundChunkComponentData, TEXT("Chunk Component requirement not found: %s"), *T::StaticStruct()->GetName());
+		return FoundChunkComponentData->ChunkComponentView.template Get<T>();
+	}
+
+	template<typename TLWComponent>
+	TArrayView<TLWComponent> GetMutableComponentView()
+	{
+		const UScriptStruct* ComponentType = TLWComponent::StaticStruct();
+		const FComponentView* View = ComponentViews.FindByPredicate([ComponentType](const FComponentView& Element) { return Element.Requirement.StructType == ComponentType; });
+		//checkfSlow(View != nullptr, TEXT("Requested component type not bound"));
+		//checkfSlow(View->Requirement.AccessMode == ELWComponentAccess::ReadWrite, TEXT("Requested component has not been bound for writing"));
+		return MakeArrayView<TLWComponent>((TLWComponent*)View->ComponentView.GetData(), View->ComponentView.Num());
+	}
+
+	template<typename TLWComponent>
+	TConstArrayView<TLWComponent> GetComponentView() const
+	{
+		const UScriptStruct* ComponentType = TLWComponent::StaticStruct();
+		const FComponentView* View = ComponentViews.FindByPredicate([ComponentType](const FComponentView& Element) { return Element.Requirement.StructType == ComponentType; });
+		//checkfSlow(View != nullptr, TEXT("Requested component type not bound"));
+		return TConstArrayView<TLWComponent>((const TLWComponent*)View->ComponentView.GetData(), View->ComponentView.Num());
+	}
+
+	TConstArrayView<FLWComponentData> GetComponentComponentView(const UScriptStruct* ComponentType) const
+	{
+		const FComponentView* View = ComponentViews.FindByPredicate([ComponentType](const FComponentView& Element) { return Element.Requirement.StructType == ComponentType; });
+		checkSlow(View);
+		return TConstArrayView<FLWComponentData>((const FLWComponentData*)View->ComponentView.GetData(), View->ComponentView.Num());;
+	}
+
+	TArrayView<FLWComponentData> GetMutableComponentView(const UScriptStruct* ComponentType) 
+	{
+		const FComponentView* View = ComponentViews.FindByPredicate([ComponentType](const FComponentView& Element) { return Element.Requirement.StructType == ComponentType; });
+		checkSlow(View);
+		return View->ComponentView;
+	}
+
+	/** Sparse chunk related operation */
+	const FArchetypeChunkCollection& GetChunkCollection() const { return ChunkCollection; }
+
+	const FInstancedStruct& GetAuxData() const { return AuxData; }
+	FInstancedStruct& GetMutableAuxData() { return AuxData; }
+
+	void FlushDeferred(UMassEntitySubsystem& EntitySystem) const;
+
+	void ClearExecutionData();
+	void SetCurrentArchetypeData(FArchetypeData& ArchetypeData);
+
+protected:
+	void SetRequirements(TConstArrayView<FLWComponentRequirement> InRequirements, TConstArrayView<FLWComponentRequirement> InChunkRequirements);
+	void ClearComponentViews()
+	{
+		for (FComponentView& View : ComponentViews)
+		{
+			View.ComponentView = TArrayView<FLWComponentData>();
+		}
+		for (FChunkComponentView& View : ChunkComponents)
+		{
+			View.ChunkComponentView.Reset();
+		}
+	}
+};
diff --git a/Engine/Plugins/Runtime/MassEntity/Source/MassEntity/Public/MassEntitySystem.h b/Engine/Plugins/Runtime/MassEntity/Source/MassEntity/Public/MassEntitySystem.h
deleted file mode 100644
diff --git a/Engine/Plugins/Runtime/MassEntity/Source/MassEntity/Public/MassEntityTypes.h b/Engine/Plugins/Runtime/MassEntity/Source/MassEntity/Public/MassEntityTypes.h
--- a/Engine/Plugins/Runtime/MassEntity/Source/MassEntity/Public/MassEntityTypes.h
+++ b/Engine/Plugins/Runtime/MassEntity/Source/MassEntity/Public/MassEntityTypes.h
@@ -10,7 +10,7 @@
 
 MASSENTITY_API DECLARE_LOG_CATEGORY_EXTERN(LogPipe, Warning, All);
 
-class UEntitySubsystem;
+class UMassEntitySubsystem;
 class UPipeProcessor;
 class UPipeSchematic;
 class UPipeCompositeProcessor;
@@ -38,7 +38,7 @@ struct MASSENTITY_API FPipeContext
 	GENERATED_BODY()
 
 	UPROPERTY()
-	UEntitySubsystem* EntitySubsystem = nullptr;
+	UMassEntitySubsystem* EntitySubsystem = nullptr;
 	
 	UPROPERTY()
 	float DeltaSeconds = 0.f;
@@ -47,7 +47,7 @@ struct MASSENTITY_API FPipeContext
 	FInstancedStruct AuxData;
 
 	FPipeContext() = default;
-	FPipeContext(UEntitySubsystem& InEntities, const float InDeltaSeconds);
+	FPipeContext(UMassEntitySubsystem& InEntities, const float InDeltaSeconds);
 };
 
 /** 
diff --git a/Engine/Plugins/Runtime/MassEntity/Source/MassEntity/Public/MassExecutor.h b/Engine/Plugins/Runtime/MassEntity/Source/MassEntity/Public/MassExecutor.h
--- a/Engine/Plugins/Runtime/MassEntity/Source/MassEntity/Public/MassExecutor.h
+++ b/Engine/Plugins/Runtime/MassEntity/Source/MassEntity/Public/MassExecutor.h
@@ -3,7 +3,7 @@
 #pragma once
 
 #include "Subsystems/WorldSubsystem.h"
-#include "EntitySubsystem.h"
+#include "MassEntitySubsystem.h"
 #include "Async/TaskGraphInterfaces.h"
 
 
diff --git a/Engine/Plugins/Runtime/MassEntity/Source/MassEntity/Public/MassProcessingPhase.h b/Engine/Plugins/Runtime/MassEntity/Source/MassEntity/Public/MassProcessingPhase.h
--- a/Engine/Plugins/Runtime/MassEntity/Source/MassEntity/Public/MassProcessingPhase.h
+++ b/Engine/Plugins/Runtime/MassEntity/Source/MassEntity/Public/MassProcessingPhase.h
@@ -13,7 +13,7 @@
 class UPipeProcessingPhaseManager;
 class UPipeProcessor;
 class UPipeCompositeProcessor;
-class UPipeEntitySubsystem;
+class UMassEntitySubsystem;
 struct FLWCCommandBuffer;
 
 USTRUCT()
@@ -80,7 +80,7 @@ class MASSENTITY_API UPipeProcessingPhaseManager : public UObject
 {
 	GENERATED_BODY()
 public:
-	UPipeEntitySubsystem& GetEntitySubsystemRef() { check(EntitySubsystem); return *EntitySubsystem; }
+	UMassEntitySubsystem& GetEntitySubsystemRef() { check(EntitySubsystem); return *EntitySubsystem; }
 
 	/** Retrieves OnPhaseStart multicast delegate's reference for a given Phase */
 	FPipeProcessingPhase::FOnPhaseEvent& GetOnPhaseStart(const EPipeProcessingPhase Phase) { return ProcessingPhases[uint8(Phase)].OnPhaseStart; }
@@ -98,7 +98,7 @@ public:
 	 *  at hand it's suggested to use that Start version, otherwise call the World-using one. 
 	 */
 	void Start(UWorld& World);
-	void Start(UPipeEntitySubsystem& InEntitySubsystem);
+	void Start(UMassEntitySubsystem& InEntitySubsystem);
 	void Stop();
 	bool IsRunning() const { return EntitySubsystem != nullptr; }
 
@@ -141,7 +141,7 @@ protected:
 	FPipeProcessingPhase ProcessingPhases[(uint8)EPipeProcessingPhase::MAX];
 
 	UPROPERTY()
-	UPipeEntitySubsystem* EntitySubsystem;
+	UMassEntitySubsystem* EntitySubsystem;
 
 	EPipeProcessingPhase CurrentPhase = EPipeProcessingPhase::MAX;
 
diff --git a/Engine/Plugins/Runtime/MassEntity/Source/MassEntity/Public/MassProcessor.h b/Engine/Plugins/Runtime/MassEntity/Source/MassEntity/Public/MassProcessor.h
--- a/Engine/Plugins/Runtime/MassEntity/Source/MassEntity/Public/MassProcessor.h
+++ b/Engine/Plugins/Runtime/MassEntity/Source/MassEntity/Public/MassProcessor.h
@@ -3,7 +3,7 @@
 #pragma once
 
 #include "CoreMinimal.h"
-#include "EntitySubsystem.h"
+#include "MassEntitySubsystem.h"
 #include "MassEntityTypes.h"
 #include "Async/TaskGraphInterfaces.h"
 #include "LWCCommandBuffer.h"
@@ -46,13 +46,13 @@ public:
 	UPipeProcessor(const FObjectInitializer& ObjectInitializer = FObjectInitializer::Get());
 
 	virtual void Initialize(UObject& Owner) {}
-	virtual FGraphEventRef DispatchProcessorTasks(UEntitySubsystem& EntitySubsystem, FLWComponentSystemExecutionContext& ExecutionContext, const FGraphEventArray& Prerequisites = FGraphEventArray());
+	virtual FGraphEventRef DispatchProcessorTasks(UMassEntitySubsystem& EntitySubsystem, FLWComponentSystemExecutionContext& ExecutionContext, const FGraphEventArray& Prerequisites = FGraphEventArray());
 
 	EProcessorExecutionFlags GetExecutionFlags() const { return (EProcessorExecutionFlags)ExecutionFlags; }
 
 	/** Whether this processor should execute according the CurrentExecutionFlags parameters */
 	bool ShouldExecute(const EProcessorExecutionFlags CurrentExecutionFlags) const { return (GetExecutionFlags() & CurrentExecutionFlags) != EProcessorExecutionFlags::None; }
-	void CallExecute(UEntitySubsystem& EntitySubsystem, FLWComponentSystemExecutionContext& Context);
+	void CallExecute(UMassEntitySubsystem& EntitySubsystem, FLWComponentSystemExecutionContext& Context);
 	
 	bool AllowDuplicates() const { return bAllowDuplicates; }
 
@@ -88,7 +88,7 @@ public:
 protected:
 	virtual void ConfigureQueries() PURE_VIRTUAL(UPipeProcessor::ConfigureQueries);
 	virtual void PostInitProperties() override;
-	virtual void Execute(UEntitySubsystem& EntitySubsystem, FLWComponentSystemExecutionContext& Context) PURE_VIRTUAL(UPipeProcessor::Execute);
+	virtual void Execute(UMassEntitySubsystem& EntitySubsystem, FLWComponentSystemExecutionContext& Context) PURE_VIRTUAL(UPipeProcessor::Execute);
 #if WITH_EDITOR
 	virtual void PostEditChangeProperty(FPropertyChangedEvent& PropertyChangedEvent) override;
 #endif // WITH_EDITOR
@@ -161,7 +161,7 @@ public:
 	 *  and if it's missing it will be created and AddGroupedProcessor will be called recursively */
 	void AddGroupedProcessor(FName RequestedGroupName, UPipeProcessor& SubProcessor);
 
-	virtual FGraphEventRef DispatchProcessorTasks(UEntitySubsystem& EntitySubsystem, FLWComponentSystemExecutionContext& ExecutionContext, const FGraphEventArray& Prerequisites = FGraphEventArray()) override;
+	virtual FGraphEventRef DispatchProcessorTasks(UMassEntitySubsystem& EntitySubsystem, FLWComponentSystemExecutionContext& ExecutionContext, const FGraphEventArray& Prerequisites = FGraphEventArray()) override;
 
 	bool IsEmpty() const { return ChildPipeline.IsEmpty(); }
 
@@ -169,7 +169,7 @@ public:
 
 protected:
 	virtual void ConfigureQueries() override;
-	virtual void Execute(UEntitySubsystem& EntitySubsystem, FLWComponentSystemExecutionContext& Context) override;
+	virtual void Execute(UMassEntitySubsystem& EntitySubsystem, FLWComponentSystemExecutionContext& Context) override;
 
 	/**
 	 *  Called recursively to add processors and composite processors to ChildPipeline based on ProcessorsAndGroups
diff --git a/Engine/Plugins/Runtime/MassEntity/Source/MassEntityTestSuite/Private/LWComponentTestFarmPlot.cpp b/Engine/Plugins/Runtime/MassEntity/Source/MassEntityTestSuite/Private/LWComponentTestFarmPlot.cpp
--- a/Engine/Plugins/Runtime/MassEntity/Source/MassEntityTestSuite/Private/LWComponentTestFarmPlot.cpp
+++ b/Engine/Plugins/Runtime/MassEntity/Source/MassEntityTestSuite/Private/LWComponentTestFarmPlot.cpp
@@ -20,7 +20,7 @@ void UComponentUpdateSystem::PostInitProperties()
 //////////////////////////////////////////////////////////////////////////
 // UFarmHarvestTimerSetIcon
 
-void UFarmHarvestTimerSetIcon::Execute(UEntitySubsystem& EntitySubsystem, FLWComponentSystemExecutionContext& Context)
+void UFarmHarvestTimerSetIcon::Execute(UMassEntitySubsystem& EntitySubsystem, FLWComponentSystemExecutionContext& Context)
 {
 	QUICK_SCOPE_CYCLE_COUNTER(SET_ICON_SET_ICON_SET_ICON);
 
@@ -64,7 +64,7 @@ ALWComponentTestFarmPlot::ALWComponentTestFarmPlot()
 	HarvestIconISMC->SetCollisionProfileName(UCollisionProfile::NoCollision_ProfileName);
 }
 
-void ALWComponentTestFarmPlot::AddItemToGrid(UEntitySubsystem* EntitySystem, int32 X, int32 Y, FArchetypeHandle Archetype, int32 VisualIndex)
+void ALWComponentTestFarmPlot::AddItemToGrid(UMassEntitySubsystem* EntitySystem, int32 X, int32 Y, FArchetypeHandle Archetype, int32 VisualIndex)
 {
 	FLWEntity NewItem = EntitySystem->CreateEntity(Archetype);
 	PlantedSquares[X + Y * GridWidth] = NewItem;
@@ -91,7 +91,7 @@ void ALWComponentTestFarmPlot::BeginPlay()
 {
 	Super::BeginPlay();
 
-	UEntitySubsystem* EntitySystem = UWorld::GetSubsystem<UEntitySubsystem>(GetWorld());
+	UMassEntitySubsystem* EntitySystem = UWorld::GetSubsystem<UMassEntitySubsystem>(GetWorld());
 
 	FArchetypeHandle CropArchetype = EntitySystem->CreateArchetype(TArray<const UScriptStruct*>{ FFarmWaterComponent::StaticStruct(), FFarmCropComponent::StaticStruct(), FHarvestTimerComponent::StaticStruct(), FFarmVisualComponent::StaticStruct(), FFarmGridCellData::StaticStruct() });
 	FArchetypeHandle FlowerArchetype = EntitySystem->CreateArchetype(TArray<const UScriptStruct*>{ FFarmWaterComponent::StaticStruct(), FFarmFlowerComponent::StaticStruct(), FHarvestTimerComponent::StaticStruct(), FFarmVisualComponent::StaticStruct(), FFarmGridCellData::StaticStruct() });
@@ -150,7 +150,7 @@ void ALWComponentTestFarmPlot::TickActor(float DeltaTime, enum ELevelTick TickTy
 
 	QUICK_SCOPE_CYCLE_COUNTER(HeyaTick);
 
-	UEntitySubsystem* EntitySystem = UWorld::GetSubsystem<UEntitySubsystem>(GetWorld());
+	UMassEntitySubsystem* EntitySystem = UWorld::GetSubsystem<UMassEntitySubsystem>(GetWorld());
 
 	// Run every frame systems
 	{
diff --git a/Engine/Plugins/Runtime/MassEntity/Source/MassEntityTestSuite/Private/MassCompositeProcessorTest.cpp b/Engine/Plugins/Runtime/MassEntity/Source/MassEntityTestSuite/Private/MassCompositeProcessorTest.cpp
--- a/Engine/Plugins/Runtime/MassEntity/Source/MassEntityTestSuite/Private/MassCompositeProcessorTest.cpp
+++ b/Engine/Plugins/Runtime/MassEntity/Source/MassEntityTestSuite/Private/MassCompositeProcessorTest.cpp
@@ -4,7 +4,7 @@
 #include "AITestsCommon.h"
 
 #include "Engine/World.h"
-#include "MassEntitySystem.h"
+#include "MassEntitySubsystem.h"
 #include "MassEntityTypes.h"
 #include "MassEntityTestTypes.h"
 #include "MassExecutor.h"
@@ -48,7 +48,7 @@ struct FCompositeProcessorTest_MultipleSubProcessors : FEntityTestBase
 			for (int i = 0; i < 3; ++i)
 			{
 				UPipeTestProcessorBase* Processor = NewObject<UPipeTestProcessorBase>(EntitySubsystem);
-				Processor->ExecutionFunction = [Processor, &Result, i](UEntitySubsystem& InEntitySubsystem, FLWComponentSystemExecutionContext& Context) {
+				Processor->ExecutionFunction = [Processor, &Result, i](UMassEntitySubsystem& InEntitySubsystem, FLWComponentSystemExecutionContext& Context) {
 						check(Processor);
 						Result += FMath::Pow(10.f, float(i));
 					};
diff --git a/Engine/Plugins/Runtime/MassEntity/Source/MassEntityTestSuite/Private/MassEntityTest.cpp b/Engine/Plugins/Runtime/MassEntity/Source/MassEntityTestSuite/Private/MassEntityTest.cpp
--- a/Engine/Plugins/Runtime/MassEntity/Source/MassEntityTestSuite/Private/MassEntityTest.cpp
+++ b/Engine/Plugins/Runtime/MassEntity/Source/MassEntityTestSuite/Private/MassEntityTest.cpp
@@ -4,7 +4,7 @@
 #include "AITestsCommon.h"
 
 #include "Engine/World.h"
-#include "MassEntitySystem.h"
+#include "MassEntitySubsystem.h"
 #include "MassEntityTypes.h"
 #include "MassEntityTestTypes.h"
 #include "MassExecutor.h"
diff --git a/Engine/Plugins/Runtime/MassEntity/Source/MassEntityTestSuite/Private/MassExecutionTest.cpp b/Engine/Plugins/Runtime/MassEntity/Source/MassEntityTestSuite/Private/MassExecutionTest.cpp
--- a/Engine/Plugins/Runtime/MassEntity/Source/MassEntityTestSuite/Private/MassExecutionTest.cpp
+++ b/Engine/Plugins/Runtime/MassEntity/Source/MassEntityTestSuite/Private/MassExecutionTest.cpp
@@ -2,7 +2,7 @@
 
 #include "CoreMinimal.h"
 #include "Engine/World.h"
-#include "MassEntitySystem.h"
+#include "MassEntitySubsystem.h"
 #include "MassEntityTypes.h"
 #include "MassEntityTestTypes.h"
 #include "MassExecutor.h"
diff --git a/Engine/Plugins/Runtime/MassEntity/Source/MassEntityTestSuite/Private/MassProcessorTest.cpp b/Engine/Plugins/Runtime/MassEntity/Source/MassEntityTestSuite/Private/MassProcessorTest.cpp
--- a/Engine/Plugins/Runtime/MassEntity/Source/MassEntityTestSuite/Private/MassProcessorTest.cpp
+++ b/Engine/Plugins/Runtime/MassEntity/Source/MassEntityTestSuite/Private/MassProcessorTest.cpp
@@ -4,7 +4,7 @@
 #include "AITestsCommon.h"
 
 #include "Engine/World.h"
-#include "MassEntitySystem.h"
+#include "MassEntitySubsystem.h"
 #include "MassEntityTypes.h"
 #include "MassEntityTestTypes.h"
 #include "MassExecutor.h"
@@ -17,11 +17,11 @@ namespace FPipeProcessorTest
 {
 
 template<typename TProcessor>
-int32 SimpleProcessorRun(UPipeEntitySubsystem& EntitySubsystem)
+int32 SimpleProcessorRun(UMassEntitySubsystem& EntitySubsystem)
 {
 	int32 EntityProcessedCount = 0;
 	TProcessor* Processor = NewObject<TProcessor>(&EntitySubsystem);
-	Processor->ExecutionFunction = [Processor, &EntityProcessedCount](UEntitySubsystem& InEntitySubsystem, FLWComponentSystemExecutionContext& Context) {
+	Processor->ExecutionFunction = [Processor, &EntityProcessedCount](UMassEntitySubsystem& InEntitySubsystem, FLWComponentSystemExecutionContext& Context) {
 		check(Processor);
 		Processor->TestGetQuery().ForEachEntityChunk(InEntitySubsystem, Context, [Processor, &EntityProcessedCount](FLWComponentSystemExecutionContext& Context)
 			{
diff --git a/Engine/Plugins/Runtime/MassEntity/Source/MassEntityTestSuite/Private/MassQueryTest.cpp b/Engine/Plugins/Runtime/MassEntity/Source/MassEntityTestSuite/Private/MassQueryTest.cpp
--- a/Engine/Plugins/Runtime/MassEntity/Source/MassEntityTestSuite/Private/MassQueryTest.cpp
+++ b/Engine/Plugins/Runtime/MassEntity/Source/MassEntityTestSuite/Private/MassQueryTest.cpp
@@ -4,7 +4,7 @@
 #include "AITestsCommon.h"
 
 #include "Engine/World.h"
-#include "MassEntitySystem.h"
+#include "MassEntitySubsystem.h"
 #include "MassEntityTypes.h"
 #include "MassEntityTestTypes.h"
 #include "MassExecutor.h"
@@ -76,7 +76,7 @@ struct FQueryTest_FragmentViewBinding : FEntityTestBase
 		AITEST_TRUE("Initial value of the component should match expectations", TestedFragment.Value == 0.f);
 
 		UPipeTestProcessor_Floats* Processor = NewObject<UPipeTestProcessor_Floats>(EntitySubsystem);
-		Processor->ExecutionFunction = [Processor](UEntitySubsystem& InEntitySubsystem, FLWComponentSystemExecutionContext& Context) {
+		Processor->ExecutionFunction = [Processor](UMassEntitySubsystem& InEntitySubsystem, FLWComponentSystemExecutionContext& Context) {
 			check(Processor);
 			//FLWComponentQuery Query(*Processor);
 
diff --git a/Engine/Plugins/Runtime/MassEntity/Source/MassEntityTestSuite/Private/MassTestTypes.cpp b/Engine/Plugins/Runtime/MassEntity/Source/MassEntityTestSuite/Private/MassTestTypes.cpp
--- a/Engine/Plugins/Runtime/MassEntity/Source/MassEntityTestSuite/Private/MassTestTypes.cpp
+++ b/Engine/Plugins/Runtime/MassEntity/Source/MassEntityTestSuite/Private/MassTestTypes.cpp
@@ -1,7 +1,7 @@
 // Copyright Epic Games, Inc. All Rights Reserved.
 
 #include "MassEntityTestTypes.h"
-#include "MassEntitySystem.h"
+#include "MassEntitySubsystem.h"
 #include "MassExecutor.h"
 #include "Engine/World.h"
 
@@ -11,7 +11,7 @@
 bool FExecutionTestBase::SetUp()
 {
 	World = FAITestHelpers::GetWorld();
-	EntitySubsystem = NewObject<UPipeEntitySubsystem>(World);
+	EntitySubsystem = NewObject<UMassEntitySubsystem>(World);
 	check(EntitySubsystem);
 	struct FSubsystemCollection_TestInit : FSubsystemCollectionBase
 	{
@@ -54,7 +54,7 @@ UPipeTestProcessorBase::UPipeTestProcessorBase()
 	bAutoRegisterWithProcessingPhases = false;
 	ExecutionFlags = int32(EProcessorExecutionFlags::All);
 
-	ExecutionFunction = [](UEntitySubsystem& InEntitySubsystem, FLWComponentSystemExecutionContext& Context) {};
+	ExecutionFunction = [](UMassEntitySubsystem& InEntitySubsystem, FLWComponentSystemExecutionContext& Context) {};
 	RequirementsFunction = [](FLWComponentQuery& Query){};
 }
 
diff --git a/Engine/Plugins/Runtime/MassEntity/Source/MassEntityTestSuite/Public/LWComponentTestFarmPlot.h b/Engine/Plugins/Runtime/MassEntity/Source/MassEntityTestSuite/Public/LWComponentTestFarmPlot.h
--- a/Engine/Plugins/Runtime/MassEntity/Source/MassEntityTestSuite/Public/LWComponentTestFarmPlot.h
+++ b/Engine/Plugins/Runtime/MassEntity/Source/MassEntityTestSuite/Public/LWComponentTestFarmPlot.h
@@ -7,7 +7,7 @@
 #include "Engine/DataTable.h"
 
 #include "LWComponentTypes.h"
-#include "EntitySubsystem.h"
+#include "MassEntitySubsystem.h"
 #include "LWCCommandBuffer.h"
 
 #include "LWComponentTestFarmPlot.generated.h"
@@ -101,7 +101,7 @@ class UComponentUpdateSystem : public UObject
 	GENERATED_BODY()
 public:
 	virtual void PostInitProperties() override;
-	virtual void Execute(UEntitySubsystem& EntitySubsystem, FLWComponentSystemExecutionContext& Context) {}
+	virtual void Execute(UMassEntitySubsystem& EntitySubsystem, FLWComponentSystemExecutionContext& Context) {}
 protected:
 	virtual void ConfigureQueries() {}
 protected:
@@ -120,7 +120,7 @@ public:
 		EntityQuery.AddRequirement<FFarmWaterComponent>(ELWComponentAccess::ReadWrite);
 	}
 
-	virtual void Execute(UEntitySubsystem& EntitySubsystem, FLWComponentSystemExecutionContext& Context) override
+	virtual void Execute(UMassEntitySubsystem& EntitySubsystem, FLWComponentSystemExecutionContext& Context) override
 	{
 		QUICK_SCOPE_CYCLE_COUNTER(UFarmWaterUpdateSystem_Run);
 		EntityQuery.ForEachEntityChunk(EntitySubsystem, Context, [this](FLWComponentSystemExecutionContext& Context) {
@@ -149,7 +149,7 @@ public:
 		EntityQuery.AddRequirement<FFarmFlowerComponent>(ELWComponentAccess::ReadWrite);
 	}
 
-	virtual void Execute(UEntitySubsystem& EntitySubsystem, FLWComponentSystemExecutionContext& Context) override
+	virtual void Execute(UMassEntitySubsystem& EntitySubsystem, FLWComponentSystemExecutionContext& Context) override
 	{
 		QUICK_SCOPE_CYCLE_COUNTER(UFarmHarvestTimerSystem_Flowers_Run);
 		
@@ -192,7 +192,7 @@ public:
 		EntityQuery.AddRequirement<FFarmCropComponent>(ELWComponentAccess::ReadOnly);
 	}
 
-	virtual void Execute(UEntitySubsystem& EntitySubsystem, FLWComponentSystemExecutionContext& Context) override
+	virtual void Execute(UMassEntitySubsystem& EntitySubsystem, FLWComponentSystemExecutionContext& Context) override
 	{
 		QUICK_SCOPE_CYCLE_COUNTER(UFarmHarvestTimerSystem_Crops_Run);
 		
@@ -227,7 +227,7 @@ public:
 		EntityQuery.AddTagRequirement<FFarmReadyToHarvest>(ELWComponentPresence::None);
 	}
 
-	virtual void Execute(UEntitySubsystem& EntitySubsystem, FLWComponentSystemExecutionContext& Context) override
+	virtual void Execute(UMassEntitySubsystem& EntitySubsystem, FLWComponentSystemExecutionContext& Context) override
 	{
 		QUICK_SCOPE_CYCLE_COUNTER(UFarmHarvestTimerExpired_Run);
 		
@@ -269,7 +269,7 @@ public:
 		EntityQuery.AddTagRequirement<FFarmJustBecameReadyToHarvest>(ELWComponentPresence::All);
 	}
 
-	virtual void Execute(UEntitySubsystem& EntitySubsystem, FLWComponentSystemExecutionContext& Context) override;
+	virtual void Execute(UMassEntitySubsystem& EntitySubsystem, FLWComponentSystemExecutionContext& Context) override;
 };
 
 //////////////////////////////////////////////////////////////////////
@@ -337,7 +337,7 @@ private:
 	TArray<int32> FreeHarvestIconIndicies;
 
 private:
-	void AddItemToGrid(UEntitySubsystem* EntitySystem, int32 X, int32 Y, FArchetypeHandle Archetype, int32 VisualIndex);
+	void AddItemToGrid(UMassEntitySubsystem* EntitySystem, int32 X, int32 Y, FArchetypeHandle Archetype, int32 VisualIndex);
 
 public:
 	ALWComponentTestFarmPlot();
diff --git a/Engine/Plugins/Runtime/MassEntity/Source/MassEntityTestSuite/Public/MassEntityTestTypes.h b/Engine/Plugins/Runtime/MassEntity/Source/MassEntityTestSuite/Public/MassEntityTestTypes.h
--- a/Engine/Plugins/Runtime/MassEntity/Source/MassEntityTestSuite/Public/MassEntityTestTypes.h
+++ b/Engine/Plugins/Runtime/MassEntity/Source/MassEntityTestSuite/Public/MassEntityTestTypes.h
@@ -4,13 +4,13 @@
 
 #include "MassProcessor.h"
 #include "LWComponentTypes.h"
-#include "EntitySubsystem.h"
+#include "MassEntitySubsystem.h"
 #include "AITestsCommon.h"
 #include "MassEntityTestTypes.generated.h"
 
 
 class UWorld;
-class UPipeEntitySubsystem;
+class UMassEntitySubsystem;
 
 USTRUCT()
 struct FTestFragment_Float : public FLWComponentData
@@ -72,7 +72,7 @@ class UPipeTestProcessorBase : public UPipeProcessor
 public:
 	UPipeTestProcessorBase();
 	FPipeProcessorExecutionOrder& GetMutableExecutionOrder() { return ExecutionOrder; }
-	virtual void Execute(UEntitySubsystem& EntitySubsystem, FLWComponentSystemExecutionContext& Context) override 
+	virtual void Execute(UMassEntitySubsystem& EntitySubsystem, FLWComponentSystemExecutionContext& Context) override 
 	{
 		ExecutionFunction(EntitySubsystem, Context);
 	}
@@ -81,7 +81,7 @@ public:
 		RequirementsFunction(EntityQuery);
 	}
 
-	TFunction<void(UEntitySubsystem& EntitySubsystem, FLWComponentSystemExecutionContext& Context)> ExecutionFunction;
+	TFunction<void(UMassEntitySubsystem& EntitySubsystem, FLWComponentSystemExecutionContext& Context)> ExecutionFunction;
 	TFunction<void(FLWComponentQuery& Query)> RequirementsFunction;
 
 	FLWComponentQuery& TestGetQuery() { return EntityQuery; }
@@ -156,7 +156,7 @@ public:
 
 struct FExecutionTestBase : FAITestBase
 {
-	UPipeEntitySubsystem* EntitySubsystem = nullptr;
+	UMassEntitySubsystem* EntitySubsystem = nullptr;
 	UWorld* World = nullptr;
 
 	virtual bool SetUp() override;
diff --git a/Engine/Plugins/Runtime/MassGameplay/Source/MassActors/Private/MassActorHelper.cpp b/Engine/Plugins/Runtime/MassGameplay/Source/MassActors/Private/MassActorHelper.cpp
--- a/Engine/Plugins/Runtime/MassGameplay/Source/MassActors/Private/MassActorHelper.cpp
+++ b/Engine/Plugins/Runtime/MassGameplay/Source/MassActors/Private/MassActorHelper.cpp
@@ -3,7 +3,7 @@
 #include "MassActorHelper.h"
 #include "MassActorTypes.h"
 #include "MassActorSubsystem.h"
-#include "MassEntitySystem.h"
+#include "MassEntitySubsystem.h"
 #include "Engine/World.h"
 
 namespace UE::MassActor
@@ -11,7 +11,7 @@ namespace UE::MassActor
 	bool AddEntityTagToActor(const AActor& Actor, const UScriptStruct& TagType)
 	{
 		UWorld* World = Actor.GetWorld();
-		UPipeEntitySubsystem* EntitySystem = UPipeEntitySubsystem::GetCurrent(World);
+		UMassEntitySubsystem* EntitySystem = UMassEntitySubsystem::GetCurrent(World);
 		UMassActorSubsystem* MassActorSubsystem = UWorld::GetSubsystem<UMassActorSubsystem>(World);
 		if (EntitySystem && MassActorSubsystem)
 		{
@@ -36,7 +36,7 @@ namespace UE::MassActor
 	bool RemoveEntityTagFromActor(const AActor& Actor, const UScriptStruct& TagType)
 	{
 		UWorld* World = Actor.GetWorld();
-		UPipeEntitySubsystem* EntitySystem = UPipeEntitySubsystem::GetCurrent(World);
+		UMassEntitySubsystem* EntitySystem = UMassEntitySubsystem::GetCurrent(World);
 		UMassActorSubsystem* MassActorSubsystem = UWorld::GetSubsystem<UMassActorSubsystem>(World);
 		if (EntitySystem && MassActorSubsystem)
 		{
diff --git a/Engine/Plugins/Runtime/MassGameplay/Source/MassActors/Private/MassActorSubsystem.cpp b/Engine/Plugins/Runtime/MassGameplay/Source/MassActors/Private/MassActorSubsystem.cpp
--- a/Engine/Plugins/Runtime/MassGameplay/Source/MassActors/Private/MassActorSubsystem.cpp
+++ b/Engine/Plugins/Runtime/MassGameplay/Source/MassActors/Private/MassActorSubsystem.cpp
@@ -3,7 +3,7 @@
 #include "MassActorSubsystem.h"
 #include "MassCommonTypes.h"
 #include "Engine/World.h"
-#include "MassEntitySystem.h"
+#include "MassEntitySubsystem.h"
 #include "MassSimulationSubsystem.h"
 #include "VisualLogger/VisualLogger.h"
 
@@ -62,7 +62,7 @@ void UMassActorSubsystem::Initialize(FSubsystemCollectionBase& Collection)
 	// making sure UMassSimulationSubsystem gets created before the MassActorManager
 	Collection.InitializeDependency<UMassSimulationSubsystem>();
 	
-	EntitySystem = UPipeEntitySubsystem::GetCurrent(GetWorld());
+	EntitySystem = UMassEntitySubsystem::GetCurrent(GetWorld());
 }
 
 FMassHandle UMassActorSubsystem::GetHandleFromActor(const TObjectKey<const AActor> Actor)
diff --git a/Engine/Plugins/Runtime/MassGameplay/Source/MassActors/Private/MassAgentComponent.cpp b/Engine/Plugins/Runtime/MassGameplay/Source/MassActors/Private/MassAgentComponent.cpp
--- a/Engine/Plugins/Runtime/MassGameplay/Source/MassActors/Private/MassAgentComponent.cpp
+++ b/Engine/Plugins/Runtime/MassGameplay/Source/MassActors/Private/MassAgentComponent.cpp
@@ -3,7 +3,7 @@
 #include "MassAgentComponent.h"
 #include "EngineUtils.h"
 #include "EntityView.h"
-#include "MassEntitySystem.h"
+#include "MassEntitySubsystem.h"
 #include "MassCommonTypes.h"
 #include "MassAgentSubsystem.h"
 #include "VisualLogger/VisualLogger.h"
@@ -162,7 +162,7 @@ void UMassAgentComponent::SetEntityHandleInternal(FMassHandle NewHandle)
 
 #if	UE_REPLICATION_COMPILE_SERVER_CODE
 	// Fetch NetID if it exist
-	if (const UPipeEntitySubsystem* EntitySubsystem = UPipeEntitySubsystem::GetCurrent(GetWorld()))
+	if (const UMassEntitySubsystem* EntitySubsystem = UMassEntitySubsystem::GetCurrent(GetWorld()))
 	{
 		if (const FMassNetworkIDFragment* NetIDFragment = EntitySubsystem->GetComponentDataPtr<FMassNetworkIDFragment>(AgentHandle.GetLWEntity()))
 		{
@@ -188,7 +188,7 @@ void UMassAgentComponent::SetEntityHandleInternal(FMassHandle NewHandle)
 	}
 
 	// Sync up with mass
-	if (const UPipeEntitySubsystem* EntitySubsystem = UPipeEntitySubsystem::GetCurrent(GetWorld()))
+	if (const UMassEntitySubsystem* EntitySubsystem = UMassEntitySubsystem::GetCurrent(GetWorld()))
 	{
 		if (IsNetSimulating())
 		{
@@ -277,7 +277,7 @@ void UMassAgentComponent::ClearEntityHandleInternal()
 	}
 
 	// Sync up with mass
-	if (const UPipeEntitySubsystem* EntitySubsystem = UPipeEntitySubsystem::GetCurrent(GetWorld()))
+	if (const UMassEntitySubsystem* EntitySubsystem = UMassEntitySubsystem::GetCurrent(GetWorld()))
 	{
 		if (IsNetSimulating())
 		{
@@ -291,7 +291,7 @@ void UMassAgentComponent::ClearEntityHandleInternal()
 	}
 
 #if WITH_EDITORONLY_DATA
-	DebugEntity = UEntitySubsystem::InvalidEntity;
+	DebugEntity = UMassEntitySubsystem::InvalidEntity;
 #endif // WITH_EDITORONLY_DATA
 	AgentHandle = FMassHandle::InvalidHandle;
 }
@@ -352,7 +352,7 @@ void UMassAgentComponent::DebugCheckStateConsistency()
 			MASSAGENT_CHECK(bValidAgentHandle, TEXT("Expecting a valid mass agent handle in state %s"), *UEnum::GetValueAsString(State));
 			if (bValidAgentHandle)
 			{
-				if (UPipeEntitySubsystem* EntitySubsystem = UPipeEntitySubsystem::GetCurrent(GetWorld()))
+				if (UMassEntitySubsystem* EntitySubsystem = UMassEntitySubsystem::GetCurrent(GetWorld()))
 				{
 					const bool bIsValidEntity = EntitySubsystem->IsEntityValid(AgentHandle.GetLWEntity());
 					MASSAGENT_CHECK(bIsValidEntity, TEXT("Exepecting a valid entity in state"), *UEnum::GetValueAsString(State))
@@ -418,7 +418,7 @@ void UMassAgentComponent::KillEntity(const bool bDestroyActor)
 
 	AActor* Owner = GetOwner();
 	UWorld* World = GetWorld();
-	UPipeEntitySubsystem* EntitySubsystem = UPipeEntitySubsystem::GetCurrent(World);
+	UMassEntitySubsystem* EntitySubsystem = UMassEntitySubsystem::GetCurrent(World);
 	if (Owner == nullptr || EntitySubsystem == nullptr)
 	{
 		return;
diff --git a/Engine/Plugins/Runtime/MassGameplay/Source/MassActors/Private/MassAgentSubsystem.cpp b/Engine/Plugins/Runtime/MassGameplay/Source/MassActors/Private/MassAgentSubsystem.cpp
--- a/Engine/Plugins/Runtime/MassGameplay/Source/MassActors/Private/MassAgentSubsystem.cpp
+++ b/Engine/Plugins/Runtime/MassGameplay/Source/MassActors/Private/MassAgentSubsystem.cpp
@@ -3,7 +3,7 @@
 #include "MassAgentSubsystem.h"
 #include "LWCCommandBuffer.h"
 #include "Engine/World.h"
-#include "MassEntitySystem.h"
+#include "MassEntitySubsystem.h"
 #include "MassCommonTypes.h"
 #include "MassSimulationSubsystem.h"
 #include "MassSpawnerSubsystem.h"
@@ -46,7 +46,7 @@ void UMassAgentSubsystem::Initialize(FSubsystemCollectionBase& Collection)
 #endif // UE_REPLICATION_COMPILE_CLIENT_CODE
 	
 	UWorld* World = GetWorld();
-	EntitySystem = UPipeEntitySubsystem::GetCurrent(World);
+	EntitySystem = UMassEntitySubsystem::GetCurrent(World);
 	SpawnerSystem = UWorld::GetSubsystem<UMassSpawnerSubsystem>(World);
 
 	SimulationSystem = UWorld::GetSubsystem<UMassSimulationSubsystem>(World);
diff --git a/Engine/Plugins/Runtime/MassGameplay/Source/MassActors/Private/Translators/MassCapsuleComponentTranslators.cpp b/Engine/Plugins/Runtime/MassGameplay/Source/MassActors/Private/Translators/MassCapsuleComponentTranslators.cpp
--- a/Engine/Plugins/Runtime/MassGameplay/Source/MassActors/Private/Translators/MassCapsuleComponentTranslators.cpp
+++ b/Engine/Plugins/Runtime/MassGameplay/Source/MassActors/Private/Translators/MassCapsuleComponentTranslators.cpp
@@ -2,7 +2,7 @@
 
 #include "Translators/MassCapsuleComponentTranslators.h"
 #include "MassCommonTypes.h"
-#include "MassEntitySystem.h"
+#include "MassEntitySubsystem.h"
 #include "Components/CapsuleComponent.h"
 #include "MassMovementTypes.h"
 
@@ -24,7 +24,7 @@ void UMassCapsuleTransformToMassTranslator::ConfigureQueries()
 	EntityQuery.AddRequirement<FDataFragment_Transform>(ELWComponentAccess::ReadWrite);
 }
 
-void UMassCapsuleTransformToMassTranslator::Execute(UEntitySubsystem& EntitySubsystem, FLWComponentSystemExecutionContext& Context)
+void UMassCapsuleTransformToMassTranslator::Execute(UMassEntitySubsystem& EntitySubsystem, FLWComponentSystemExecutionContext& Context)
 {
 	EntityQuery.ForEachEntityChunk(EntitySubsystem, Context, [this](FLWComponentSystemExecutionContext& Context)
 		{
@@ -58,7 +58,7 @@ void UMassTransformToActorCapsuleTranslator::ConfigureQueries()
 	EntityQuery.AddRequirement<FDataFragment_Transform>(ELWComponentAccess::ReadOnly);
 }
 
-void UMassTransformToActorCapsuleTranslator::Execute(UEntitySubsystem& EntitySubsystem, FLWComponentSystemExecutionContext& Context)
+void UMassTransformToActorCapsuleTranslator::Execute(UMassEntitySubsystem& EntitySubsystem, FLWComponentSystemExecutionContext& Context)
 {
 	EntityQuery.ForEachEntityChunk(EntitySubsystem, Context, [this](FLWComponentSystemExecutionContext& Context)
 		{
diff --git a/Engine/Plugins/Runtime/MassGameplay/Source/MassActors/Private/Translators/MassCharacterMovementTranslators.cpp b/Engine/Plugins/Runtime/MassGameplay/Source/MassActors/Private/Translators/MassCharacterMovementTranslators.cpp
--- a/Engine/Plugins/Runtime/MassGameplay/Source/MassActors/Private/Translators/MassCharacterMovementTranslators.cpp
+++ b/Engine/Plugins/Runtime/MassGameplay/Source/MassActors/Private/Translators/MassCharacterMovementTranslators.cpp
@@ -4,7 +4,7 @@
 #include "Logging/LogMacros.h"
 #include "GameFramework/CharacterMovementComponent.h"
 #include "VisualLogger/VisualLogger.h"
-#include "MassEntitySystem.h"
+#include "MassEntitySubsystem.h"
 #include "MassCommonTypes.h"
 #include "MassMovementTypes.h"
 #include "MassMovementFragments.h"
@@ -28,7 +28,7 @@ void UMassCharacterMovementToMassTranslator::ConfigureQueries()
 	EntityQuery.AddRequirement<FMassVelocityFragment>(ELWComponentAccess::ReadWrite);
 }
 
-void UMassCharacterMovementToMassTranslator::Execute(UEntitySubsystem& EntitySubsystem, FLWComponentSystemExecutionContext& Context)
+void UMassCharacterMovementToMassTranslator::Execute(UMassEntitySubsystem& EntitySubsystem, FLWComponentSystemExecutionContext& Context)
 {
 	EntityQuery.ForEachEntityChunk(EntitySubsystem, Context, [this](FLWComponentSystemExecutionContext& Context)
 	{
@@ -68,7 +68,7 @@ void UMassCharacterMovementToActorTranslator::ConfigureQueries()
 	EntityQuery.AddRequirement<FMassVelocityFragment>(ELWComponentAccess::ReadOnly);
 }
 
-void UMassCharacterMovementToActorTranslator::Execute(UEntitySubsystem& EntitySubsystem, FLWComponentSystemExecutionContext& Context)
+void UMassCharacterMovementToActorTranslator::Execute(UMassEntitySubsystem& EntitySubsystem, FLWComponentSystemExecutionContext& Context)
 {
 	EntityQuery.ForEachEntityChunk(EntitySubsystem, Context, [this](FLWComponentSystemExecutionContext& Context)
 	{
@@ -104,7 +104,7 @@ void UMassCharacterOrientationToMassTranslator::ConfigureQueries()
 	EntityQuery.AddRequirement<FDataFragment_Transform>(ELWComponentAccess::ReadWrite);
 }
 
-void UMassCharacterOrientationToMassTranslator::Execute(UEntitySubsystem& EntitySubsystem, FLWComponentSystemExecutionContext& Context)
+void UMassCharacterOrientationToMassTranslator::Execute(UMassEntitySubsystem& EntitySubsystem, FLWComponentSystemExecutionContext& Context)
 {
 	EntityQuery.ForEachEntityChunk(EntitySubsystem, Context, [this](FLWComponentSystemExecutionContext& Context)
 	{
@@ -144,7 +144,7 @@ void UMassCharacterOrientationToActorTranslator::ConfigureQueries()
 	EntityQuery.AddRequirement<FDataFragment_Transform>(ELWComponentAccess::ReadOnly);
 }
 
-void UMassCharacterOrientationToActorTranslator::Execute(UEntitySubsystem& EntitySubsystem, FLWComponentSystemExecutionContext& Context)
+void UMassCharacterOrientationToActorTranslator::Execute(UMassEntitySubsystem& EntitySubsystem, FLWComponentSystemExecutionContext& Context)
 {
 	EntityQuery.ForEachEntityChunk(EntitySubsystem, Context, [this](FLWComponentSystemExecutionContext& Context)
 	{
@@ -181,7 +181,7 @@ void UMassFragmentInitializer_NavLocation::ConfigureQueries()
 	EntityQuery.AddRequirement<FDataFragment_NavLocation>(ELWComponentAccess::ReadWrite);
 }
 
-void UMassFragmentInitializer_NavLocation::Execute(UEntitySubsystem& EntitySubsystem, FLWComponentSystemExecutionContext& Context)
+void UMassFragmentInitializer_NavLocation::Execute(UMassEntitySubsystem& EntitySubsystem, FLWComponentSystemExecutionContext& Context)
 {
 	EntityQuery.ForEachEntityChunk(EntitySubsystem, Context, [this](FLWComponentSystemExecutionContext& Context)
 		{
@@ -208,7 +208,7 @@ void UMassFragmentInitializer_Transform::ConfigureQueries()
 	EntityQuery.AddRequirement<FDataFragment_Transform>(ELWComponentAccess::ReadWrite);
 }
 
-void UMassFragmentInitializer_Transform::Execute(UEntitySubsystem& EntitySubsystem, FLWComponentSystemExecutionContext& Context)
+void UMassFragmentInitializer_Transform::Execute(UMassEntitySubsystem& EntitySubsystem, FLWComponentSystemExecutionContext& Context)
 {
 	const UWorld* World = EntitySubsystem.GetWorld();
 
diff --git a/Engine/Plugins/Runtime/MassGameplay/Source/MassActors/Private/Translators/MassSceneComponentLocationTranslator.cpp b/Engine/Plugins/Runtime/MassGameplay/Source/MassActors/Private/Translators/MassSceneComponentLocationTranslator.cpp
--- a/Engine/Plugins/Runtime/MassGameplay/Source/MassActors/Private/Translators/MassSceneComponentLocationTranslator.cpp
+++ b/Engine/Plugins/Runtime/MassGameplay/Source/MassActors/Private/Translators/MassSceneComponentLocationTranslator.cpp
@@ -3,7 +3,7 @@
 #include "Translators/MassSceneComponentLocationTranslator.h"
 #include "MassMovementTypes.h"
 #include "Components/SceneComponent.h"
-#include "MassEntitySystem.h"
+#include "MassEntitySubsystem.h"
 
 //----------------------------------------------------------------------//
 //  UMassSceneComponentLocationToMassTranslator
@@ -22,7 +22,7 @@ void UMassSceneComponentLocationToMassTranslator::ConfigureQueries()
 	EntityQuery.AddRequirement<FDataFragment_Transform>(ELWComponentAccess::ReadWrite);
 }
 
-void UMassSceneComponentLocationToMassTranslator::Execute(UEntitySubsystem& EntitySubsystem, FLWComponentSystemExecutionContext& Context)
+void UMassSceneComponentLocationToMassTranslator::Execute(UMassEntitySubsystem& EntitySubsystem, FLWComponentSystemExecutionContext& Context)
 {
 	EntityQuery.ForEachEntityChunk(EntitySubsystem, Context, [this](FLWComponentSystemExecutionContext& Context)
 	{
@@ -58,7 +58,7 @@ void UMassSceneComponentLocationToActorTranslator::ConfigureQueries()
 	EntityQuery.AddRequirement<FDataFragment_Transform>(ELWComponentAccess::ReadWrite);
 }
 
-void UMassSceneComponentLocationToActorTranslator::Execute(UEntitySubsystem& EntitySubsystem, FLWComponentSystemExecutionContext& Context)
+void UMassSceneComponentLocationToActorTranslator::Execute(UMassEntitySubsystem& EntitySubsystem, FLWComponentSystemExecutionContext& Context)
 {
 	EntityQuery.ForEachEntityChunk(EntitySubsystem, Context, [this](FLWComponentSystemExecutionContext& Context)
 		{
diff --git a/Engine/Plugins/Runtime/MassGameplay/Source/MassActors/Private/Translators/MassTranslators_BehaviorTree.cpp b/Engine/Plugins/Runtime/MassGameplay/Source/MassActors/Private/Translators/MassTranslators_BehaviorTree.cpp
--- a/Engine/Plugins/Runtime/MassGameplay/Source/MassActors/Private/Translators/MassTranslators_BehaviorTree.cpp
+++ b/Engine/Plugins/Runtime/MassGameplay/Source/MassActors/Private/Translators/MassTranslators_BehaviorTree.cpp
@@ -2,7 +2,7 @@
 
 #include "Translators/MassTranslators_BehaviorTree.h"
 #include "BehaviorTree/BehaviorTreeComponent.h"
-#include "MassEntitySystem.h"
+#include "MassEntitySubsystem.h"
 
 UMassTranslator_BehaviorTree::UMassTranslator_BehaviorTree()
 {
diff --git a/Engine/Plugins/Runtime/MassGameplay/Source/MassActors/Public/MassActorSubsystem.h b/Engine/Plugins/Runtime/MassGameplay/Source/MassActors/Public/MassActorSubsystem.h
--- a/Engine/Plugins/Runtime/MassGameplay/Source/MassActors/Public/MassActorSubsystem.h
+++ b/Engine/Plugins/Runtime/MassGameplay/Source/MassActors/Public/MassActorSubsystem.h
@@ -12,7 +12,7 @@
 
 struct FMassHandle;
 class AActor;
-class UPipeEntitySubsystem;
+class UMassEntitySubsystem;
 
 /**
  * Light weight component data to save the actor pointer of a light weight entity if it exist
@@ -108,5 +108,5 @@ protected:
 	UE_MT_DECLARE_RW_ACCESS_DETECTOR(ActorHandleMapDetector);
 	
 	UPROPERTY()
-	UPipeEntitySubsystem* EntitySystem;
+	UMassEntitySubsystem* EntitySystem;
 };
diff --git a/Engine/Plugins/Runtime/MassGameplay/Source/MassActors/Public/MassAgentSubsystem.h b/Engine/Plugins/Runtime/MassGameplay/Source/MassActors/Public/MassAgentSubsystem.h
--- a/Engine/Plugins/Runtime/MassGameplay/Source/MassActors/Public/MassAgentSubsystem.h
+++ b/Engine/Plugins/Runtime/MassGameplay/Source/MassActors/Public/MassAgentSubsystem.h
@@ -10,7 +10,7 @@
 
 struct FMassHandle;
 class AActor;
-class UPipeEntitySubsystem;
+class UMassEntitySubsystem;
 class UMassSpawnerSubsystem;
 class UMassAgentComponent;
 class UMassSimulationSubsystem;
@@ -123,7 +123,7 @@ protected:
 
 protected:
 	UPROPERTY()
-	UPipeEntitySubsystem* EntitySystem;
+	UMassEntitySubsystem* EntitySystem;
 
 	UPROPERTY()
 	UMassSpawnerSubsystem* SpawnerSystem;
diff --git a/Engine/Plugins/Runtime/MassGameplay/Source/MassActors/Public/Translators/MassCapsuleComponentTranslators.h b/Engine/Plugins/Runtime/MassGameplay/Source/MassActors/Public/Translators/MassCapsuleComponentTranslators.h
--- a/Engine/Plugins/Runtime/MassGameplay/Source/MassActors/Public/Translators/MassCapsuleComponentTranslators.h
+++ b/Engine/Plugins/Runtime/MassGameplay/Source/MassActors/Public/Translators/MassCapsuleComponentTranslators.h
@@ -40,7 +40,7 @@ public:
 
 protected:
 	virtual void ConfigureQueries() override;
-	virtual void Execute(UEntitySubsystem& EntitySubsystem, FLWComponentSystemExecutionContext& Context) override;	
+	virtual void Execute(UMassEntitySubsystem& EntitySubsystem, FLWComponentSystemExecutionContext& Context) override;	
 
 	FLWComponentQuery EntityQuery;
 };
@@ -60,7 +60,7 @@ public:
 
 protected:
 	virtual void ConfigureQueries() override;
-	virtual void Execute(UEntitySubsystem& EntitySubsystem, FLWComponentSystemExecutionContext& Context) override;
+	virtual void Execute(UMassEntitySubsystem& EntitySubsystem, FLWComponentSystemExecutionContext& Context) override;
 
 	FLWComponentQuery EntityQuery;
 };
diff --git a/Engine/Plugins/Runtime/MassGameplay/Source/MassActors/Public/Translators/MassCharacterMovementTranslators.h b/Engine/Plugins/Runtime/MassGameplay/Source/MassActors/Public/Translators/MassCharacterMovementTranslators.h
--- a/Engine/Plugins/Runtime/MassGameplay/Source/MassActors/Public/Translators/MassCharacterMovementTranslators.h
+++ b/Engine/Plugins/Runtime/MassGameplay/Source/MassActors/Public/Translators/MassCharacterMovementTranslators.h
@@ -35,7 +35,7 @@ public:
 
 protected:
 	virtual void ConfigureQueries() override;
-	virtual void Execute(UEntitySubsystem& EntitySubsystem, FLWComponentSystemExecutionContext& Context) override;
+	virtual void Execute(UMassEntitySubsystem& EntitySubsystem, FLWComponentSystemExecutionContext& Context) override;
 
 	FLWComponentQuery EntityQuery;
 };
@@ -57,7 +57,7 @@ public:
 
 protected:
 	virtual void ConfigureQueries() override;
-	virtual void Execute(UEntitySubsystem& EntitySubsystem, FLWComponentSystemExecutionContext& Context) override;
+	virtual void Execute(UMassEntitySubsystem& EntitySubsystem, FLWComponentSystemExecutionContext& Context) override;
 
 	FLWComponentQuery EntityQuery;
 };
@@ -79,7 +79,7 @@ public:
 
 protected:
 	virtual void ConfigureQueries() override;
-	virtual void Execute(UEntitySubsystem& EntitySubsystem, FLWComponentSystemExecutionContext& Context) override;
+	virtual void Execute(UMassEntitySubsystem& EntitySubsystem, FLWComponentSystemExecutionContext& Context) override;
 
 	FLWComponentQuery EntityQuery;
 };
@@ -101,7 +101,7 @@ public:
 
 protected:
 	virtual void ConfigureQueries() override;
-	virtual void Execute(UEntitySubsystem& EntitySubsystem, FLWComponentSystemExecutionContext& Context) override;
+	virtual void Execute(UMassEntitySubsystem& EntitySubsystem, FLWComponentSystemExecutionContext& Context) override;
 
 	FLWComponentQuery EntityQuery;
 };
@@ -115,7 +115,7 @@ public:
 
 protected:
 	virtual void ConfigureQueries() override;
-	virtual void Execute(UEntitySubsystem& EntitySubsystem, FLWComponentSystemExecutionContext& Context) override;
+	virtual void Execute(UMassEntitySubsystem& EntitySubsystem, FLWComponentSystemExecutionContext& Context) override;
 
 	FLWComponentQuery EntityQuery;
 };
@@ -130,7 +130,7 @@ public:
 
 protected:
 	virtual void ConfigureQueries() override;
-	virtual void Execute(UEntitySubsystem& EntitySubsystem, FLWComponentSystemExecutionContext& Context) override;
+	virtual void Execute(UMassEntitySubsystem& EntitySubsystem, FLWComponentSystemExecutionContext& Context) override;
 
 	FLWComponentQuery EntityQuery;
 };
\ No newline at end of file
diff --git a/Engine/Plugins/Runtime/MassGameplay/Source/MassActors/Public/Translators/MassSceneComponentLocationTranslator.h b/Engine/Plugins/Runtime/MassGameplay/Source/MassActors/Public/Translators/MassSceneComponentLocationTranslator.h
--- a/Engine/Plugins/Runtime/MassGameplay/Source/MassActors/Public/Translators/MassSceneComponentLocationTranslator.h
+++ b/Engine/Plugins/Runtime/MassGameplay/Source/MassActors/Public/Translators/MassSceneComponentLocationTranslator.h
@@ -30,7 +30,7 @@ public:
 
 protected:
 	virtual void ConfigureQueries() override;
-	virtual void Execute(UEntitySubsystem& EntitySubsystem, FLWComponentSystemExecutionContext& Context) override;
+	virtual void Execute(UMassEntitySubsystem& EntitySubsystem, FLWComponentSystemExecutionContext& Context) override;
 
 	FLWComponentQuery EntityQuery;
 };
@@ -52,7 +52,7 @@ public:
 
 protected:
 	virtual void ConfigureQueries() override;
-	virtual void Execute(UEntitySubsystem& EntitySubsystem, FLWComponentSystemExecutionContext& Context) override;
+	virtual void Execute(UMassEntitySubsystem& EntitySubsystem, FLWComponentSystemExecutionContext& Context) override;
 
 	FLWComponentQuery EntityQuery;
 };
\ No newline at end of file
diff --git a/Engine/Plugins/Runtime/MassGameplay/Source/MassActors/Public/Translators/MassTranslators_BehaviorTree.h b/Engine/Plugins/Runtime/MassGameplay/Source/MassActors/Public/Translators/MassTranslators_BehaviorTree.h
--- a/Engine/Plugins/Runtime/MassGameplay/Source/MassActors/Public/Translators/MassTranslators_BehaviorTree.h
+++ b/Engine/Plugins/Runtime/MassGameplay/Source/MassActors/Public/Translators/MassTranslators_BehaviorTree.h
@@ -26,7 +26,7 @@ public:
 
 protected:
 	virtual void ConfigureQueries() override;
-	virtual void Execute(UEntitySubsystem& EntitySubsystem, FLWComponentSystemExecutionContext& Context) override {}
+	virtual void Execute(UMassEntitySubsystem& EntitySubsystem, FLWComponentSystemExecutionContext& Context) override {}
 
 	FLWComponentQuery EntityQuery;
 };
diff --git a/Engine/Plugins/Runtime/MassGameplay/Source/MassCommon/Private/MassCommonTypes.cpp b/Engine/Plugins/Runtime/MassGameplay/Source/MassCommon/Private/MassCommonTypes.cpp
--- a/Engine/Plugins/Runtime/MassGameplay/Source/MassCommon/Private/MassCommonTypes.cpp
+++ b/Engine/Plugins/Runtime/MassGameplay/Source/MassCommon/Private/MassCommonTypes.cpp
@@ -1,7 +1,7 @@
 // Copyright Epic Games, Inc. All Rights Reserved.
 
 #include "MassCommonTypes.h"
-#include "EntitySubsystem.h"
+#include "MassEntitySubsystem.h"
 #include "Math/ColorList.h"
 
 #if WITH_MASS_DEBUG
@@ -103,7 +103,7 @@ FColor GetEntityDebugColor(FLWEntity Entity)
 }} // namespace UE::MassDebug
 #endif // WITH_MASS_DEBUG
 
-const FMassHandle FMassHandle::InvalidHandle(UEntitySubsystem::InvalidEntity);
+const FMassHandle FMassHandle::InvalidHandle(UMassEntitySubsystem::InvalidEntity);
 
 //----------------------------------------------------------------------//
 // FMassUniqueFragmentCollection 
diff --git a/Engine/Plugins/Runtime/MassGameplay/Source/MassCommon/Private/MassGenericProcessors.cpp b/Engine/Plugins/Runtime/MassGameplay/Source/MassCommon/Private/MassGenericProcessors.cpp
--- a/Engine/Plugins/Runtime/MassGameplay/Source/MassCommon/Private/MassGenericProcessors.cpp
+++ b/Engine/Plugins/Runtime/MassGameplay/Source/MassCommon/Private/MassGenericProcessors.cpp
@@ -12,7 +12,7 @@ URandomizeVectorProcessor::URandomizeVectorProcessor()
 	bAutoRegisterWithProcessingPhases = false;
 }
 
-void URandomizeVectorProcessor::Execute(UEntitySubsystem& EntitySubsystem, FLWComponentSystemExecutionContext& Context)
+void URandomizeVectorProcessor::Execute(UMassEntitySubsystem& EntitySubsystem, FLWComponentSystemExecutionContext& Context)
 {
 	QUICK_SCOPE_CYCLE_COUNTER(RandomizeVectorProcessor_Run);
 	
diff --git a/Engine/Plugins/Runtime/MassGameplay/Source/MassCommon/Public/MassGenericProcessors.h b/Engine/Plugins/Runtime/MassGameplay/Source/MassCommon/Public/MassGenericProcessors.h
--- a/Engine/Plugins/Runtime/MassGameplay/Source/MassCommon/Public/MassGenericProcessors.h
+++ b/Engine/Plugins/Runtime/MassGameplay/Source/MassCommon/Public/MassGenericProcessors.h
@@ -35,7 +35,7 @@ public:
 		}
 	}
 
-	virtual void Execute(UEntitySubsystem& EntitySubsystem, FLWComponentSystemExecutionContext& Context) override;
+	virtual void Execute(UMassEntitySubsystem& EntitySubsystem, FLWComponentSystemExecutionContext& Context) override;
 
 	UPROPERTY(EditAnywhere, Category = "Mass", meta = (ClampMin = 1, UIMin = 1))
 	float MinMagnitude = 1.f;
diff --git a/Engine/Plugins/Runtime/MassGameplay/Source/MassGameplayDebug/Private/AssignDebugVisProcessor.cpp b/Engine/Plugins/Runtime/MassGameplay/Source/MassGameplayDebug/Private/AssignDebugVisProcessor.cpp
--- a/Engine/Plugins/Runtime/MassGameplay/Source/MassGameplayDebug/Private/AssignDebugVisProcessor.cpp
+++ b/Engine/Plugins/Runtime/MassGameplay/Source/MassGameplayDebug/Private/AssignDebugVisProcessor.cpp
@@ -20,7 +20,7 @@ void UAssignDebugVisProcessor::ConfigureQueries()
 	EntityQuery.AddRequirement<FSimDebugVisComponent>(ELWComponentAccess::ReadWrite);
 }
 
-void UAssignDebugVisProcessor::Execute(UEntitySubsystem& EntitySubsystem, FLWComponentSystemExecutionContext& Context)
+void UAssignDebugVisProcessor::Execute(UMassEntitySubsystem& EntitySubsystem, FLWComponentSystemExecutionContext& Context)
 {
 	UMassDebugVisualizationComponent* Visualizer = WeakVisualizer.Get();
 	if (!ensure(Visualizer))
diff --git a/Engine/Plugins/Runtime/MassGameplay/Source/MassGameplayDebug/Private/DebugVisLocationProcessor.cpp b/Engine/Plugins/Runtime/MassGameplay/Source/MassGameplayDebug/Private/DebugVisLocationProcessor.cpp
--- a/Engine/Plugins/Runtime/MassGameplay/Source/MassGameplayDebug/Private/DebugVisLocationProcessor.cpp
+++ b/Engine/Plugins/Runtime/MassGameplay/Source/MassGameplayDebug/Private/DebugVisLocationProcessor.cpp
@@ -5,7 +5,7 @@
 #include "MassDebugVisualizationComponent.h"
 #include "MassCommonFragments.h"
 #include "Components/HierarchicalInstancedStaticMeshComponent.h"
-#include "EntitySubsystem.h"
+#include "MassEntitySubsystem.h"
 #include "MassMovementFragments.h"
 
 //----------------------------------------------------------------------//
@@ -23,7 +23,7 @@ void UDebugVisLocationProcessor::ConfigureQueries()
 	EntityQuery.AddRequirement<FMassVelocityFragment>(ELWComponentAccess::ReadOnly);
 }
 
-void UDebugVisLocationProcessor::Execute(UEntitySubsystem& EntitySubsystem, FLWComponentSystemExecutionContext& Context)
+void UDebugVisLocationProcessor::Execute(UMassEntitySubsystem& EntitySubsystem, FLWComponentSystemExecutionContext& Context)
 {
 #if WITH_EDITORONLY_DATA
 	UMassDebugVisualizationComponent* Visualizer = WeakVisualizer.Get();
@@ -92,7 +92,7 @@ void UMassProcessor_UpdateDebugVis::ConfigureQueries()
 	EntityQuery.AddRequirement<FDataFragment_AgentRadius>(ELWComponentAccess::ReadWrite);
 }
 
-void UMassProcessor_UpdateDebugVis::Execute(UEntitySubsystem& EntitySubsystem, FLWComponentSystemExecutionContext& Context)
+void UMassProcessor_UpdateDebugVis::Execute(UMassEntitySubsystem& EntitySubsystem, FLWComponentSystemExecutionContext& Context)
 {
 	UMassDebuggerSubsystem* Debugger = UWorld::GetSubsystem<UMassDebuggerSubsystem>(GetWorld());
 	if (Debugger == nullptr)
diff --git a/Engine/Plugins/Runtime/MassGameplay/Source/MassGameplayDebug/Private/MassDebuggerSubsystem.cpp b/Engine/Plugins/Runtime/MassGameplay/Source/MassGameplayDebug/Private/MassDebuggerSubsystem.cpp
--- a/Engine/Plugins/Runtime/MassGameplay/Source/MassGameplayDebug/Private/MassDebuggerSubsystem.cpp
+++ b/Engine/Plugins/Runtime/MassGameplay/Source/MassGameplayDebug/Private/MassDebuggerSubsystem.cpp
@@ -5,7 +5,7 @@
 #include "MassCommonTypes.h"
 #include "MassSimulationSubsystem.h"
 #include "MassDebugVisualizationComponent.h"
-#include "MassEntitySystem.h"
+#include "MassEntitySubsystem.h"
 #include "MassDebugVisualizer.h"
 
 
diff --git a/Engine/Plugins/Runtime/MassGameplay/Source/MassGameplayDebug/Public/AssignDebugVisProcessor.h b/Engine/Plugins/Runtime/MassGameplay/Source/MassGameplayDebug/Public/AssignDebugVisProcessor.h
--- a/Engine/Plugins/Runtime/MassGameplay/Source/MassGameplayDebug/Public/AssignDebugVisProcessor.h
+++ b/Engine/Plugins/Runtime/MassGameplay/Source/MassGameplayDebug/Public/AssignDebugVisProcessor.h
@@ -16,7 +16,7 @@ class MASSGAMEPLAYDEBUG_API UAssignDebugVisProcessor : public UMassFragmentIniti
 public:
 	UAssignDebugVisProcessor();
 	virtual void ConfigureQueries() override;
-	virtual void Execute(UEntitySubsystem& EntitySubsystem, FLWComponentSystemExecutionContext& Context) override;
+	virtual void Execute(UMassEntitySubsystem& EntitySubsystem, FLWComponentSystemExecutionContext& Context) override;
 	virtual void Initialize(UObject& InOwner) override;
 
 protected:
diff --git a/Engine/Plugins/Runtime/MassGameplay/Source/MassGameplayDebug/Public/DebugVisLocationProcessor.h b/Engine/Plugins/Runtime/MassGameplay/Source/MassGameplayDebug/Public/DebugVisLocationProcessor.h
--- a/Engine/Plugins/Runtime/MassGameplay/Source/MassGameplayDebug/Public/DebugVisLocationProcessor.h
+++ b/Engine/Plugins/Runtime/MassGameplay/Source/MassGameplayDebug/Public/DebugVisLocationProcessor.h
@@ -18,7 +18,7 @@ public:
 	
 protected:
 	virtual void ConfigureQueries() override;
-	virtual void Execute(UEntitySubsystem& EntitySubsystem, FLWComponentSystemExecutionContext& Context) override;
+	virtual void Execute(UMassEntitySubsystem& EntitySubsystem, FLWComponentSystemExecutionContext& Context) override;
 	virtual void Initialize(UObject& InOwner) override;
 
 	TWeakObjectPtr<UMassDebugVisualizationComponent> WeakVisualizer;
@@ -37,7 +37,7 @@ class MASSGAMEPLAYDEBUG_API UMassProcessor_UpdateDebugVis : public UPipeProcesso
 public:
 	UMassProcessor_UpdateDebugVis();
 	virtual void ConfigureQueries() override;
-	virtual void Execute(UEntitySubsystem& EntitySubsystem, FLWComponentSystemExecutionContext& Context) override;
+	virtual void Execute(UMassEntitySubsystem& EntitySubsystem, FLWComponentSystemExecutionContext& Context) override;
 
 protected:
 	FLWComponentQuery EntityQuery;
diff --git a/Engine/Plugins/Runtime/MassGameplay/Source/MassLOD/Private/MassLODCollectorProcessor.cpp b/Engine/Plugins/Runtime/MassGameplay/Source/MassLOD/Private/MassLODCollectorProcessor.cpp
--- a/Engine/Plugins/Runtime/MassGameplay/Source/MassLOD/Private/MassLODCollectorProcessor.cpp
+++ b/Engine/Plugins/Runtime/MassGameplay/Source/MassLOD/Private/MassLODCollectorProcessor.cpp
@@ -76,7 +76,7 @@ void UMassLODCollectorProcessor::ConfigureQueries()
 }
 
 template <typename TCollector>
-void CollectLOD(UEntitySubsystem& EntitySubsystem, FLWComponentSystemExecutionContext& Context, const TArray<FViewerInfo>& Viewers, TCollector& Collector, TArrayView<FLWComponentQuery> HighFrequencyTickingEntityQueries, FLWComponentQuery& LowFrequencyTickingEntityQuery)
+void CollectLOD(UMassEntitySubsystem& EntitySubsystem, FLWComponentSystemExecutionContext& Context, const TArray<FViewerInfo>& Viewers, TCollector& Collector, TArrayView<FLWComponentQuery> HighFrequencyTickingEntityQueries, FLWComponentQuery& LowFrequencyTickingEntityQuery)
 {
 	Collector.PrepareExecution(Viewers);
 
@@ -101,7 +101,7 @@ void CollectLOD(UEntitySubsystem& EntitySubsystem, FLWComponentSystemExecutionCo
 	}
 }
 
-void UMassLODCollectorProcessor::Execute(UEntitySubsystem& EntitySubsystem, FLWComponentSystemExecutionContext& Context)
+void UMassLODCollectorProcessor::Execute(UMassEntitySubsystem& EntitySubsystem, FLWComponentSystemExecutionContext& Context)
 {
 	check(LODManager);
 	const TArray<FViewerInfo>& Viewers = LODManager->GetViewers();
diff --git a/Engine/Plugins/Runtime/MassGameplay/Source/MassLOD/Private/MassRepresentationLODCollectorProcessor.cpp b/Engine/Plugins/Runtime/MassGameplay/Source/MassLOD/Private/MassRepresentationLODCollectorProcessor.cpp
--- a/Engine/Plugins/Runtime/MassGameplay/Source/MassLOD/Private/MassRepresentationLODCollectorProcessor.cpp
+++ b/Engine/Plugins/Runtime/MassGameplay/Source/MassLOD/Private/MassRepresentationLODCollectorProcessor.cpp
@@ -42,7 +42,7 @@ void UMassRepresentationLODCollectorProcessor::ConfigureQueries()
 }
 
 template <typename TCollector>
-void CollectLODInfo(UEntitySubsystem& EntitySubsystem, FLWComponentSystemExecutionContext& Context, const TArray<FViewerInfo>& Viewers, TCollector& Collector, FLWComponentQuery& HighFrequencyTickingEntityQuery, FLWComponentQuery& LowFrequencyTickingEntityQuery)
+void CollectLODInfo(UMassEntitySubsystem& EntitySubsystem, FLWComponentSystemExecutionContext& Context, const TArray<FViewerInfo>& Viewers, TCollector& Collector, FLWComponentQuery& HighFrequencyTickingEntityQuery, FLWComponentQuery& LowFrequencyTickingEntityQuery)
 {
 	Collector.PrepareExecution(Viewers);
 
@@ -64,7 +64,7 @@ void CollectLODInfo(UEntitySubsystem& EntitySubsystem, FLWComponentSystemExecuti
 	}
 }
 
-void UMassRepresentationLODCollectorProcessor::Execute(UEntitySubsystem& EntitySubsystem, FLWComponentSystemExecutionContext& Context)
+void UMassRepresentationLODCollectorProcessor::Execute(UMassEntitySubsystem& EntitySubsystem, FLWComponentSystemExecutionContext& Context)
 {
 	check(LODManager);
 	const TArray<FViewerInfo>& Viewers = LODManager->GetViewers();
diff --git a/Engine/Plugins/Runtime/MassGameplay/Source/MassLOD/Private/MassSimulationLOD.cpp b/Engine/Plugins/Runtime/MassGameplay/Source/MassLOD/Private/MassSimulationLOD.cpp
--- a/Engine/Plugins/Runtime/MassGameplay/Source/MassLOD/Private/MassSimulationLOD.cpp
+++ b/Engine/Plugins/Runtime/MassGameplay/Source/MassLOD/Private/MassSimulationLOD.cpp
@@ -28,7 +28,7 @@ void UMassProcessor_MassSimulationLODViewersInfo::ConfigureQueries()
 	EntityQuery.AddRequirement<FDataFragment_MassSimulationLODInfo>(ELWComponentAccess::ReadWrite);
 }
 
-void UMassProcessor_MassSimulationLODViewersInfo::Execute(UEntitySubsystem& EntitySubsystem, FLWComponentSystemExecutionContext& Context)
+void UMassProcessor_MassSimulationLODViewersInfo::Execute(UMassEntitySubsystem& EntitySubsystem, FLWComponentSystemExecutionContext& Context)
 {
 	check(LODManager);
 	const TArray<FViewerInfo>& Viewers = LODManager->GetViewers();
@@ -106,7 +106,7 @@ void UMassSimulationLODProcessor::Initialize(UObject& InOwner)
 	Super::Initialize(InOwner);
 }
 
-void UMassSimulationLODProcessor::Execute(UEntitySubsystem& EntitySubsystem, FLWComponentSystemExecutionContext& Context)
+void UMassSimulationLODProcessor::Execute(UMassEntitySubsystem& EntitySubsystem, FLWComponentSystemExecutionContext& Context)
 {
 	TRACE_CPUPROFILER_EVENT_SCOPE(TEXT("SimulationLOD"))
 
@@ -116,7 +116,7 @@ void UMassSimulationLODProcessor::Execute(UEntitySubsystem& EntitySubsystem, FLW
 	}
 }
 
-void UMassSimulationLODProcessor::CalculateLODForConfig(UEntitySubsystem& EntitySubsystem, FLWComponentSystemExecutionContext& Context, FMassSimulationLODConfig& LODConfig)
+void UMassSimulationLODProcessor::CalculateLODForConfig(UMassEntitySubsystem& EntitySubsystem, FLWComponentSystemExecutionContext& Context, FMassSimulationLODConfig& LODConfig)
 {
 	check(LODManager);
 	const TArray<FViewerInfo>& Viewers = LODManager->GetViewers();
diff --git a/Engine/Plugins/Runtime/MassGameplay/Source/MassLOD/Public/MassLODCollectorProcessor.h b/Engine/Plugins/Runtime/MassGameplay/Source/MassLOD/Public/MassLODCollectorProcessor.h
--- a/Engine/Plugins/Runtime/MassGameplay/Source/MassLOD/Public/MassLODCollectorProcessor.h
+++ b/Engine/Plugins/Runtime/MassGameplay/Source/MassLOD/Public/MassLODCollectorProcessor.h
@@ -53,7 +53,7 @@ public:
 protected:
 	virtual void Initialize(UObject& Owner) override;
 	virtual void ConfigureQueries() override;
-	virtual void Execute(UEntitySubsystem& EntitySubsystem, FLWComponentSystemExecutionContext& Context) override;
+	virtual void Execute(UMassEntitySubsystem& EntitySubsystem, FLWComponentSystemExecutionContext& Context) override;
 
 	UPROPERTY(EditAnywhere, Category = "Mass|LOD", config)
 	TArray<FMassCollectorLODConfig> LODConfigs;
diff --git a/Engine/Plugins/Runtime/MassGameplay/Source/MassLOD/Public/MassLODTypes.h b/Engine/Plugins/Runtime/MassGameplay/Source/MassLOD/Public/MassLODTypes.h
--- a/Engine/Plugins/Runtime/MassGameplay/Source/MassLOD/Public/MassLODTypes.h
+++ b/Engine/Plugins/Runtime/MassGameplay/Source/MassLOD/Public/MassLODTypes.h
@@ -4,7 +4,7 @@
 
 #include "LWComponentTypes.h"
 #include "Containers/StaticArray.h"
-#include "EntitySubsystem.h"
+#include "MassEntitySubsystem.h"
 
 #include "MassLODTypes.generated.h"
 
diff --git a/Engine/Plugins/Runtime/MassGameplay/Source/MassLOD/Public/MassLODUtils.h b/Engine/Plugins/Runtime/MassGameplay/Source/MassLOD/Public/MassLODUtils.h
--- a/Engine/Plugins/Runtime/MassGameplay/Source/MassLOD/Public/MassLODUtils.h
+++ b/Engine/Plugins/Runtime/MassGameplay/Source/MassLOD/Public/MassLODUtils.h
@@ -3,7 +3,7 @@
 #pragma once
 
 #include "LWComponentTypes.h"
-#include "EntitySubsystem.h"
+#include "MassEntitySubsystem.h"
 #include "MassCommonTypes.h"
 #include "MassLODTypes.h"
 #include "Logging/LogMacros.h"
diff --git a/Engine/Plugins/Runtime/MassGameplay/Source/MassLOD/Public/MassRepresentationLODCollectorProcessor.h b/Engine/Plugins/Runtime/MassGameplay/Source/MassLOD/Public/MassRepresentationLODCollectorProcessor.h
--- a/Engine/Plugins/Runtime/MassGameplay/Source/MassLOD/Public/MassRepresentationLODCollectorProcessor.h
+++ b/Engine/Plugins/Runtime/MassGameplay/Source/MassLOD/Public/MassRepresentationLODCollectorProcessor.h
@@ -45,7 +45,7 @@ class MASSLOD_API UMassRepresentationLODCollectorProcessor : public UMassProcess
 	protected:
 	virtual void Initialize(UObject& Owner) override;
 	virtual void ConfigureQueries() override;
-	virtual void Execute(UEntitySubsystem& EntitySubsystem, FLWComponentSystemExecutionContext& Context) override;
+	virtual void Execute(UMassEntitySubsystem& EntitySubsystem, FLWComponentSystemExecutionContext& Context) override;
 
 	UPROPERTY(EditAnywhere, Category = "Mass|LOD", config)
 	TArray<FMassRepresentationLODCollectorConfig> LODConfigs;
diff --git a/Engine/Plugins/Runtime/MassGameplay/Source/MassLOD/Public/MassSimulationLOD.h b/Engine/Plugins/Runtime/MassGameplay/Source/MassLOD/Public/MassSimulationLOD.h
--- a/Engine/Plugins/Runtime/MassGameplay/Source/MassLOD/Public/MassSimulationLOD.h
+++ b/Engine/Plugins/Runtime/MassGameplay/Source/MassLOD/Public/MassSimulationLOD.h
@@ -43,7 +43,7 @@ public:
 protected:
 	virtual void Initialize(UObject& Owner) override;
 	virtual void ConfigureQueries() override;
-	virtual void Execute(UEntitySubsystem& EntitySubsystem, FLWComponentSystemExecutionContext& Context) override;
+	virtual void Execute(UMassEntitySubsystem& EntitySubsystem, FLWComponentSystemExecutionContext& Context) override;
 
 	TMassLODCollector<FSimulationLODLogic> LODCollector;
 	FLWComponentQuery EntityQuery;
@@ -132,9 +132,9 @@ protected:
 
 	virtual void ConfigureQueries() override;
 	virtual void Initialize(UObject& InOwner) override;
-	virtual void Execute(UEntitySubsystem& EntitySubsystem, FLWComponentSystemExecutionContext& Context) override;
+	virtual void Execute(UMassEntitySubsystem& EntitySubsystem, FLWComponentSystemExecutionContext& Context) override;
 
-	void CalculateLODForConfig(UEntitySubsystem& EntitySubsystem, FLWComponentSystemExecutionContext& Context, FMassSimulationLODConfig& LODConfig);
+	void CalculateLODForConfig(UMassEntitySubsystem& EntitySubsystem, FLWComponentSystemExecutionContext& Context, FMassSimulationLODConfig& LODConfig);
 
 	UPROPERTY(EditAnywhere, Category = "Mass|LOD", config)
 	TArray<FMassSimulationLODConfig> LODConfigs;
diff --git a/Engine/Plugins/Runtime/MassGameplay/Source/MassReplication/Private/MassClientBubbleSerializerBase.cpp b/Engine/Plugins/Runtime/MassGameplay/Source/MassReplication/Private/MassClientBubbleSerializerBase.cpp
--- a/Engine/Plugins/Runtime/MassGameplay/Source/MassReplication/Private/MassClientBubbleSerializerBase.cpp
+++ b/Engine/Plugins/Runtime/MassGameplay/Source/MassReplication/Private/MassClientBubbleSerializerBase.cpp
@@ -12,7 +12,7 @@ void FMassClientBubbleSerializerBase::InitializeForWorld(UWorld& InWorld)
 
 	ReplicationManager = InWorld.GetSubsystem<UMassReplicationManager>();
 
-	EntitySystem = UPipeEntitySubsystem::GetCurrent(World);
+	EntitySystem = UMassEntitySubsystem::GetCurrent(World);
 	check(EntitySystem);
 
 	checkf(ClientHandler, TEXT("ClientHandler must be setup! See SetClientHandler()"));
diff --git a/Engine/Plugins/Runtime/MassGameplay/Source/MassReplication/Private/MassReplicationManager.cpp b/Engine/Plugins/Runtime/MassGameplay/Source/MassReplication/Private/MassReplicationManager.cpp
--- a/Engine/Plugins/Runtime/MassGameplay/Source/MassReplication/Private/MassReplicationManager.cpp
+++ b/Engine/Plugins/Runtime/MassGameplay/Source/MassReplication/Private/MassReplicationManager.cpp
@@ -5,7 +5,7 @@
 #include "Engine/ChildConnection.h"
 #include "GameFramework/GameModeBase.h"
 #include "MassCommonTypes.h"
-#include "MassEntitySystem.h"
+#include "MassEntitySubsystem.h"
 #include "MassClientBubbleHandler.h"
 #include "MassClientBubbleInfoBase.h"
 
@@ -24,7 +24,7 @@ void UMassReplicationManager::Initialize(FSubsystemCollectionBase& Collection)
 
 	check(MassLODManager);
 
-	EntitySystem = UPipeEntitySubsystem::GetCurrent(GetWorld());
+	EntitySystem = UMassEntitySubsystem::GetCurrent(GetWorld());
 }
 
 void UMassReplicationManager::Deinitialize()
diff --git a/Engine/Plugins/Runtime/MassGameplay/Source/MassReplication/Private/MassReplicationProcessorBase.cpp b/Engine/Plugins/Runtime/MassGameplay/Source/MassReplication/Private/MassReplicationProcessorBase.cpp
--- a/Engine/Plugins/Runtime/MassGameplay/Source/MassReplication/Private/MassReplicationProcessorBase.cpp
+++ b/Engine/Plugins/Runtime/MassGameplay/Source/MassReplication/Private/MassReplicationProcessorBase.cpp
@@ -129,7 +129,7 @@ void UMassReplicationProcessorBase::PrepareExecution()
 #endif //UE_REPLICATION_COMPILE_SERVER_CODE
 }
 
-void UMassReplicationProcessorBase::Execute(UEntitySubsystem& EntitySubsystem, FLWComponentSystemExecutionContext& Context)
+void UMassReplicationProcessorBase::Execute(UMassEntitySubsystem& EntitySubsystem, FLWComponentSystemExecutionContext& Context)
 {
 #if UE_REPLICATION_COMPILE_SERVER_CODE
 	check(World);
diff --git a/Engine/Plugins/Runtime/MassGameplay/Source/MassReplication/Private/MassReplicationTypes.cpp b/Engine/Plugins/Runtime/MassGameplay/Source/MassReplication/Private/MassReplicationTypes.cpp
--- a/Engine/Plugins/Runtime/MassGameplay/Source/MassReplication/Private/MassReplicationTypes.cpp
+++ b/Engine/Plugins/Runtime/MassGameplay/Source/MassReplication/Private/MassReplicationTypes.cpp
@@ -29,7 +29,7 @@ void UMassNetworkIDFragmentInitializer::Initialize(UObject& Owner)
 #endif //UE_REPLICATION_COMPILE_SERVER_CODE
 }
 
-void UMassNetworkIDFragmentInitializer::Execute(UEntitySubsystem& EntitySubsystem, FLWComponentSystemExecutionContext& Context)
+void UMassNetworkIDFragmentInitializer::Execute(UMassEntitySubsystem& EntitySubsystem, FLWComponentSystemExecutionContext& Context)
 {
 	QUICK_SCOPE_CYCLE_COUNTER(MassProcessor_InitNetworkID_Run);
 
diff --git a/Engine/Plugins/Runtime/MassGameplay/Source/MassReplication/Public/MassClientBubbleHandler.h b/Engine/Plugins/Runtime/MassGameplay/Source/MassReplication/Public/MassClientBubbleHandler.h
--- a/Engine/Plugins/Runtime/MassGameplay/Source/MassReplication/Public/MassClientBubbleHandler.h
+++ b/Engine/Plugins/Runtime/MassGameplay/Source/MassReplication/Public/MassClientBubbleHandler.h
@@ -4,7 +4,7 @@
 
 #include "MassCommonTypes.h"
 #include "MassReplicationTypes.h"
-#include "MassEntitySystem.h"
+#include "MassEntitySubsystem.h"
 #include "Containers/ArrayView.h"
 #include "MassClientBubbleSerializerBase.h"
 #include "MassSpawnerSubsystem.h"
@@ -17,7 +17,7 @@
 #include "MassClientBubbleHandler.generated.h"
 
 class UWorld;
-class UPipeEntitySubsystem;
+class UMassEntitySubsystem;
 
 namespace UE::Mass::Replication
 {
@@ -427,7 +427,7 @@ void TClientBubbleHandlerBase<AgentArrayItem>::PostReplicatedAddHelper(const TAr
 {
 	TArray<FLWEntity> EntitiesDestroy;
 
-	UPipeEntitySubsystem* EntitySystem = Serializer->GetEntitySystem();
+	UMassEntitySubsystem* EntitySystem = Serializer->GetEntitySystem();
 	check(EntitySystem);
 
 	UMassReplicationManager* ReplicationManager = Serializer->GetReplicationManager();
@@ -506,7 +506,7 @@ void TClientBubbleHandlerBase<AgentArrayItem>::PostReplicatedAddEntitiesHelper(c
 {
 	check(Serializer);
 
-	UPipeEntitySubsystem* EntitySystem = Serializer->GetEntitySystem();
+	UMassEntitySubsystem* EntitySystem = Serializer->GetEntitySystem();
 	check(EntitySystem);
 
 	UMassSpawnerSubsystem* SpawnerSubsystem = Serializer->GetSpawnerSubsystem();
@@ -584,7 +584,7 @@ void TClientBubbleHandlerBase<AgentArrayItem>::PostReplicatedAddEntitiesHelper(c
 template<typename AgentArrayItem>
 void TClientBubbleHandlerBase<AgentArrayItem>::PostReplicatedChangeHelper(const TArrayView<int32> ChangedIndices, FSetModifiedEntityDataFunction SetModifiedEntityData)
 {
-	UPipeEntitySubsystem* EntitySystem = Serializer->GetEntitySystem();
+	UMassEntitySubsystem* EntitySystem = Serializer->GetEntitySystem();
 	check(EntitySystem);
 
 	UMassReplicationManager* ReplicationManager = Serializer->GetReplicationManager();
@@ -652,7 +652,7 @@ template<typename AgentArrayItem>
 void TClientBubbleHandlerBase<AgentArrayItem>::DebugValidateBubbleOnClient()
 {
 #if UE_ALLOW_DEBUG_REPLICATION
-	UPipeEntitySubsystem* EntitySystem = Serializer->GetEntitySystem();
+	UMassEntitySubsystem* EntitySystem = Serializer->GetEntitySystem();
 	check(EntitySystem);
 
 	UMassReplicationManager* ReplicationManager = Serializer->GetReplicationManager();
@@ -715,7 +715,7 @@ void TClientBubbleHandlerBase<AgentArrayItem>::DebugValidateBubbleOnServer()
 #if UE_ALLOW_DEBUG_REPLICATION
 	using namespace UE::Mass::Replication;
 
-	UPipeEntitySubsystem* EntitySystem = Serializer->GetEntitySystem();
+	UMassEntitySubsystem* EntitySystem = Serializer->GetEntitySystem();
 	check(EntitySystem);
 
 	for (int32 OuterIdx = 0; OuterIdx < (*Agents).Num(); ++OuterIdx)
diff --git a/Engine/Plugins/Runtime/MassGameplay/Source/MassReplication/Public/MassClientBubbleSerializerBase.h b/Engine/Plugins/Runtime/MassGameplay/Source/MassReplication/Public/MassClientBubbleSerializerBase.h
--- a/Engine/Plugins/Runtime/MassGameplay/Source/MassReplication/Public/MassClientBubbleSerializerBase.h
+++ b/Engine/Plugins/Runtime/MassGameplay/Source/MassReplication/Public/MassClientBubbleSerializerBase.h
@@ -10,7 +10,7 @@
 
 class UWorld;
 class UMassSpawnerSubsystem;
-class UPipeEntitySubsystem;
+class UMassEntitySubsystem;
 class IClientBubbleHandlerInterface;
 
 /**
@@ -35,7 +35,7 @@ public:
 	UWorld* GetWorld() const { return World; }
 	UMassSpawnerSubsystem* GetSpawnerSubsystem() const { return SpawnerSubsystem; }
 	UMassReplicationManager* GetReplicationManager() const { return ReplicationManager; }
-	UPipeEntitySubsystem* GetEntitySystem() const { return EntitySystem; }
+	UMassEntitySubsystem* GetEntitySystem() const { return EntitySystem; }
 
 	void SetClientHandler(IClientBubbleHandlerInterface& InClientHandler) { ClientHandler = &InClientHandler; }
 	IClientBubbleHandlerInterface* GetClientHandler() const { return ClientHandler; }
@@ -48,7 +48,7 @@ private:
 	UMassSpawnerSubsystem* SpawnerSubsystem = nullptr;
 
 	UPROPERTY(Transient)
-	UPipeEntitySubsystem* EntitySystem = nullptr;
+	UMassEntitySubsystem* EntitySystem = nullptr;
 
 	UPROPERTY(Transient)
 	UMassReplicationManager* ReplicationManager = nullptr;
diff --git a/Engine/Plugins/Runtime/MassGameplay/Source/MassReplication/Public/MassReplicationManager.h b/Engine/Plugins/Runtime/MassGameplay/Source/MassReplication/Public/MassReplicationManager.h
--- a/Engine/Plugins/Runtime/MassGameplay/Source/MassReplication/Public/MassReplicationManager.h
+++ b/Engine/Plugins/Runtime/MassGameplay/Source/MassReplication/Public/MassReplicationManager.h
@@ -9,7 +9,7 @@
 
 #include "MassReplicationManager.generated.h"
 
-class UPipeEntitySubsystem;
+class UMassEntitySubsystem;
 class AMassCrowdClientBubbleInfo;
 class UMassLODManager;
 class AMassClientBubbleInfoBase;
@@ -245,7 +245,7 @@ protected:
 	static uint32 CurrentNetMassCounter;
 
 	UPROPERTY()
-	UPipeEntitySubsystem* EntitySystem;
+	UMassEntitySubsystem* EntitySystem;
 
 	/** Clients free list FMassClientHandle manager, handles will to the indices of FMassClientReplicationData::ClientBubbles */
 	FMassClientHandleManager ClientHandleManager;
diff --git a/Engine/Plugins/Runtime/MassGameplay/Source/MassReplication/Public/MassReplicationProcessorBase.h b/Engine/Plugins/Runtime/MassGameplay/Source/MassReplication/Public/MassReplicationProcessorBase.h
--- a/Engine/Plugins/Runtime/MassGameplay/Source/MassReplication/Public/MassReplicationProcessorBase.h
+++ b/Engine/Plugins/Runtime/MassGameplay/Source/MassReplication/Public/MassReplicationProcessorBase.h
@@ -32,7 +32,7 @@ public:
 protected:
 	virtual void ConfigureQueries() override;
 	virtual void Initialize(UObject& Owner) override;
-	virtual void Execute(UEntitySubsystem& EntitySubsystem, FLWComponentSystemExecutionContext& Context) override;
+	virtual void Execute(UMassEntitySubsystem& EntitySubsystem, FLWComponentSystemExecutionContext& Context) override;
 
 	void PrepareExecution();
 
diff --git a/Engine/Plugins/Runtime/MassGameplay/Source/MassReplication/Public/MassReplicationTypes.h b/Engine/Plugins/Runtime/MassGameplay/Source/MassReplication/Public/MassReplicationTypes.h
--- a/Engine/Plugins/Runtime/MassGameplay/Source/MassReplication/Public/MassReplicationTypes.h
+++ b/Engine/Plugins/Runtime/MassGameplay/Source/MassReplication/Public/MassReplicationTypes.h
@@ -221,7 +221,7 @@ public:
 protected:
 	virtual void ConfigureQueries() override;
 	virtual void Initialize(UObject& Owner) override;
-	virtual void Execute(UEntitySubsystem& EntitySubsystem, FLWComponentSystemExecutionContext& Context) override;
+	virtual void Execute(UMassEntitySubsystem& EntitySubsystem, FLWComponentSystemExecutionContext& Context) override;
 
 	UPROPERTY()
 	UMassReplicationManager* ReplicationManager = nullptr;
diff --git a/Engine/Plugins/Runtime/MassGameplay/Source/MassRepresentation/Private/MassRepresentationProcessor.cpp b/Engine/Plugins/Runtime/MassGameplay/Source/MassRepresentation/Private/MassRepresentationProcessor.cpp
--- a/Engine/Plugins/Runtime/MassGameplay/Source/MassRepresentation/Private/MassRepresentationProcessor.cpp
+++ b/Engine/Plugins/Runtime/MassGameplay/Source/MassRepresentation/Private/MassRepresentationProcessor.cpp
@@ -6,7 +6,7 @@
 #include "MassActorSubsystem.h"
 #include "MassCommonFragments.h"
 #include "LWCCommandBuffer.h"
-#include "MassEntitySystem.h"
+#include "MassEntitySubsystem.h"
 #include "EntityView.h"
 #include "Engine/World.h"
 
@@ -39,7 +39,7 @@ void UMassRepresentationProcessor::Initialize(UObject& Owner)
 	Super::Initialize(Owner);
 
 	World = Owner.GetWorld();
-	PipeEntitySubsystem = UPipeEntitySubsystem::GetCurrent(World);
+	PipeEntitySubsystem = UMassEntitySubsystem::GetCurrent(World);
 	RepresentationSubsystem = UWorld::GetSubsystem<UMassRepresentationSubsystem>(World);
 
 	// Calculate the default representation when actor isn't spawned yet.
@@ -207,7 +207,7 @@ void UMassRepresentationProcessor::UpdateRepresentation(FLWComponentSystemExecut
 	}
 }
 
-void UMassRepresentationProcessor::Execute(UEntitySubsystem& EntitySubsystem, FLWComponentSystemExecutionContext& Context)
+void UMassRepresentationProcessor::Execute(UMassEntitySubsystem& EntitySubsystem, FLWComponentSystemExecutionContext& Context)
 {
 	check(RepresentationSubsystem);
 
@@ -247,7 +247,7 @@ bool UMassRepresentationProcessor::ReleaseActorOrCancelSpawning(const FMassHandl
 			TObjectKey<const AActor> ActorKey(Actor); 
 			if (UMassActorSubsystem* MassActorSubsystem = UWorld::GetSubsystem<UMassActorSubsystem>(World))
 			{
-				Context.Defer().EmplaceCommand<FDeferredCommand>([MassActorSubsystem, ActorKey](UEntitySubsystem&)
+				Context.Defer().EmplaceCommand<FDeferredCommand>([MassActorSubsystem, ActorKey](UMassEntitySubsystem&)
 				{
 					MassActorSubsystem->RemoveHandleForActor(ActorKey);
 				});
@@ -268,7 +268,7 @@ void UMassRepresentationProcessor::SetActorEnabled(const EActorEnabledType Enabl
 	if (Actor.GetActorEnableCollision() != bEnabled)
 	{
 		// Deferring this as there is a callback internally that could end up doing things outside of the game thread and will fire checks(Chaos mostly)
-		Context.Defer().EmplaceCommand<FDeferredCommand>( [&Actor,bEnabled](UEntitySubsystem&)
+		Context.Defer().EmplaceCommand<FDeferredCommand>( [&Actor,bEnabled](UMassEntitySubsystem&)
 		{
 			Actor.SetActorEnableCollision(bEnabled);
 		});
@@ -279,7 +279,7 @@ void UMassRepresentationProcessor::TeleportActor(const FTransform& Transform, AA
 {
 	if (!Actor.GetTransform().Equals(Transform))
 	{
-		Context.Defer().EmplaceCommand<FDeferredCommand>([&Actor, Transform](UEntitySubsystem&)
+		Context.Defer().EmplaceCommand<FDeferredCommand>([&Actor, Transform](UMassEntitySubsystem&)
 		{
 			Actor.SetActorTransform(Transform, /*bSweep*/false, /*OutSweepHitResult*/nullptr, ETeleportType::TeleportPhysics);
 		});
@@ -288,7 +288,7 @@ void UMassRepresentationProcessor::TeleportActor(const FTransform& Transform, AA
 
 void UMassRepresentationProcessor::ReleaseAnyActorOrCancelAnySpawning(UMassRepresentationSubsystem& RepresentationSubsystem, const FMassHandle MassAgent)
 {
-	if (const UPipeEntitySubsystem* EntitySubsystem = UPipeEntitySubsystem::GetCurrent(RepresentationSubsystem.GetWorld()))
+	if (const UMassEntitySubsystem* EntitySubsystem = UMassEntitySubsystem::GetCurrent(RepresentationSubsystem.GetWorld()))
 	{
 		if (FDataFragment_Actor* ActorInfo = EntitySubsystem->GetComponentDataPtr<FDataFragment_Actor>(MassAgent.GetLWEntity()))
 		{
@@ -469,7 +469,7 @@ void UMassRepresentationFragmentDestructor::ConfigureQueries()
 	EntityQuery.AddRequirement<FDataFragment_Actor>(ELWComponentAccess::ReadWrite);
 }
 
-void UMassRepresentationFragmentDestructor::Execute(UEntitySubsystem& EntitySubsystem, FLWComponentSystemExecutionContext& Context)
+void UMassRepresentationFragmentDestructor::Execute(UMassEntitySubsystem& EntitySubsystem, FLWComponentSystemExecutionContext& Context)
 {
 	check(RepresentationSubsystem);
 
diff --git a/Engine/Plugins/Runtime/MassGameplay/Source/MassRepresentation/Private/MassRepresentationSubsystem.cpp b/Engine/Plugins/Runtime/MassGameplay/Source/MassRepresentation/Private/MassRepresentationSubsystem.cpp
--- a/Engine/Plugins/Runtime/MassGameplay/Source/MassRepresentation/Private/MassRepresentationSubsystem.cpp
+++ b/Engine/Plugins/Runtime/MassGameplay/Source/MassRepresentation/Private/MassRepresentationSubsystem.cpp
@@ -11,7 +11,7 @@
 #include "MassSimulationSettings.h"
 #include "MassAgentComponent.h"
 #include "MassAgentSubsystem.h"
-#include "MassEntitySystem.h"
+#include "MassEntitySubsystem.h"
 #include "MassRepresentationFragments.h"
 #include "MassRepresentationProcessor.h"
 #include "EntityView.h"
@@ -302,7 +302,7 @@ void UMassRepresentationSubsystem::Initialize(FSubsystemCollectionBase& Collecti
 
 	if (UWorld* World = GetWorld())
 	{
-		EntitySubsystem = UPipeEntitySubsystem::GetCurrent(World);
+		EntitySubsystem = UMassEntitySubsystem::GetCurrent(World);
 		ActorSpawnerSubsystem = World->GetSubsystem<UMassActorSpawnerSubsystem>();
 		WorldPartitionSubsystem = World->GetSubsystem<UWorldPartitionSubsystem>();
 
diff --git a/Engine/Plugins/Runtime/MassGameplay/Source/MassRepresentation/Private/MassUpdateISMProcessor.cpp b/Engine/Plugins/Runtime/MassGameplay/Source/MassRepresentation/Private/MassUpdateISMProcessor.cpp
--- a/Engine/Plugins/Runtime/MassGameplay/Source/MassRepresentation/Private/MassUpdateISMProcessor.cpp
+++ b/Engine/Plugins/Runtime/MassGameplay/Source/MassRepresentation/Private/MassUpdateISMProcessor.cpp
@@ -3,7 +3,7 @@
 #include "MassUpdateISMProcessor.h"
 #include "MassVisualizationComponent.h"
 #include "MassRepresentationSubsystem.h"
-#include "MassEntitySystem.h"
+#include "MassEntitySubsystem.h"
 #include "MassRepresentationFragments.h"
 #include "MassCommonFragments.h"
 #include "MassLODTypes.h"
@@ -40,7 +40,7 @@ void UMassUpdateISMProcessor::ConfigureQueries()
 	}
 }
 
-void UMassUpdateISMProcessor::Execute(UEntitySubsystem& EntitySubsystem, FLWComponentSystemExecutionContext& Context)
+void UMassUpdateISMProcessor::Execute(UMassEntitySubsystem& EntitySubsystem, FLWComponentSystemExecutionContext& Context)
 {
 	check(RepresentationSubsystem);
 	FMassInstancedStaticMeshInfoArrayView ISMInfo = RepresentationSubsystem->GetMutableInstancedStaticMeshInfos();
diff --git a/Engine/Plugins/Runtime/MassGameplay/Source/MassRepresentation/Private/MassVisualizationLODProcessor.cpp b/Engine/Plugins/Runtime/MassGameplay/Source/MassRepresentation/Private/MassVisualizationLODProcessor.cpp
--- a/Engine/Plugins/Runtime/MassGameplay/Source/MassRepresentation/Private/MassVisualizationLODProcessor.cpp
+++ b/Engine/Plugins/Runtime/MassGameplay/Source/MassRepresentation/Private/MassVisualizationLODProcessor.cpp
@@ -61,7 +61,7 @@ void UMassVisualizationLODProcessor::PrepareExecution()
 	LODCalculator.PrepareExecution(Viewers);
 }
 
-void UMassVisualizationLODProcessor::Execute(UEntitySubsystem& EntitySubsystem, FLWComponentSystemExecutionContext& Context)
+void UMassVisualizationLODProcessor::Execute(UMassEntitySubsystem& EntitySubsystem, FLWComponentSystemExecutionContext& Context)
 {
 	ExecuteInternal(EntitySubsystem, Context);
 }
\ No newline at end of file
diff --git a/Engine/Plugins/Runtime/MassGameplay/Source/MassRepresentation/Public/MassRepresentationFragments.h b/Engine/Plugins/Runtime/MassGameplay/Source/MassRepresentation/Public/MassRepresentationFragments.h
--- a/Engine/Plugins/Runtime/MassGameplay/Source/MassRepresentation/Public/MassRepresentationFragments.h
+++ b/Engine/Plugins/Runtime/MassGameplay/Source/MassRepresentation/Public/MassRepresentationFragments.h
@@ -4,7 +4,7 @@
 
 #include "MassLODTypes.h"
 #include "LWComponentTypes.h"
-#include "EntitySubsystem.h"
+#include "MassEntitySubsystem.h"
 #include "MassActorSpawnerSubsystem.h"
 #include "MassRepresentationTypes.h"
 
diff --git a/Engine/Plugins/Runtime/MassGameplay/Source/MassRepresentation/Public/MassRepresentationProcessor.h b/Engine/Plugins/Runtime/MassGameplay/Source/MassRepresentation/Public/MassRepresentationProcessor.h
--- a/Engine/Plugins/Runtime/MassGameplay/Source/MassRepresentation/Public/MassRepresentationProcessor.h
+++ b/Engine/Plugins/Runtime/MassGameplay/Source/MassRepresentation/Public/MassRepresentationProcessor.h
@@ -42,7 +42,7 @@ protected:
 	 * @param EntitySubsystem is the system to execute the lambdas on each entity chunk
 	 * @param Context is the execution context to be passed when executing the lambdas
 	 */
-	virtual void Execute(UEntitySubsystem& EntitySubsystem, FLWComponentSystemExecutionContext& Context) override;
+	virtual void Execute(UMassEntitySubsystem& EntitySubsystem, FLWComponentSystemExecutionContext& Context) override;
 
 	/** 
 	 * Returns an actor of the template type and setup fragments values from it
@@ -183,7 +183,7 @@ protected:
 	UWorld* World;
 
 	UPROPERTY(Transient)
-	UPipeEntitySubsystem* PipeEntitySubsystem;
+	UMassEntitySubsystem* PipeEntitySubsystem;
 
 	/** A cache pointer to the representation subsystem */
 	UPROPERTY(Transient)
@@ -210,7 +210,7 @@ public:
 protected:
 	virtual void Initialize(UObject& Owner) override;
 	virtual void ConfigureQueries() override;
-	virtual void Execute(UEntitySubsystem& EntitySubsystem, FLWComponentSystemExecutionContext& Context) override;
+	virtual void Execute(UMassEntitySubsystem& EntitySubsystem, FLWComponentSystemExecutionContext& Context) override;
 
 	/** A cache pointer to the representation subsystem */
 	UPROPERTY(Transient)
diff --git a/Engine/Plugins/Runtime/MassGameplay/Source/MassRepresentation/Public/MassRepresentationSubsystem.h b/Engine/Plugins/Runtime/MassGameplay/Source/MassRepresentation/Public/MassRepresentationSubsystem.h
--- a/Engine/Plugins/Runtime/MassGameplay/Source/MassRepresentation/Public/MassRepresentationSubsystem.h
+++ b/Engine/Plugins/Runtime/MassGameplay/Source/MassRepresentation/Public/MassRepresentationSubsystem.h
@@ -18,7 +18,7 @@ struct FMassInstancedStaticMeshInfo;
 struct FMassHandle_ActorSpawnRequest;
 class UMassActorSpawnerSubsystem;
 class UMassAgentComponent;
-class UPipeEntitySubsystem;
+class UMassEntitySubsystem;
 enum class EPipeProcessingPhase : uint8;
 class UWorldPartitionSubsystem;
 
@@ -150,7 +150,7 @@ protected:
 	UMassActorSpawnerSubsystem* ActorSpawnerSubsystem;
 
 	UPROPERTY(Transient)
-	UPipeEntitySubsystem* EntitySubsystem;
+	UMassEntitySubsystem* EntitySubsystem;
 
 	UPROPERTY(Transient)
 	UWorldPartitionSubsystem* WorldPartitionSubsystem;
diff --git a/Engine/Plugins/Runtime/MassGameplay/Source/MassRepresentation/Public/MassRepresentationUtils.h b/Engine/Plugins/Runtime/MassGameplay/Source/MassRepresentation/Public/MassRepresentationUtils.h
--- a/Engine/Plugins/Runtime/MassGameplay/Source/MassRepresentation/Public/MassRepresentationUtils.h
+++ b/Engine/Plugins/Runtime/MassGameplay/Source/MassRepresentation/Public/MassRepresentationUtils.h
@@ -3,7 +3,7 @@
 #pragma once
 
 #include "LWComponentTypes.h"
-#include "EntitySubsystem.h"
+#include "MassEntitySubsystem.h"
 #include "MassCommonTypes.h"
 #include "MassRepresentationTypes.h"
 
diff --git a/Engine/Plugins/Runtime/MassGameplay/Source/MassRepresentation/Public/MassUpdateISMProcessor.h b/Engine/Plugins/Runtime/MassGameplay/Source/MassRepresentation/Public/MassUpdateISMProcessor.h
--- a/Engine/Plugins/Runtime/MassGameplay/Source/MassRepresentation/Public/MassUpdateISMProcessor.h
+++ b/Engine/Plugins/Runtime/MassGameplay/Source/MassRepresentation/Public/MassUpdateISMProcessor.h
@@ -33,7 +33,7 @@ protected:
 	 * Execution method for this processor
 	 * @param EntitySubsystem is the system to execute the lambdas on each entity chunk
 	 * @param Context is the execution context to be passed when executing the lambdas */
-	virtual void Execute(UEntitySubsystem& EntitySubsystem, FLWComponentSystemExecutionContext& Context) override;
+	virtual void Execute(UMassEntitySubsystem& EntitySubsystem, FLWComponentSystemExecutionContext& Context) override;
 
 	/** Component Tag that will be used to filter out entities to update ISM */
 	UPROPERTY(EditAnywhere, Category = "Mass|LOD", config, meta = (BaseStruct = "ComponentTag"))
diff --git a/Engine/Plugins/Runtime/MassGameplay/Source/MassRepresentation/Public/MassVisualizationLODProcessor.h b/Engine/Plugins/Runtime/MassGameplay/Source/MassRepresentation/Public/MassVisualizationLODProcessor.h
--- a/Engine/Plugins/Runtime/MassGameplay/Source/MassRepresentation/Public/MassVisualizationLODProcessor.h
+++ b/Engine/Plugins/Runtime/MassGameplay/Source/MassRepresentation/Public/MassVisualizationLODProcessor.h
@@ -32,12 +32,12 @@ protected:
 	 * @param EntitySubsystem is the system to execute the lambdas on each entity chunk
 	 * @param Context is the execution context to be passed when executing the lambdas
 	 */
-	virtual void Execute(UEntitySubsystem& EntitySubsystem, FLWComponentSystemExecutionContext& Context) override;
+	virtual void Execute(UMassEntitySubsystem& EntitySubsystem, FLWComponentSystemExecutionContext& Context) override;
 
 	void PrepareExecution();
 
 	template <typename TMassViewerLODInfoFragment = FMassLODInfoFragment>
-	void ExecuteInternal(UEntitySubsystem& EntitySubsystem, FLWComponentSystemExecutionContext& Context);
+	void ExecuteInternal(UMassEntitySubsystem& EntitySubsystem, FLWComponentSystemExecutionContext& Context);
 
 	
 	/**
@@ -73,7 +73,7 @@ namespace UE::MassRepresentation
 } // UE::MassRepresentation
 
 template <typename TMassViewerLODInfoFragment>
-void UMassVisualizationLODProcessor::ExecuteInternal(UEntitySubsystem& EntitySubsystem, FLWComponentSystemExecutionContext& Context)
+void UMassVisualizationLODProcessor::ExecuteInternal(UMassEntitySubsystem& EntitySubsystem, FLWComponentSystemExecutionContext& Context)
 {
 	if (bForceOFFLOD)
 	{
diff --git a/Engine/Plugins/Runtime/MassGameplay/Source/MassSignals/Private/MassSignalProcessorBase.cpp b/Engine/Plugins/Runtime/MassGameplay/Source/MassSignals/Private/MassSignalProcessorBase.cpp
--- a/Engine/Plugins/Runtime/MassGameplay/Source/MassSignals/Private/MassSignalProcessorBase.cpp
+++ b/Engine/Plugins/Runtime/MassGameplay/Source/MassSignals/Private/MassSignalProcessorBase.cpp
@@ -31,7 +31,7 @@ void UMassSignalProcessorBase::SubscribeToSignal(const FName SignalName)
 	SignalSubsystem->GetSignalDelegateByName(SignalName).AddUObject(this, &UMassSignalProcessorBase::OnSignalReceived);
 }
 
-void UMassSignalProcessorBase::Execute(UEntitySubsystem& EntitySubsystem, FLWComponentSystemExecutionContext& Context)
+void UMassSignalProcessorBase::Execute(UMassEntitySubsystem& EntitySubsystem, FLWComponentSystemExecutionContext& Context)
 {
 	QUICK_SCOPE_CYCLE_COUNTER(SignalEntities);
 
diff --git a/Engine/Plugins/Runtime/MassGameplay/Source/MassSignals/Private/MassSignalSubsystem.cpp b/Engine/Plugins/Runtime/MassGameplay/Source/MassSignals/Private/MassSignalSubsystem.cpp
--- a/Engine/Plugins/Runtime/MassGameplay/Source/MassSignals/Private/MassSignalSubsystem.cpp
+++ b/Engine/Plugins/Runtime/MassGameplay/Source/MassSignals/Private/MassSignalSubsystem.cpp
@@ -77,7 +77,7 @@ void UMassSignalSubsystem::SignalEntityDeferred(FLWComponentSystemExecutionConte
 void UMassSignalSubsystem::SignalEntitiesDeferred(FLWComponentSystemExecutionContext& Context, FName SignalName, TConstArrayView<FLWEntity> Entities)
 {
 	checkf(Entities.Num() > 0, TEXT("Expecting entities to signal"));
-	Context.Defer().EmplaceCommand<FDeferredCommand>([SignalName, Entities](const UEntitySubsystem& System)
+	Context.Defer().EmplaceCommand<FDeferredCommand>([SignalName, Entities](const UMassEntitySubsystem& System)
 	{
 		UMassSignalSubsystem* SignalSubsystem = UWorld::GetSubsystem<UMassSignalSubsystem>(System.GetWorld());
 		SignalSubsystem->SignalEntities(SignalName, Entities);
@@ -96,7 +96,7 @@ void UMassSignalSubsystem::DelaySignalEntityDeferred(FLWComponentSystemExecution
 void UMassSignalSubsystem::DelaySignalEntitiesDeferred(FLWComponentSystemExecutionContext& Context, FName SignalName, TConstArrayView<FLWEntity> Entities, const float DelayInSeconds)
 {
 	checkf(Entities.Num() > 0, TEXT("Expecting entities to signal"));
-	Context.Defer().EmplaceCommand<FDeferredCommand>([SignalName, Entities, DelayInSeconds](const UEntitySubsystem& System)
+	Context.Defer().EmplaceCommand<FDeferredCommand>([SignalName, Entities, DelayInSeconds](const UMassEntitySubsystem& System)
 	{
 		UMassSignalSubsystem* SignalSubsystem = UWorld::GetSubsystem<UMassSignalSubsystem>(System.GetWorld());
 		SignalSubsystem->DelaySignalEntities(SignalName, Entities, DelayInSeconds);
diff --git a/Engine/Plugins/Runtime/MassGameplay/Source/MassSignals/Private/MassSignalTypes.cpp b/Engine/Plugins/Runtime/MassGameplay/Source/MassSignals/Private/MassSignalTypes.cpp
--- a/Engine/Plugins/Runtime/MassGameplay/Source/MassSignals/Private/MassSignalTypes.cpp
+++ b/Engine/Plugins/Runtime/MassGameplay/Source/MassSignals/Private/MassSignalTypes.cpp
@@ -1,7 +1,7 @@
 // Copyright Epic Games, Inc. All Rights Reserved.
 
 #include "MassSignalTypes.h"
-#include "EntitySubsystem.h"
+#include "MassEntitySubsystem.h"
 
 uint64 FMassSignalNameLookup::GetOrAddSignalName(const FName SignalName)
 {
diff --git a/Engine/Plugins/Runtime/MassGameplay/Source/MassSignals/Public/MassSignalProcessorBase.h b/Engine/Plugins/Runtime/MassGameplay/Source/MassSignals/Public/MassSignalProcessorBase.h
--- a/Engine/Plugins/Runtime/MassGameplay/Source/MassSignals/Public/MassSignalProcessorBase.h
+++ b/Engine/Plugins/Runtime/MassGameplay/Source/MassSignals/Public/MassSignalProcessorBase.h
@@ -3,7 +3,7 @@
 #pragma once
 
 #include "MassProcessor.h"
-#include "EntitySubsystem.h"
+#include "MassEntitySubsystem.h"
 #include "MassSignalTypes.h"
 #include "Containers/StaticArray.h"
 #include "MassSignalProcessorBase.generated.h"
@@ -37,7 +37,7 @@ protected:
 	 * @param Context is the execution context to be passed when executing the lambdas
 	 * @param EntitySignals Look up to retrieve for each entities their raised signal via GetSignalsForEntity
 	 */
-	 virtual void SignalEntities(UEntitySubsystem& EntitySubsystem, FLWComponentSystemExecutionContext& Context, FMassSignalNameLookup& EntitySignals) PURE_VIRTUAL(UMassSignalProcessorBase::SignalEntities, );
+	 virtual void SignalEntities(UMassEntitySubsystem& EntitySubsystem, FLWComponentSystemExecutionContext& Context, FMassSignalNameLookup& EntitySignals) PURE_VIRTUAL(UMassSignalProcessorBase::SignalEntities, );
 
 	/**
 	 * Callback that is being called when new signal is raised
@@ -51,7 +51,7 @@ protected:
 	 * @param EntitySubsystem is the system to execute the lambdas on each entity chunk
 	 * @param Context is the execution context to be passed when executing the lambdas
 	 */
-	virtual void Execute(UEntitySubsystem& EntitySubsystem, FLWComponentSystemExecutionContext& Context) override;
+	virtual void Execute(UMassEntitySubsystem& EntitySubsystem, FLWComponentSystemExecutionContext& Context) override;
 
 	/**
 	 * To receive notification about a particular signal, you need to subscribe to it.
diff --git a/Engine/Plugins/Runtime/MassGameplay/Source/MassSignals/Public/MassSignalSubsystem.h b/Engine/Plugins/Runtime/MassGameplay/Source/MassSignals/Public/MassSignalSubsystem.h
--- a/Engine/Plugins/Runtime/MassGameplay/Source/MassSignals/Public/MassSignalSubsystem.h
+++ b/Engine/Plugins/Runtime/MassGameplay/Source/MassSignals/Public/MassSignalSubsystem.h
@@ -2,7 +2,7 @@
 
 #pragma once
 
-#include "EntitySubsystem.h"
+#include "MassEntitySubsystem.h"
 #include "Subsystems/WorldSubsystem.h"
 #include "MassSignalSubsystem.generated.h"
 
diff --git a/Engine/Plugins/Runtime/MassGameplay/Source/MassSignals/Public/MassSignalTypes.h b/Engine/Plugins/Runtime/MassGameplay/Source/MassSignals/Public/MassSignalTypes.h
--- a/Engine/Plugins/Runtime/MassGameplay/Source/MassSignals/Public/MassSignalTypes.h
+++ b/Engine/Plugins/Runtime/MassGameplay/Source/MassSignals/Public/MassSignalTypes.h
@@ -4,7 +4,7 @@
 
 #include "CoreMinimal.h"
 #include "Logging/LogMacros.h"
-#include "EntitySubsystem.h"
+#include "MassEntitySubsystem.h"
 //#include "MassSignalTypes.generated.h"
 
 DECLARE_LOG_CATEGORY_EXTERN(LogMassSignals, Log, All)
diff --git a/Engine/Plugins/Runtime/MassGameplay/Source/MassSimulation/Private/MassSimulationSubsystem.cpp b/Engine/Plugins/Runtime/MassGameplay/Source/MassSimulation/Private/MassSimulationSubsystem.cpp
--- a/Engine/Plugins/Runtime/MassGameplay/Source/MassSimulation/Private/MassSimulationSubsystem.cpp
+++ b/Engine/Plugins/Runtime/MassGameplay/Source/MassSimulation/Private/MassSimulationSubsystem.cpp
@@ -2,7 +2,7 @@
 
 #include "MassSimulationSubsystem.h"
 #include "MassSimulationLocalCoordinator.h"
-#include "MassEntitySystem.h"
+#include "MassEntitySubsystem.h"
 #include "MassExecutor.h"
 #include "MassSimulationSettings.h"
 #include "VisualLogger/VisualLogger.h"
@@ -12,15 +12,6 @@
 
 DEFINE_LOG_CATEGORY(LogMassSim);
 
-namespace FMassSimulationSubsystemHelper
-{
-	UPipeEntitySubsystem* GetEntitySubsystem(const UWorld* World)
-	{
-		UMassSimulationSubsystem* Instance = UWorld::GetSubsystem<UMassSimulationSubsystem>(World);
-		return Instance ? Instance->GetEntitySubsystem() : nullptr;
-	}
-}
-
 namespace UE::MassSimulation
 {
 	int32 bDoEntityCompaction = 1;
@@ -42,16 +33,8 @@ void UMassSimulationSubsystem::PostInitProperties()
 
 	Super::PostInitProperties();
 
-	if (HasAnyFlags(RF_ClassDefaultObject))
-	{
-		UPipeEntitySubsystem::SetGetter_Internal([](const UWorld* World)
-		{
-			return FMassSimulationSubsystemHelper::GetEntitySubsystem(World);
-		});
-	}
-	else
+	if (HasAnyFlags(RF_ClassDefaultObject) == false)
 	{
-		CreateEntitySubsystem();
 #if WITH_EDITOR
 		GET_MASS_CONFIG_VALUE(GetOnTickSchematicChanged()).AddUObject(this, &UMassSimulationSubsystem::RebuildTickPipeline);
 #endif // WITH_EDITOR
@@ -85,23 +68,13 @@ FPipeProcessingPhase::FOnPhaseEvent& UMassSimulationSubsystem::GetOnProcessingPh
 	return PhaseManager->GetOnPhaseEnd(Phase);
 }
 
-void UMassSimulationSubsystem::CreateEntitySubsystem()
-{
-	check(HasAnyFlags(RF_ClassDefaultObject) == false);
-	if (ensureMsgf(EntitySubsystem == nullptr, TEXT("Trying to create a second instance of EntitySubsystem. The call will be ignored.")))
-	{
-		UWorld* World = GetWorld();
-		check(World);
-		EntitySubsystem = NewObject<UPipeEntitySubsystem>(World);
-	}
-}
-
 void UMassSimulationSubsystem::Initialize(FSubsystemCollectionBase& Collection)
 {
 	Super::Initialize(Collection);
 
-	check(EntitySubsystem);
-	EntitySubsystem->Initialize(Collection);
+	CachedEntitySubsystem = Collection.InitializeDependency<UMassEntitySubsystem>();
+	check(CachedEntitySubsystem);
+	
 	GetOnProcessingPhaseStarted(EPipeProcessingPhase::PrePhysics).AddUObject(this, &UMassSimulationSubsystem::OnProcessingPhaseStarted, EPipeProcessingPhase::PrePhysics);
 }
 
@@ -175,10 +148,10 @@ void UMassSimulationSubsystem::OnProcessingPhaseStarted(const float DeltaSeconds
 		case EPipeProcessingPhase::PrePhysics:
 			{
 				TRACE_CPUPROFILER_EVENT_SCOPE(DoEntityCompation);
-				check(EntitySubsystem);
+				check(CachedEntitySubsystem);
 				if (UE::MassSimulation::bDoEntityCompaction)
 				{
-					EntitySubsystem->DoEntityCompaction(GET_MASS_CONFIG_VALUE(DesiredEntityCompactionTimeSlicePerTick));
+					CachedEntitySubsystem->DoEntityCompaction(GET_MASS_CONFIG_VALUE(DesiredEntityCompactionTimeSlicePerTick));
 				}
 			}
 			break;
diff --git a/Engine/Plugins/Runtime/MassGameplay/Source/MassSimulation/Public/MassSimulationLocalCoordinator.h b/Engine/Plugins/Runtime/MassGameplay/Source/MassSimulation/Public/MassSimulationLocalCoordinator.h
--- a/Engine/Plugins/Runtime/MassGameplay/Source/MassSimulation/Public/MassSimulationLocalCoordinator.h
+++ b/Engine/Plugins/Runtime/MassGameplay/Source/MassSimulation/Public/MassSimulationLocalCoordinator.h
@@ -16,7 +16,7 @@
 
 
 class UPipeSchematic;
-class UPipeEntitySubsystem;
+class UMassEntitySubsystem;
 class UHierarchicalInstancedStaticMeshComponent;
 class UMassSpawningManager;
 class UMassDebugVisualizationComponent;
diff --git a/Engine/Plugins/Runtime/MassGameplay/Source/MassSimulation/Public/MassSimulationSubsystem.h b/Engine/Plugins/Runtime/MassGameplay/Source/MassSimulation/Public/MassSimulationSubsystem.h
--- a/Engine/Plugins/Runtime/MassGameplay/Source/MassSimulation/Public/MassSimulationSubsystem.h
+++ b/Engine/Plugins/Runtime/MassGameplay/Source/MassSimulation/Public/MassSimulationSubsystem.h
@@ -10,7 +10,7 @@
 
 
 class AMassSimulationLocalCoordinator;
-class UPipeEntitySubsystem;
+class UMassEntitySubsystem;
 class UPipeProcessingPhaseManager;
 
 DECLARE_LOG_CATEGORY_EXTERN(LogMassSim, Log, All);
@@ -25,7 +25,7 @@ public:
 	
 	UMassSimulationSubsystem(const FObjectInitializer& ObjectInitializer = FObjectInitializer::Get());
 
-	UPipeEntitySubsystem* GetEntitySubsystem() const { return EntitySubsystem; }
+	UMassEntitySubsystem* GetEntitySubsystem() const { return CachedEntitySubsystem; }
 	const UPipeProcessingPhaseManager& GetPhaseManager() const { check(PhaseManager); return *PhaseManager; }
 
 	FPipeProcessingPhase::FOnPhaseEvent& GetOnProcessingPhaseStarted(const EPipeProcessingPhase Phase) const;
@@ -43,11 +43,6 @@ protected:
 	virtual void Deinitialize() override;
 	virtual void BeginDestroy() override;
 
-	/** Creates an instance of UPipeEntitySubsystem and assigns it to EntitySubsystem property. This is the entity 
-	 *  subsystem that Mass is going to be using. If you want to instantiate a different entity system class or point 
-	 *  at a preexisting instance just override this function. */
-	virtual void CreateEntitySubsystem();
-
 	void RebuildTickPipeline();
 
 	void StartSimulation(UWorld& InWorld);
@@ -63,7 +58,7 @@ protected:
 protected:
 
 	UPROPERTY()
-	UPipeEntitySubsystem* EntitySubsystem;
+	UMassEntitySubsystem* CachedEntitySubsystem;
 
 	UPROPERTY()
 	UPipeProcessingPhaseManager* PhaseManager;
diff --git a/Engine/Plugins/Runtime/MassGameplay/Source/MassSmartObjects/Private/MassSmartObjectBehaviorConfig.cpp b/Engine/Plugins/Runtime/MassGameplay/Source/MassSmartObjects/Private/MassSmartObjectBehaviorConfig.cpp
--- a/Engine/Plugins/Runtime/MassGameplay/Source/MassSmartObjects/Private/MassSmartObjectBehaviorConfig.cpp
+++ b/Engine/Plugins/Runtime/MassGameplay/Source/MassSmartObjects/Private/MassSmartObjectBehaviorConfig.cpp
@@ -3,7 +3,7 @@
 #include "MassSmartObjectBehaviorConfig.h"
 #include "MassSmartObjectProcessor.h"
 
-void USmartObjectMassBehaviorConfig::Activate(UEntitySubsystem& EntitySubsystem,
+void USmartObjectMassBehaviorConfig::Activate(UMassEntitySubsystem& EntitySubsystem,
                                               FLWComponentSystemExecutionContext& Context,
                                               const FMassBehaviorEntityContext& EntityContext) const
 {
diff --git a/Engine/Plugins/Runtime/MassGameplay/Source/MassSmartObjects/Private/MassSmartObjectHandler.cpp b/Engine/Plugins/Runtime/MassGameplay/Source/MassSmartObjects/Private/MassSmartObjectHandler.cpp
--- a/Engine/Plugins/Runtime/MassGameplay/Source/MassSmartObjects/Private/MassSmartObjectHandler.cpp
+++ b/Engine/Plugins/Runtime/MassGameplay/Source/MassSmartObjects/Private/MassSmartObjectHandler.cpp
@@ -12,7 +12,7 @@ namespace UE::Mass::SmartObject
 struct FPayload
 {
 	FLWEntity Entity;
-	TWeakObjectPtr<UEntitySubsystem> EntitySubsystem;
+	TWeakObjectPtr<UMassEntitySubsystem> EntitySubsystem;
 	TWeakObjectPtr<USmartObjectSubsystem> SmartObjectSubsystem;
 };
 
diff --git a/Engine/Plugins/Runtime/MassGameplay/Source/MassSmartObjects/Private/MassSmartObjectProcessor.cpp b/Engine/Plugins/Runtime/MassGameplay/Source/MassSmartObjects/Private/MassSmartObjectProcessor.cpp
--- a/Engine/Plugins/Runtime/MassGameplay/Source/MassSmartObjects/Private/MassSmartObjectProcessor.cpp
+++ b/Engine/Plugins/Runtime/MassGameplay/Source/MassSmartObjects/Private/MassSmartObjectProcessor.cpp
@@ -56,7 +56,7 @@ UMassProcessor_SmartObjectCandidatesFinder::UMassProcessor_SmartObjectCandidates
 	ExecutionOrder.ExecuteBefore.Add(UE::Mass::ProcessorGroupNames::Behavior);
 }
 
-void UMassProcessor_SmartObjectCandidatesFinder::Execute(UEntitySubsystem& EntitySubsystem, FLWComponentSystemExecutionContext& Context)
+void UMassProcessor_SmartObjectCandidatesFinder::Execute(UMassEntitySubsystem& EntitySubsystem, FLWComponentSystemExecutionContext& Context)
 {
 	checkf(SmartObjectSubsystem != nullptr, TEXT("SmartObjectSubsystem should exist when executing processors."));
 	checkf(SignalSubsystem != nullptr, TEXT("MassSignalSubsystem should exist when executing processors."));
@@ -278,7 +278,7 @@ UMassProcessor_SmartObjectTimedBehavior::UMassProcessor_SmartObjectTimedBehavior
 	ExecutionOrder.ExecuteAfter.Add(UE::Mass::ProcessorGroupNames::Behavior);
 }
 
-void UMassProcessor_SmartObjectTimedBehavior::Execute(UEntitySubsystem& EntitySubsystem, FLWComponentSystemExecutionContext& Context)
+void UMassProcessor_SmartObjectTimedBehavior::Execute(UMassEntitySubsystem& EntitySubsystem, FLWComponentSystemExecutionContext& Context)
 {
 	checkf(SmartObjectSubsystem != nullptr, TEXT("SmartObjectSubsystem should exist when executing processors."));
 	checkf(SignalSubsystem != nullptr, TEXT("MassSignalSubsystem should exist when executing processors."));
diff --git a/Engine/Plugins/Runtime/MassGameplay/Source/MassSmartObjects/Public/MassSmartObjectBehaviorConfig.h b/Engine/Plugins/Runtime/MassGameplay/Source/MassSmartObjects/Public/MassSmartObjectBehaviorConfig.h
--- a/Engine/Plugins/Runtime/MassGameplay/Source/MassSmartObjects/Public/MassSmartObjectBehaviorConfig.h
+++ b/Engine/Plugins/Runtime/MassGameplay/Source/MassSmartObjects/Public/MassSmartObjectBehaviorConfig.h
@@ -5,7 +5,7 @@
 #include "SmartObjectConfig.h"
 #include "MassSmartObjectBehaviorConfig.generated.h"
 
-class UEntitySubsystem;
+class UMassEntitySubsystem;
 class USmartObjectSubsystem;
 struct FLWEntity;
 struct FLWComponentSystemExecutionContext;
@@ -43,7 +43,7 @@ public:
 	 * This virtual method allows subclasses to configure the LW Entity based on their
 	 * parameters (e.g. Add new fragments)
 	 */
-	virtual void Activate(UEntitySubsystem& EntitySubsystem, FLWComponentSystemExecutionContext& Context, const FMassBehaviorEntityContext& EntityContext) const;
+	virtual void Activate(UMassEntitySubsystem& EntitySubsystem, FLWComponentSystemExecutionContext& Context, const FMassBehaviorEntityContext& EntityContext) const;
 
 	/**
 	 * Indicates the amount of time the LW entity
diff --git a/Engine/Plugins/Runtime/MassGameplay/Source/MassSmartObjects/Public/MassSmartObjectHandler.h b/Engine/Plugins/Runtime/MassGameplay/Source/MassSmartObjects/Public/MassSmartObjectHandler.h
--- a/Engine/Plugins/Runtime/MassGameplay/Source/MassSmartObjects/Public/MassSmartObjectHandler.h
+++ b/Engine/Plugins/Runtime/MassGameplay/Source/MassSmartObjects/Public/MassSmartObjectHandler.h
@@ -5,7 +5,7 @@
 #include "MassSmartObjectRequest.h"
 #include "MassSmartObjectTypes.h"
 
-class UEntitySubsystem;
+class UMassEntitySubsystem;
 class USmartObjectSubsystem;
 struct FLWComponentSystemExecutionContext;
 struct FLWEntity;
@@ -28,7 +28,7 @@ struct MASSSMARTOBJECTS_API FMassSmartObjectHandler
 	 * @param InExecutionContext is the current execution context of the entity subsystem
 	 * @param InSmartObjectSubsystem is the smart object subsystem
 	 */
-	FMassSmartObjectHandler(UEntitySubsystem& InEntitySubsystem, FLWComponentSystemExecutionContext& InExecutionContext, USmartObjectSubsystem& InSmartObjectSubsystem)
+	FMassSmartObjectHandler(UMassEntitySubsystem& InEntitySubsystem, FLWComponentSystemExecutionContext& InExecutionContext, USmartObjectSubsystem& InSmartObjectSubsystem)
 		: EntitySubsystem(InEntitySubsystem)
 		, ExecutionContext(InExecutionContext)
 		, SmartObjectSubsystem(InSmartObjectSubsystem)
@@ -118,7 +118,7 @@ struct MASSSMARTOBJECTS_API FMassSmartObjectHandler
 	void ReleaseSmartObject(const FLWEntity Entity, FDataFragment_SmartObjectUser& User, const EMassSmartObjectInteractionStatus Status = EMassSmartObjectInteractionStatus::Unset) const;
 
 private:
-	UEntitySubsystem& EntitySubsystem;
+	UMassEntitySubsystem& EntitySubsystem;
 	FLWComponentSystemExecutionContext& ExecutionContext;
 	USmartObjectSubsystem& SmartObjectSubsystem;
 };
diff --git a/Engine/Plugins/Runtime/MassGameplay/Source/MassSmartObjects/Public/MassSmartObjectProcessor.h b/Engine/Plugins/Runtime/MassGameplay/Source/MassSmartObjects/Public/MassSmartObjectProcessor.h
--- a/Engine/Plugins/Runtime/MassGameplay/Source/MassSmartObjects/Public/MassSmartObjectProcessor.h
+++ b/Engine/Plugins/Runtime/MassGameplay/Source/MassSmartObjects/Public/MassSmartObjectProcessor.h
@@ -77,7 +77,7 @@ public:
 protected:
 	virtual void Initialize(UObject& Owner) override;
 	virtual void ConfigureQueries() override;
-	virtual void Execute(UEntitySubsystem& EntitySubsystem, FLWComponentSystemExecutionContext& Context) override;
+	virtual void Execute(UMassEntitySubsystem& EntitySubsystem, FLWComponentSystemExecutionContext& Context) override;
 
 	/** Extents used to perform the spatial query in the octree for world location queries. */
 	UPROPERTY(EditDefaultsOnly, Category = SmartObject, config)
@@ -112,7 +112,7 @@ public:
 	UMassProcessor_SmartObjectTimedBehavior();
 
 protected:
-	virtual void Execute(UEntitySubsystem& EntitySubsystem, FLWComponentSystemExecutionContext& Context) override;
+	virtual void Execute(UMassEntitySubsystem& EntitySubsystem, FLWComponentSystemExecutionContext& Context) override;
 	virtual void ConfigureQueries() override;
 
 	FLWComponentQuery EntityQuery;
diff --git a/Engine/Plugins/Runtime/MassGameplay/Source/MassSpawner/Private/MassEntityTemplate.cpp b/Engine/Plugins/Runtime/MassGameplay/Source/MassSpawner/Private/MassEntityTemplate.cpp
--- a/Engine/Plugins/Runtime/MassGameplay/Source/MassSpawner/Private/MassEntityTemplate.cpp
+++ b/Engine/Plugins/Runtime/MassGameplay/Source/MassSpawner/Private/MassEntityTemplate.cpp
@@ -29,7 +29,7 @@ void FMassEntityTemplate::SetUpProcessors(const FMassObjectHandlers& ObjectHandl
 	DeinitializationPipeline.InitializeFromArray(ObjectHandlers.Deinitializers, PipelineOwner);
 }
 
-FString FMassEntityTemplate::DebugGetDescription(UEntitySubsystem* EntitySubsystem) const
+FString FMassEntityTemplate::DebugGetDescription(UMassEntitySubsystem* EntitySubsystem) const
 { 
 	FStringOutputDevice Ar;
 #if WITH_MASS_DEBUG
@@ -50,7 +50,7 @@ FString FMassEntityTemplate::DebugGetDescription(UEntitySubsystem* EntitySubsyst
 	return MoveTemp(Ar);
 }
 
-FString FMassEntityTemplate::DebugGetArchetypeDescription(UEntitySubsystem& EntitySubsystem) const
+FString FMassEntityTemplate::DebugGetArchetypeDescription(UMassEntitySubsystem& EntitySubsystem) const
 {
 	FStringOutputDevice OutDescription;
 #if WITH_MASS_DEBUG
diff --git a/Engine/Plugins/Runtime/MassGameplay/Source/MassSpawner/Private/MassEntityTemplateRegistry.cpp b/Engine/Plugins/Runtime/MassGameplay/Source/MassSpawner/Private/MassEntityTemplateRegistry.cpp
--- a/Engine/Plugins/Runtime/MassGameplay/Source/MassSpawner/Private/MassEntityTemplateRegistry.cpp
+++ b/Engine/Plugins/Runtime/MassGameplay/Source/MassSpawner/Private/MassEntityTemplateRegistry.cpp
@@ -2,7 +2,7 @@
 
 #include "MassEntityTemplateRegistry.h"
 #include "MassSpawnerTypes.h"
-#include "MassEntitySystem.h"
+#include "MassEntitySubsystem.h"
 #include "Engine/World.h"
 #include "VisualLogger/VisualLogger.h"
 #include "HAL/IConsoleManager.h"
@@ -138,7 +138,7 @@ bool UMassEntityTemplateRegistry::BuildTemplateImpl(const FStructToTemplateBuild
 		InitializeEntityTemplate(OutTemplate);
 
 		UE_VLOG(this, LogMassSpawner, Log, TEXT("Created entity template for %s:\n%s"), *GetNameSafe(StructInstance.GetScriptStruct())
-			, *OutTemplate.DebugGetDescription(UPipeEntitySubsystem::GetCurrent(World)));
+			, *OutTemplate.DebugGetDescription(UMassEntitySubsystem::GetCurrent(World)));
 
 		return true;
 	}
@@ -252,7 +252,7 @@ void UMassEntityTemplateRegistry::InitializeEntityTemplate(FMassEntityTemplate&
 
 	UWorld* World = GetWorld();
 	// find or create template
-	UPipeEntitySubsystem* EntitySys = UPipeEntitySubsystem::GetCurrent(World);
+	UMassEntitySubsystem* EntitySys = UMassEntitySubsystem::GetCurrent(World);
 	check(EntitySys);
 
 	FTemplateRegistryHelpers::FragmentInstancesToTypes(OutTemplate.GetFragments(), OutTemplate.GetMutableFragmentTypes());
diff --git a/Engine/Plugins/Runtime/MassGameplay/Source/MassSpawner/Private/MassSpawnerSubsystem.cpp b/Engine/Plugins/Runtime/MassGameplay/Source/MassSpawner/Private/MassSpawnerSubsystem.cpp
--- a/Engine/Plugins/Runtime/MassGameplay/Source/MassSpawner/Private/MassSpawnerSubsystem.cpp
+++ b/Engine/Plugins/Runtime/MassGameplay/Source/MassSpawner/Private/MassSpawnerSubsystem.cpp
@@ -3,7 +3,7 @@
 #include "MassSpawnerSubsystem.h"
 #include "MassSpawnerTypes.h"
 #include "MassEntityTemplate.h"
-#include "MassEntitySystem.h"
+#include "MassEntitySubsystem.h"
 #include "MassEntityTemplateRegistry.h"
 #include "Engine/World.h"
 #include "MassExecutor.h"
@@ -19,7 +19,7 @@
 namespace UE::Mass::SpawnerSubsystem 
 {
 
-	bool RunFragmentDestructors(const FArchetypeChunkCollection& Chunks, UMassSpawnerSubsystem& SpawnerSubSystem, UPipeEntitySubsystem& EntitySystem)
+	bool RunFragmentDestructors(const FArchetypeChunkCollection& Chunks, UMassSpawnerSubsystem& SpawnerSubSystem, UMassEntitySubsystem& EntitySystem)
 	{
 		const FArchetypeHandle ArchetypeHandle = Chunks.GetArchetype();
 		// @todo this is a temporary measure to skip entities we've destroyed without the
@@ -56,7 +56,7 @@ namespace UE::Mass::SpawnerSubsystem
 		return Destructors.Num() > 0;
 	}
 	
-	void CreateSparseChunks(const UPipeEntitySubsystem& EntitySystem, const TConstArrayView<FLWEntity> Entities, TArray<FArchetypeChunkCollection>& OutChunkCollections)
+	void CreateSparseChunks(const UMassEntitySubsystem& EntitySystem, const TConstArrayView<FLWEntity> Entities, TArray<FArchetypeChunkCollection>& OutChunkCollections)
 	{
 		TRACE_CPUPROFILER_EVENT_SCOPE_STR("SpawnerSubsystem_CreateSparseChunks");
 
@@ -93,7 +93,7 @@ void UMassSpawnerSubsystem::Initialize(FSubsystemCollectionBase& Collection)
 	Collection.InitializeDependency<UMassSimulationSubsystem>();
 
 	SimulationSystem = UWorld::GetSubsystem<UMassSimulationSubsystem>(GetWorld());
-	EntitySystem = UPipeEntitySubsystem::GetCurrent(GetWorld());
+	EntitySystem = UMassEntitySubsystem::GetCurrent(GetWorld());
 	check(EntitySystem);
 }
 
diff --git a/Engine/Plugins/Runtime/MassGameplay/Source/MassSpawner/Public/MassEntityTemplate.h b/Engine/Plugins/Runtime/MassGameplay/Source/MassSpawner/Public/MassEntityTemplate.h
--- a/Engine/Plugins/Runtime/MassGameplay/Source/MassSpawner/Public/MassEntityTemplate.h
+++ b/Engine/Plugins/Runtime/MassGameplay/Source/MassSpawner/Public/MassEntityTemplate.h
@@ -4,7 +4,7 @@
 
 #include "CoreMinimal.h"
 #include "MassEntityTypes.h"
-#include "EntitySubsystem.h"
+#include "MassEntitySubsystem.h"
 #include "MassCommonTypes.h"
 #include "MassTranslator.h"
 #include "MassEntityTemplate.generated.h"
@@ -142,8 +142,8 @@ struct MASSSPAWNER_API FMassEntityTemplate
 		return ChunkFragments[Index].template GetMutable<T>();
 	}
 
-	FString DebugGetDescription(UEntitySubsystem* EntitySubsystem = nullptr) const;
-	FString DebugGetArchetypeDescription(UEntitySubsystem& EntitySubsystem) const;
+	FString DebugGetDescription(UMassEntitySubsystem* EntitySubsystem = nullptr) const;
+	FString DebugGetArchetypeDescription(UMassEntitySubsystem& EntitySubsystem) const;
 
 private:
 	FArchetypeHandle Archetype;
diff --git a/Engine/Plugins/Runtime/MassGameplay/Source/MassSpawner/Public/MassEntityTemplateRegistry.h b/Engine/Plugins/Runtime/MassGameplay/Source/MassSpawner/Public/MassEntityTemplateRegistry.h
--- a/Engine/Plugins/Runtime/MassGameplay/Source/MassSpawner/Public/MassEntityTemplateRegistry.h
+++ b/Engine/Plugins/Runtime/MassGameplay/Source/MassSpawner/Public/MassEntityTemplateRegistry.h
@@ -12,7 +12,7 @@
 
 
 class UWorld;
-class UEntitySubsystem;
+class UMassEntitySubsystem;
 class UMassTranslator;
 
 enum class EFragmentInitialization : uint8
diff --git a/Engine/Plugins/Runtime/MassGameplay/Source/MassSpawner/Public/MassSpawnerSubsystem.h b/Engine/Plugins/Runtime/MassGameplay/Source/MassSpawner/Public/MassSpawnerSubsystem.h
--- a/Engine/Plugins/Runtime/MassGameplay/Source/MassSpawner/Public/MassSpawnerSubsystem.h
+++ b/Engine/Plugins/Runtime/MassGameplay/Source/MassSpawner/Public/MassSpawnerSubsystem.h
@@ -2,11 +2,11 @@
 
 #pragma once
 
-#include "EntitySubsystem.h"
+#include "MassEntitySubsystem.h"
 #include "Subsystems/WorldSubsystem.h"
 #include "MassSpawnerSubsystem.generated.h"
 
-class UPipeEntitySubsystem;
+class UMassEntitySubsystem;
 struct FMassSpawnConfigBase;
 struct FMassEntityTemplate;
 class UMassEntityTemplateRegistry;
@@ -56,7 +56,7 @@ protected:
 	void DoSpawning(const FMassEntityTemplate& EntityTemplate, const FMassSpawnConfigBase& Data, const FStructView& AuxData, TArray<FLWEntity>& OutEntities) const;
 
 	UPROPERTY()
-	UPipeEntitySubsystem* EntitySystem;
+	UMassEntitySubsystem* EntitySystem;
 
 	UPROPERTY()
 	UMassEntityTemplateRegistry* TemplateRegistryInstance;
diff --git a/Engine/Plugins/Runtime/MassGameplay/Source/MassSpawner/Public/MassTranslator.h b/Engine/Plugins/Runtime/MassGameplay/Source/MassSpawner/Public/MassTranslator.h
--- a/Engine/Plugins/Runtime/MassGameplay/Source/MassSpawner/Public/MassTranslator.h
+++ b/Engine/Plugins/Runtime/MassGameplay/Source/MassSpawner/Public/MassTranslator.h
@@ -8,7 +8,7 @@
 #include "MassTranslator.generated.h"
 
 
-class UPipeEntitySubsystem;
+class UMassEntitySubsystem;
 
 UENUM()
 enum class EMassTranslationDirection : uint8
