AggregateTicking and Pipe type renames:

diff --git a/Engine/Plugins/Runtime/MassEntity/Source/MassEntity/Private/ArchetypeData.cpp b/Engine/Plugins/Runtime/MassEntity/Source/MassEntity/Private/ArchetypeData.cpp
--- a/Engine/Plugins/Runtime/MassEntity/Source/MassEntity/Private/ArchetypeData.cpp
+++ b/Engine/Plugins/Runtime/MassEntity/Source/MassEntity/Private/ArchetypeData.cpp
@@ -6,26 +6,26 @@
 #include "Async/ParallelFor.h"
 
 //////////////////////////////////////////////////////////////////////
-// FArchetypeData
+// FMassArchetypeData
 namespace UE { namespace AggregateTicking 
 {
 	constexpr static bool bBitwiseRelocateComponents = true;
 }} // UE::AggregateTicking
 
-void FArchetypeData::ForEachComponentType(TFunction< void(const UScriptStruct* /*Component*/)> Function) const
+void FMassArchetypeData::ForEachComponentType(TFunction< void(const UScriptStruct* /*Component*/)> Function) const
 {
-	for (const FArchetypeComponentConfig& ComponentData : ComponentConfigs)
+	for (const FMassArchetypeFragmentConfig& ComponentData : ComponentConfigs)
 	{
 		Function(ComponentData.ComponentType);
 	}
 }
 
-bool FArchetypeData::HasComponentType(const UScriptStruct* ComponentType) const
+bool FMassArchetypeData::HasComponentType(const UScriptStruct* ComponentType) const
 {
 	return (ComponentType && CompositionDescriptor.Components.Contains(*ComponentType));
 }
 
-bool FArchetypeData::IsEquivalent(TConstArrayView<const UScriptStruct*> OtherList) const
+bool FMassArchetypeData::IsEquivalent(TConstArrayView<const UScriptStruct*> OtherList) const
 {
 	if (OtherList.Num() != ComponentConfigs.Num())
 	{
@@ -43,12 +43,12 @@ bool FArchetypeData::IsEquivalent(TConstArrayView<const UScriptStruct*> OtherLis
 	return true;
 }
 
-void FArchetypeData::Initialize(const FLWComponentBitSet& Components, const FLWTagBitSet& Tags, const FLWChunkComponentBitSet& ChunkComponents)
+void FMassArchetypeData::Initialize(const FMassFragmentBitSet& Components, const FMassTagBitSet& Tags, const FMassChunkFragmentBitSet& ChunkComponents)
 {
 	TArray<const UScriptStruct*, TInlineAllocator<16>> SortedComponentList;
 	Components.ExportTypes(SortedComponentList);
 
-	SortedComponentList.Sort(FLWComponentSorterOperator<UScriptStruct>());
+	SortedComponentList.Sort(FMassSorterOperator<UScriptStruct>());
 
 	// Figure out how many bytes all of the individual components (and metadata) will cost per entity
 	int32 ComponentSizeTallyBytes = 0;
@@ -57,7 +57,7 @@ void FArchetypeData::Initialize(const FLWComponentBitSet& Components, const FLWT
 	int32 AlignmentPadding = 0;
 	{
 		// Save room for the 'metadata' (entity array)
-		ComponentSizeTallyBytes += sizeof(FLWEntity);
+		ComponentSizeTallyBytes += sizeof(FMassEntityHandle);
 
 		// Tally up the component sizes and place them in the index map
 		ComponentConfigs.AddDefaulted(SortedComponentList.Num());
@@ -83,7 +83,7 @@ void FArchetypeData::Initialize(const FLWComponentBitSet& Components, const FLWT
 	// create the chunk component instances
 	TArray<const UScriptStruct*, TInlineAllocator<16>> ChunkComponentList;
 	ChunkComponents.ExportTypes(ChunkComponentList);
-	ChunkComponentList.Sort(FLWComponentSorterOperator<UScriptStruct>());
+	ChunkComponentList.Sort(FMassSorterOperator<UScriptStruct>());
 	for (const UScriptStruct* ChunkComponentType : ChunkComponentList)
 	{
 		check(ChunkComponentType);
@@ -98,8 +98,8 @@ void FArchetypeData::Initialize(const FLWComponentBitSet& Components, const FLWT
 
 	// Set up the offsets for each component into the chunk data
 	EntityListOffsetWithinChunk = 0;
-	int32 CurrentOffset = NumEntitiesPerChunk * sizeof(FLWEntity);
-	for (FArchetypeComponentConfig& ComponentData : ComponentConfigs)
+	int32 CurrentOffset = NumEntitiesPerChunk * sizeof(FMassEntityHandle);
+	for (FMassArchetypeFragmentConfig& ComponentData : ComponentConfigs)
 	{
 		CurrentOffset = Align(CurrentOffset, ComponentData.ComponentType->GetMinAlignment());
 		ComponentData.ArrayOffsetWithinChunk = CurrentOffset;
@@ -108,7 +108,7 @@ void FArchetypeData::Initialize(const FLWComponentBitSet& Components, const FLWT
 	}
 }
 
-void FArchetypeData::InitializeWithSibling(const FArchetypeData& SiblingArchetype, const FLWTagBitSet& OverrideTags)
+void FMassArchetypeData::InitializeWithSibling(const FMassArchetypeData& SiblingArchetype, const FMassTagBitSet& OverrideTags)
 {
 	checkf(IsInitialized() == false, TEXT("Trying to %s but this archetype has already been initialized"));
 
@@ -126,12 +126,12 @@ void FArchetypeData::InitializeWithSibling(const FArchetypeData& SiblingArchetyp
 	EntityListOffsetWithinChunk = 0;
 }
 
-void FArchetypeData::AddEntity(FLWEntity Entity)
+void FMassArchetypeData::AddEntity(FMassEntityHandle Entity)
 {
 	AddEntityInternal(Entity, true/*bInitializeComponents*/);
 }
 
-int32 FArchetypeData::AddEntityInternal(FLWEntity Entity, const bool bInitializeComponents)
+int32 FMassArchetypeData::AddEntityInternal(FMassEntityHandle Entity, const bool bInitializeComponents)
 {
 	int32 IndexWithinChunk = 0;
 	int32 AbsoluteIndex = 0;
@@ -139,10 +139,10 @@ int32 FArchetypeData::AddEntityInternal(FLWEntity Entity, const bool bInitialize
 	int32 EmptyChunkIndex = INDEX_NONE;
 	int32 EmptyAbsoluteIndex = INDEX_NONE;
 
-	FArchetypeChunk* DestinationChunk = nullptr;
+	FMassArchetypeChunk* DestinationChunk = nullptr;
 	// Check chunks for a free spot (trying to reuse the earlier ones first so later ones might get freed up) 
 	//@TODO: This could be accelerated to include a cached index to the first chunk with free spots or similar
-	for (FArchetypeChunk& Chunk : Chunks)
+	for (FMassArchetypeChunk& Chunk : Chunks)
 	{
 		if (Chunk.GetNumInstances() == 0)
 		{
@@ -188,7 +188,7 @@ int32 FArchetypeData::AddEntityInternal(FLWEntity Entity, const bool bInitialize
 	// Initialize the component memory
 	if (bInitializeComponents)
 	{
-		for (const FArchetypeComponentConfig& ComponentConfig : ComponentConfigs)
+		for (const FMassArchetypeFragmentConfig& ComponentConfig : ComponentConfigs)
 		{
 			void* ComponentPtr = ComponentConfig.GetComponentData(DestinationChunk->GetRawMemory(), IndexWithinChunk);
 			ComponentConfig.ComponentType->InitializeStruct(ComponentPtr);
@@ -202,18 +202,18 @@ int32 FArchetypeData::AddEntityInternal(FLWEntity Entity, const bool bInitialize
 	return AbsoluteIndex;
 }
 
-void FArchetypeData::RemoveEntity(FLWEntity Entity)
+void FMassArchetypeData::RemoveEntity(FMassEntityHandle Entity)
 {
 	const int32 AbsoluteIndex = EntityMap.FindAndRemoveChecked(Entity.Index);
 	RemoveEntityInternal(AbsoluteIndex, true/*bDestroyComponents*/);
 }
 
-void FArchetypeData::RemoveEntityInternal(const int32 AbsoluteIndex, const bool bDestroyComponents)
+void FMassArchetypeData::RemoveEntityInternal(const int32 AbsoluteIndex, const bool bDestroyComponents)
 {
 	const int32 ChunkIndex = AbsoluteIndex / NumEntitiesPerChunk;
 	const int32 IndexWithinChunk = AbsoluteIndex % NumEntitiesPerChunk;
 
-	FArchetypeChunk& Chunk = Chunks[ChunkIndex];
+	FMassArchetypeChunk& Chunk = Chunks[ChunkIndex];
 
 	Chunk.RemoveInstance();
 	const int32 IndexToSwapFrom = Chunk.GetNumInstances();
@@ -223,7 +223,7 @@ void FArchetypeData::RemoveEntityInternal(const int32 AbsoluteIndex, const bool
 	// Remove and swap the last entry in the chunk to the location of the removed item (if it's not the same as the dying entry)
 	if (IndexToSwapFrom != IndexWithinChunk)
 	{
-		for (const FArchetypeComponentConfig& ComponentConfig : ComponentConfigs)
+		for (const FMassArchetypeFragmentConfig& ComponentConfig : ComponentConfigs)
 		{
 			void* DyingComponentPtr = ComponentConfig.GetComponentData(Chunk.GetRawMemory(), IndexWithinChunk);
 			void* MovingComponentPtr = ComponentConfig.GetComponentData(Chunk.GetRawMemory(), IndexToSwapFrom);
@@ -253,13 +253,13 @@ void FArchetypeData::RemoveEntityInternal(const int32 AbsoluteIndex, const bool
 		}
 
 		// Update the entity table and map
-		const FLWEntity EntityBeingSwapped = Chunk.GetEntityArrayElementRef(EntityListOffsetWithinChunk, IndexToSwapFrom);
+		const FMassEntityHandle EntityBeingSwapped = Chunk.GetEntityArrayElementRef(EntityListOffsetWithinChunk, IndexToSwapFrom);
 		Chunk.GetEntityArrayElementRef(EntityListOffsetWithinChunk, IndexWithinChunk) = EntityBeingSwapped;
 		EntityMap.FindChecked(EntityBeingSwapped.Index) = AbsoluteIndex;
 	}
 	else if (!UE::AggregateTicking::bBitwiseRelocateComponents || bDestroyComponents)
 	{
-		for (const FArchetypeComponentConfig& ComponentConfig : ComponentConfigs)
+		for (const FMassArchetypeFragmentConfig& ComponentConfig : ComponentConfigs)
 		{
 			// Destroy the component data
 			void* DyingComponentPtr = ComponentConfig.GetComponentData(Chunk.GetRawMemory(), IndexWithinChunk);
@@ -275,7 +275,7 @@ void FArchetypeData::RemoveEntityInternal(const int32 AbsoluteIndex, const bool
 	}
 }
 
-void FArchetypeData::BatchDestroyEntityChunks(const FArchetypeChunkCollection& ChunkCollection, TArray<FLWEntity>& OutEntitiesRemoved)
+void FMassArchetypeData::BatchDestroyEntityChunks(const FArchetypeChunkCollection& ChunkCollection, TArray<FMassEntityHandle>& OutEntitiesRemoved)
 {
 	const int32 InitialOutEntitiesCount = OutEntitiesRemoved.Num();
 
@@ -289,10 +289,10 @@ void FArchetypeData::BatchDestroyEntityChunks(const FArchetypeChunkCollection& C
 
 	for (const FArchetypeChunkCollection::FChunkInfo SubchunkInfo : Subchunks)
 	{ 
-		FArchetypeChunk& Chunk = Chunks[SubchunkInfo.ChunkIndex];
+		FMassArchetypeChunk& Chunk = Chunks[SubchunkInfo.ChunkIndex];
 
 		// gather entities we're about to remove
-		FLWEntity* DyingEntityPtr = &Chunk.GetEntityArrayElementRef(EntityListOffsetWithinChunk, SubchunkInfo.SubchunkStart);
+		FMassEntityHandle* DyingEntityPtr = &Chunk.GetEntityArrayElementRef(EntityListOffsetWithinChunk, SubchunkInfo.SubchunkStart);
 		OutEntitiesRemoved.Append(DyingEntityPtr, SubchunkInfo.Length);
 
 		const int32 NumberToMove = FMath::Min(Chunk.GetNumInstances() - (SubchunkInfo.SubchunkStart + SubchunkInfo.Length), SubchunkInfo.Length);
@@ -304,7 +304,7 @@ void FArchetypeData::BatchDestroyEntityChunks(const FArchetypeChunkCollection& C
 			const int32 SwapStartIndex = Chunk.GetNumInstances() - NumberToMove;
 			checkf((SubchunkInfo.SubchunkStart + NumberToMove - 1) < SwapStartIndex, TEXT("Remove and Move ranges overlap"));
 
-			for (const FArchetypeComponentConfig& ComponentConfig : ComponentConfigs)
+			for (const FMassArchetypeFragmentConfig& ComponentConfig : ComponentConfigs)
 			{
 				void* DyingComponentPtr = ComponentConfig.GetComponentData(Chunk.GetRawMemory(), SubchunkInfo.SubchunkStart);
 				void* MovingComponentPtr = ComponentConfig.GetComponentData(Chunk.GetRawMemory(), SwapStartIndex);
@@ -331,7 +331,7 @@ void FArchetypeData::BatchDestroyEntityChunks(const FArchetypeChunkCollection& C
 			}
 
 			// Update the entity table and map
-			const FLWEntity* MovingEntityPtr = &Chunk.GetEntityArrayElementRef(EntityListOffsetWithinChunk, SwapStartIndex);
+			const FMassEntityHandle* MovingEntityPtr = &Chunk.GetEntityArrayElementRef(EntityListOffsetWithinChunk, SwapStartIndex);
 			int32 AbsoluteIndex = SubchunkInfo.ChunkIndex * NumEntitiesPerChunk + SubchunkInfo.SubchunkStart;
 
 			for (int i = 0; i < NumberToMove; ++i)
@@ -346,7 +346,7 @@ void FArchetypeData::BatchDestroyEntityChunks(const FArchetypeChunkCollection& C
 			// just clean up the rest. Note that we explicitly do not clean the spots vacated by entities moved from 
 			// the back of the chunk - if we did the risk calling DestroyStruct on them multiple times
 			const int32 CutStartIndex = SubchunkInfo.SubchunkStart + NumberToMove;
-			for (const FArchetypeComponentConfig& ComponentConfig : ComponentConfigs)
+			for (const FMassArchetypeFragmentConfig& ComponentConfig : ComponentConfigs)
 			{
 				// Destroy the component data
 				void* DyingComponentPtr = ComponentConfig.GetComponentData(Chunk.GetRawMemory(), CutStartIndex);
@@ -370,12 +370,12 @@ void FArchetypeData::BatchDestroyEntityChunks(const FArchetypeChunkCollection& C
 	}
 }
 
-bool FArchetypeData::HasComponentDataForEntity(const UScriptStruct* ComponentType, int32 EntityIndex) const
+bool FMassArchetypeData::HasComponentDataForEntity(const UScriptStruct* ComponentType, int32 EntityIndex) const
 {
 	return (ComponentType && CompositionDescriptor.Components.Contains(*ComponentType));
 }
 
-void* FArchetypeData::GetComponentDataForEntityChecked(const UScriptStruct* ComponentType, int32 EntityIndex) const
+void* FMassArchetypeData::GetComponentDataForEntityChecked(const UScriptStruct* ComponentType, int32 EntityIndex) const
 {
 	const FInternalEntityHandle InternalIndex = MakeEntityHandle(EntityIndex);
 	
@@ -384,7 +384,7 @@ void* FArchetypeData::GetComponentDataForEntityChecked(const UScriptStruct* Comp
 	return GetComponentData(ComponentIndex, InternalIndex);
 }
 
-void* FArchetypeData::GetComponentDataForEntity(const UScriptStruct* ComponentType, int32 EntityIndex) const
+void* FMassArchetypeData::GetComponentDataForEntity(const UScriptStruct* ComponentType, int32 EntityIndex) const
 {
 	if (const int32* ComponentIndex = ComponentIndexMap.Find(ComponentType))
 	{
@@ -395,7 +395,7 @@ void* FArchetypeData::GetComponentDataForEntity(const UScriptStruct* ComponentTy
 	return nullptr;
 }
 
-void FArchetypeData::SetComponentsData(const FLWEntity Entity, TArrayView<const FInstancedStruct> ComponentInstances)
+void FMassArchetypeData::SetComponentsData(const FMassEntityHandle Entity, TArrayView<const FInstancedStruct> ComponentInstances)
 {
 	FInternalEntityHandle InternalIndex = MakeEntityHandle(Entity);
 
@@ -410,7 +410,7 @@ void FArchetypeData::SetComponentsData(const FLWEntity Entity, TArrayView<const
 	}
 }
 
-void FArchetypeData::SetComponentData(const FArchetypeChunkCollection& ChunkCollection, const FInstancedStruct& ComponentSource)
+void FMassArchetypeData::SetComponentData(const FArchetypeChunkCollection& ChunkCollection, const FInstancedStruct& ComponentSource)
 {
 	check(ComponentSource.IsValid());
 	const UScriptStruct* ComponentType = ComponentSource.GetScriptStruct();
@@ -431,7 +431,7 @@ void FArchetypeData::SetComponentData(const FArchetypeChunkCollection& ChunkColl
 	}
 }
 
-void FArchetypeData::SetDefaultChunkComponentValue(FConstStructView InstancedStruct)
+void FMassArchetypeData::SetDefaultChunkComponentValue(FConstStructView InstancedStruct)
 {
 	if (ensure(InstancedStruct.GetScriptStruct()) && CompositionDescriptor.ChunkComponents.Contains(*InstancedStruct.GetScriptStruct()))
 	{
@@ -442,25 +442,25 @@ void FArchetypeData::SetDefaultChunkComponentValue(FConstStructView InstancedStr
 	}
 }
 
-void FArchetypeData::MoveEntityToAnotherArchetype(const FLWEntity Entity, FArchetypeData& NewArchetype)
+void FMassArchetypeData::MoveEntityToAnotherArchetype(const FMassEntityHandle Entity, FMassArchetypeData& NewArchetype)
 {
 	check(&NewArchetype != this);
 
 	const int32 AbsoluteIndex = EntityMap.FindAndRemoveChecked(Entity.Index);
 	const int32 ChunkIndex = AbsoluteIndex / NumEntitiesPerChunk;
 	const int32 IndexWithinChunk = AbsoluteIndex % NumEntitiesPerChunk;
-	FArchetypeChunk& Chunk = Chunks[ChunkIndex];
+	FMassArchetypeChunk& Chunk = Chunks[ChunkIndex];
 
 	constexpr bool bInitializeComponentsDuringCreation = !UE::AggregateTicking::bBitwiseRelocateComponents;
 	const int32 NewAbsoluteIndex = NewArchetype.AddEntityInternal(Entity, bInitializeComponentsDuringCreation);
 	const int32 NewChunkIndex = NewAbsoluteIndex / NewArchetype.NumEntitiesPerChunk;
 	const int32 NewIndexWithinChunk = NewAbsoluteIndex % NewArchetype.NumEntitiesPerChunk;
-	FArchetypeChunk& NewChunk = NewArchetype.Chunks[NewChunkIndex];
+	FMassArchetypeChunk& NewChunk = NewArchetype.Chunks[NewChunkIndex];
 
 	// for every NewArchetype's component see if it was in the old component as well and if so copy it's value. 
 	// If not then initialize the component (unless it has already been initialized based on bInitializeComponentsDuringCreation
 	// value
-	for (const FArchetypeComponentConfig& NewComponentConfig : NewArchetype.ComponentConfigs)
+	for (const FMassArchetypeFragmentConfig& NewComponentConfig : NewArchetype.ComponentConfigs)
 	{
 		const int32* OldComponentIndex = ComponentIndexMap.Find(NewComponentConfig.ComponentType);
 		void* Dst = NewComponentConfig.GetComponentData(NewChunk.GetRawMemory(), NewIndexWithinChunk);
@@ -491,7 +491,7 @@ void FArchetypeData::MoveEntityToAnotherArchetype(const FLWEntity Entity, FArche
 	RemoveEntityInternal(AbsoluteIndex, bDestroyComponents);
 }
 
-void FArchetypeData::ExecuteFunction(FLWComponentSystemExecutionContext& RunContext, const FLWComponentSystemExecuteFunction& Function, const FLWRequirementIndicesMapping& RequirementMapping, const FArchetypeChunkCollection& ChunkCollection)
+void FMassArchetypeData::ExecuteFunction(FMassExecutionContext& RunContext, const FMassExecuteFunction& Function, const FMassQueryRequirementIndicesMapping& RequirementMapping, const FArchetypeChunkCollection& ChunkCollection)
 {
 	check(ChunkCollection.GetArchetype() == this);
 
@@ -499,7 +499,7 @@ void FArchetypeData::ExecuteFunction(FLWComponentSystemExecutionContext& RunCont
 	RunContext.SetCurrentArchetypeData(*this);
 	for (FArchetypeChunkIterator ChunkIterator(ChunkCollection); ChunkIterator; ++ChunkIterator)
 	{
-		FArchetypeChunk& Chunk = Chunks[ChunkIterator->ChunkIndex];
+		FMassArchetypeChunk& Chunk = Chunks[ChunkIterator->ChunkIndex];
 		
 		const int32 ChunkLength = ChunkIterator->Length > 0 ? ChunkIterator->Length : (Chunk.GetNumInstances() - ChunkIterator->SubchunkStart);
 		if (ChunkLength)
@@ -515,11 +515,11 @@ void FArchetypeData::ExecuteFunction(FLWComponentSystemExecutionContext& RunCont
 	}
 }
 
-void FArchetypeData::ExecuteFunction(FLWComponentSystemExecutionContext& RunContext, const FLWComponentSystemExecuteFunction& Function, const FLWRequirementIndicesMapping& RequirementMapping, const FLWComponentSystemChunkConditionFunction& ChunkCondition)
+void FMassArchetypeData::ExecuteFunction(FMassExecutionContext& RunContext, const FMassExecuteFunction& Function, const FMassQueryRequirementIndicesMapping& RequirementMapping, const FMassChunkConditionFunction& ChunkCondition)
 {
 	// mz@todo to be removed
 	RunContext.SetCurrentArchetypeData(*this);
-	for (FArchetypeChunk& Chunk : Chunks)
+	for (FMassArchetypeChunk& Chunk : Chunks)
 	{
 		if (Chunk.GetNumInstances())
 		{
@@ -535,9 +535,9 @@ void FArchetypeData::ExecuteFunction(FLWComponentSystemExecutionContext& RunCont
 	}
 }
 
-void FArchetypeData::ExecutionFunctionForChunk(FLWComponentSystemExecutionContext RunContext, const FLWComponentSystemExecuteFunction& Function, const FLWRequirementIndicesMapping& RequirementMapping, const FArchetypeChunkCollection::FChunkInfo& ChunkInfo, const FLWComponentSystemChunkConditionFunction& ChunkCondition)
+void FMassArchetypeData::ExecutionFunctionForChunk(FMassExecutionContext RunContext, const FMassExecuteFunction& Function, const FMassQueryRequirementIndicesMapping& RequirementMapping, const FArchetypeChunkCollection::FChunkInfo& ChunkInfo, const FMassChunkConditionFunction& ChunkCondition)
 {
-	FArchetypeChunk& Chunk = Chunks[ChunkInfo.ChunkIndex];
+	FMassArchetypeChunk& Chunk = Chunks[ChunkInfo.ChunkIndex];
 	const int32 ChunkLength = ChunkInfo.Length > 0 ? ChunkInfo.Length : (Chunk.GetNumInstances() - ChunkInfo.SubchunkStart);
 
 	if (ChunkLength)
@@ -554,12 +554,12 @@ void FArchetypeData::ExecutionFunctionForChunk(FLWComponentSystemExecutionContex
 	}
 }
 
-void FArchetypeData::CompactEntities(const double TimeAllowed)
+void FMassArchetypeData::CompactEntities(const double TimeAllowed)
 {
 	const double TimeAllowedEnd = FPlatformTime::Seconds() + TimeAllowed;
 
-	TArray<FArchetypeChunk*> SortedChunks;
-	for (FArchetypeChunk& Chunk : Chunks)
+	TArray<FMassArchetypeChunk*> SortedChunks;
+	for (FMassArchetypeChunk& Chunk : Chunks)
 	{
 		// Skip already full chunks
 		const int32 NumInstances = Chunk.GetNumInstances();
@@ -575,7 +575,7 @@ void FArchetypeData::CompactEntities(const double TimeAllowed)
 		return;
 	}
 
-	SortedChunks.Sort([](const FArchetypeChunk& LHS, const FArchetypeChunk& RHS)
+	SortedChunks.Sort([](const FMassArchetypeChunk& LHS, const FMassArchetypeChunk& RHS)
 	{
 		return LHS.GetNumInstances() < RHS.GetNumInstances();
 	});
@@ -597,14 +597,14 @@ void FArchetypeData::CompactEntities(const double TimeAllowed)
 			break;
 		}
 
-		FArchetypeChunk* ChunkToFill = SortedChunks[ChunkToFillSortedIdx];
-		FArchetypeChunk* ChunkToEmpty = SortedChunks[ChunkToEmptySortedIdx];
+		FMassArchetypeChunk* ChunkToFill = SortedChunks[ChunkToFillSortedIdx];
+		FMassArchetypeChunk* ChunkToEmpty = SortedChunks[ChunkToEmptySortedIdx];
 		const int32 NumberOfEntitiesToMove =  FMath::Min(NumEntitiesPerChunk-ChunkToFill->GetNumInstances(), ChunkToEmpty->GetNumInstances());
 		const int32 FromIndex = ChunkToEmpty->GetNumInstances() - NumberOfEntitiesToMove;
 		const int32 ToIndex = ChunkToFill->GetNumInstances();
 		check(NumberOfEntitiesToMove > 0);
 
-		for (const FArchetypeComponentConfig& ComponentConfig : ComponentConfigs)
+		for (const FMassArchetypeFragmentConfig& ComponentConfig : ComponentConfigs)
 		{
 			void* FromComponentPtr = ComponentConfig.GetComponentData(ChunkToEmpty->GetRawMemory(), FromIndex);
 			void* ToComponentPtr = ComponentConfig.GetComponentData(ChunkToFill->GetRawMemory(), ToIndex);
@@ -627,9 +627,9 @@ void FArchetypeData::CompactEntities(const double TimeAllowed)
 			}
 		}
 
-		FLWEntity* FromEntity = &ChunkToEmpty->GetEntityArrayElementRef(EntityListOffsetWithinChunk, FromIndex);
-		FLWEntity* ToEntity = &ChunkToFill->GetEntityArrayElementRef(EntityListOffsetWithinChunk, ToIndex);
-		FMemory::Memcpy(ToEntity, FromEntity, NumberOfEntitiesToMove * sizeof(FLWEntity));
+		FMassEntityHandle* FromEntity = &ChunkToEmpty->GetEntityArrayElementRef(EntityListOffsetWithinChunk, FromIndex);
+		FMassEntityHandle* ToEntity = &ChunkToFill->GetEntityArrayElementRef(EntityListOffsetWithinChunk, ToIndex);
+		FMemory::Memcpy(ToEntity, FromEntity, NumberOfEntitiesToMove * sizeof(FMassEntityHandle));
 		ChunkToFill->AddMultipleInstances(NumberOfEntitiesToMove);
 		ChunkToEmpty->RemoveMultipleInstances(NumberOfEntitiesToMove);
 
@@ -644,10 +644,10 @@ void FArchetypeData::CompactEntities(const double TimeAllowed)
 	}
 }
 
-void FArchetypeData::GetRequirementsComponentMapping(TConstArrayView<FLWComponentRequirement> Requirements, FLWComponentIndicesMapping& OutComponentIndices)
+void FMassArchetypeData::GetRequirementsComponentMapping(TConstArrayView<FMassFragmentRequirement> Requirements, FMassFragmentIndicesMapping& OutComponentIndices)
 {
 	OutComponentIndices.Reset(Requirements.Num());
-	for (const FLWComponentRequirement& Requirement : Requirements)
+	for (const FMassFragmentRequirement& Requirement : Requirements)
 	{
 		if (Requirement.RequiresBinding())
 		{
@@ -658,11 +658,11 @@ void FArchetypeData::GetRequirementsComponentMapping(TConstArrayView<FLWComponen
 	}
 }
 
-void FArchetypeData::GetRequirementsChunkComponentMapping(TConstArrayView<FLWComponentRequirement> ChunkRequirements, FLWComponentIndicesMapping& OutComponentIndices)
+void FMassArchetypeData::GetRequirementsChunkComponentMapping(TConstArrayView<FMassFragmentRequirement> ChunkRequirements, FMassFragmentIndicesMapping& OutComponentIndices)
 {
 	int32 LastFoundComponentIndex = -1;
 	OutComponentIndices.Reset(ChunkRequirements.Num());
-	for (const FLWComponentRequirement& Requirement : ChunkRequirements)
+	for (const FMassFragmentRequirement& Requirement : ChunkRequirements)
 	{
 		if (Requirement.RequiresBinding())
 		{
@@ -683,7 +683,7 @@ void FArchetypeData::GetRequirementsChunkComponentMapping(TConstArrayView<FLWCom
 	}
 }
 
-void FArchetypeData::BindEntityRequirements(FLWComponentSystemExecutionContext& RunContext, const FLWComponentIndicesMapping& EntityComponentsMapping, FArchetypeChunk& Chunk, const int32 SubchunkStart, const int32 SubchunkLength)
+void FMassArchetypeData::BindEntityRequirements(FMassExecutionContext& RunContext, const FMassFragmentIndicesMapping& EntityComponentsMapping, FMassArchetypeChunk& Chunk, const int32 SubchunkStart, const int32 SubchunkLength)
 {
 	// auto-correcting number of entities to process in case SubchunkStart +  SubchunkLength > Chunk.GetNumInstances()
 	const int32 NumEntities = SubchunkLength >= 0 ? FMath::Min(SubchunkLength, Chunk.GetNumInstances() - SubchunkStart) : Chunk.GetNumInstances();
@@ -695,43 +695,43 @@ void FArchetypeData::BindEntityRequirements(FLWComponentSystemExecutionContext&
 
 		for (int i = 0; i < EntityComponentsMapping.Num(); ++i)
 		{
-			FLWComponentSystemExecutionContext::FComponentView& Requirement = RunContext.ComponentViews[i];
+			FMassExecutionContext::FFragmentView& Requirement = RunContext.ComponentViews[i];
 			const int32 ComponentIndex = EntityComponentsMapping[i];
 
 			check(ComponentIndex != INDEX_NONE || Requirement.Requirement.IsOptional());
 			if (ComponentIndex != INDEX_NONE)
 			{
-				Requirement.ComponentView = TArrayView<FLWComponentData>((FLWComponentData*)GetComponentData(ComponentIndex, Chunk.GetRawMemory(), SubchunkStart), NumEntities);
+				Requirement.ComponentView = TArrayView<FMassFragment>((FMassFragment*)GetComponentData(ComponentIndex, Chunk.GetRawMemory(), SubchunkStart), NumEntities);
 			}
 			else
 			{
 				// @todo this might not be needed
-				Requirement.ComponentView = TArrayView<FLWComponentData>();
+				Requirement.ComponentView = TArrayView<FMassFragment>();
 			}
 		}
 	}
 	else
 	{
 		// Map in the required data arrays from the current chunk to the array views
-		for (FLWComponentSystemExecutionContext::FComponentView& Requirement : RunContext.GetMutableRequirements())
+		for (FMassExecutionContext::FFragmentView& Requirement : RunContext.GetMutableRequirements())
 		{
 			const int32* ComponentIndex = ComponentIndexMap.Find(Requirement.Requirement.StructType);
 			check(ComponentIndex != nullptr || Requirement.Requirement.IsOptional());
 			if (ComponentIndex)
 			{
-				Requirement.ComponentView = TArrayView<FLWComponentData>((FLWComponentData*)GetComponentData(*ComponentIndex, Chunk.GetRawMemory(), SubchunkStart), NumEntities);
+				Requirement.ComponentView = TArrayView<FMassFragment>((FMassFragment*)GetComponentData(*ComponentIndex, Chunk.GetRawMemory(), SubchunkStart), NumEntities);
 			}
 			else
 			{
-				Requirement.ComponentView = TArrayView<FLWComponentData>();
+				Requirement.ComponentView = TArrayView<FMassFragment>();
 			}
 		}
 	}
 
-	RunContext.EntityListView = TArrayView<FLWEntity>(&Chunk.GetEntityArrayElementRef(EntityListOffsetWithinChunk, SubchunkStart), NumEntities);
+	RunContext.EntityListView = TArrayView<FMassEntityHandle>(&Chunk.GetEntityArrayElementRef(EntityListOffsetWithinChunk, SubchunkStart), NumEntities);
 }
 
-void FArchetypeData::BindChunkComponentRequirements(FLWComponentSystemExecutionContext& RunContext, const FLWComponentIndicesMapping& ChunkComponentsMapping, FArchetypeChunk& Chunk)
+void FMassArchetypeData::BindChunkComponentRequirements(FMassExecutionContext& RunContext, const FMassFragmentIndicesMapping& ChunkComponentsMapping, FMassArchetypeChunk& Chunk)
 {
 	if (ChunkComponentsMapping.Num() > 0)
 	{
@@ -739,7 +739,7 @@ void FArchetypeData::BindChunkComponentRequirements(FLWComponentSystemExecutionC
 
 		for (int i = 0; i < ChunkComponentsMapping.Num(); ++i)
 		{
-			FLWComponentSystemExecutionContext::FChunkComponentView& ChunkRequirement = RunContext.ChunkComponents[i];
+			FMassExecutionContext::FChunkFragmentView& ChunkRequirement = RunContext.ChunkComponents[i];
 			const int32 ChunkComponentIndex = ChunkComponentsMapping[i];
 
 			check(ChunkComponentIndex != INDEX_NONE);
@@ -748,7 +748,7 @@ void FArchetypeData::BindChunkComponentRequirements(FLWComponentSystemExecutionC
 	}
 	else
 	{
-		for (FLWComponentSystemExecutionContext::FChunkComponentView& ChunkRequirement : RunContext.GetMutableChunkRequirements())
+		for (FMassExecutionContext::FChunkFragmentView& ChunkRequirement : RunContext.GetMutableChunkRequirements())
 		{
 			FInstancedStruct* ChunkComponentInstance = Chunk.FindMutableChunkComponent(ChunkRequirement.Requirement.StructType);
 			check(ChunkComponentInstance != nullptr);
@@ -757,9 +757,9 @@ void FArchetypeData::BindChunkComponentRequirements(FLWComponentSystemExecutionC
 	}
 }
 
-SIZE_T FArchetypeData::GetAllocatedSize() const
+SIZE_T FMassArchetypeData::GetAllocatedSize() const
 {
-	return sizeof(FArchetypeData) +
+	return sizeof(FMassArchetypeData) +
 		ComponentConfigs.GetAllocatedSize() +
 		Chunks.GetAllocatedSize() +
 		(Chunks.Num() * GetChunkAllocSize()) +
@@ -767,7 +767,7 @@ SIZE_T FArchetypeData::GetAllocatedSize() const
 		ComponentIndexMap.GetAllocatedSize();
 }
 
-FString FArchetypeData::DebugGetDescription() const
+FString FMassArchetypeData::DebugGetDescription() const
 {
 #if WITH_MASSENTITY_DEBUG
 	TStringBuilder<256> ArchetypeDebugName;
@@ -775,7 +775,7 @@ FString FArchetypeData::DebugGetDescription() const
 	ArchetypeDebugName.Append(TEXT("<"));
 
 	bool bNeedsComma = false;
-	for (const FArchetypeComponentConfig& ComponentConfig : ComponentConfigs)
+	for (const FMassArchetypeFragmentConfig& ComponentConfig : ComponentConfigs)
 	{
 		if (bNeedsComma)
 		{
@@ -793,7 +793,7 @@ FString FArchetypeData::DebugGetDescription() const
 }
 
 #if WITH_MASSENTITY_DEBUG
-void FArchetypeData::DebugPrintArchetype(FOutputDevice& Ar)
+void FMassArchetypeData::DebugPrintArchetype(FOutputDevice& Ar)
 {
 	FStringOutputDevice TagsDecription;
 	CompositionDescriptor.Tags.DebugGetStringDesc(TagsDecription);
@@ -802,7 +802,7 @@ void FArchetypeData::DebugPrintArchetype(FOutputDevice& Ar)
 	Ar.Logf(ELogVerbosity::Log, TEXT("\tChunks: %d x %d KB = %d KB total"), Chunks.Num(), GetChunkAllocSize() / 1024, (GetChunkAllocSize()*Chunks.Num()) / 1024);
 	
 	int ChunkWithComponentsCount = 0;
-	for (FArchetypeChunk& Chunk : Chunks)
+	for (FMassArchetypeChunk& Chunk : Chunks)
 	{
 		ChunkWithComponentsCount += Chunk.DebugGetChunkComponentCount() > 0 ? 1 : 0;
 	}
@@ -835,9 +835,9 @@ void FArchetypeData::DebugPrintArchetype(FOutputDevice& Ar)
 	Ar.Logf(ELogVerbosity::Log, TEXT("\tBytes / Entity  : %d"), TotalBytesPerEntity);
 	Ar.Logf(ELogVerbosity::Log, TEXT("\tEntities / Chunk: %d"), NumEntitiesPerChunk);
 
-	Ar.Logf(ELogVerbosity::Log, TEXT("\tOffset 0x%04X: Entity[] (%d bytes each)"), EntityListOffsetWithinChunk, sizeof(FLWEntity));
-	int32 TotalBytesOfValidData = sizeof(FLWEntity) * NumEntitiesPerChunk;
-	for (const FArchetypeComponentConfig& ComponentConfig : ComponentConfigs)
+	Ar.Logf(ELogVerbosity::Log, TEXT("\tOffset 0x%04X: Entity[] (%d bytes each)"), EntityListOffsetWithinChunk, sizeof(FMassEntityHandle));
+	int32 TotalBytesOfValidData = sizeof(FMassEntityHandle) * NumEntitiesPerChunk;
+	for (const FMassArchetypeFragmentConfig& ComponentConfig : ComponentConfigs)
 	{
 		TotalBytesOfValidData += ComponentConfig.ComponentType->GetStructureSize() * NumEntitiesPerChunk;
 		Ar.Logf(ELogVerbosity::Log, TEXT("\tOffset 0x%04X: %s[] (%d bytes each)"), ComponentConfig.ArrayOffsetWithinChunk, *ComponentConfig.ComponentType->GetName(), ComponentConfig.ComponentType->GetStructureSize());
@@ -858,9 +858,9 @@ void FArchetypeData::DebugPrintArchetype(FOutputDevice& Ar)
 	}
 }
 
-void FArchetypeData::DebugPrintEntity(FLWEntity Entity, FOutputDevice& Ar, const TCHAR* InPrefix) const
+void FMassArchetypeData::DebugPrintEntity(FMassEntityHandle Entity, FOutputDevice& Ar, const TCHAR* InPrefix) const
 {
-	for (const FArchetypeComponentConfig& ComponentConfig : ComponentConfigs)
+	for (const FMassArchetypeFragmentConfig& ComponentConfig : ComponentConfigs)
 	{
 		void* Data = GetComponentDataForEntityChecked(ComponentConfig.ComponentType, Entity.Index);
 		
@@ -876,9 +876,9 @@ void FArchetypeData::DebugPrintEntity(FLWEntity Entity, FOutputDevice& Ar, const
 
 #endif // WITH_MASSENTITY_DEBUG
 
-void FArchetypeData::REMOVEME_GetArrayViewForComponentInChunk(int32 ChunkIndex, const UScriptStruct* ComponentType, void*& OutChunkBase, int32& OutNumEntities)
+void FMassArchetypeData::REMOVEME_GetArrayViewForComponentInChunk(int32 ChunkIndex, const UScriptStruct* ComponentType, void*& OutChunkBase, int32& OutNumEntities)
 {
-	const FArchetypeChunk& Chunk = Chunks[ChunkIndex];
+	const FMassArchetypeChunk& Chunk = Chunks[ChunkIndex];
 	const int32 ComponentIndex = ComponentIndexMap.FindChecked(ComponentType);
 
 	OutChunkBase = ComponentConfigs[ComponentIndex].GetComponentData(Chunk.GetRawMemory(), 0);
diff --git a/Engine/Plugins/Runtime/MassEntity/Source/MassEntity/Private/ArchetypeData.h b/Engine/Plugins/Runtime/MassEntity/Source/MassEntity/Private/ArchetypeData.h
--- a/Engine/Plugins/Runtime/MassEntity/Source/MassEntity/Private/ArchetypeData.h
+++ b/Engine/Plugins/Runtime/MassEntity/Source/MassEntity/Private/ArchetypeData.h
@@ -5,13 +5,13 @@
 #include "MassEntitySubsystem.h"
 #include "ArchetypeTypes.h"
 
-struct FLWComponentQuery;
-struct FLWComponentSystemExecutionContext;
+struct FMassEntityQuery;
+struct FMassExecutionContext;
 class FOutputDevice;
 struct FArchetypeChunkCollection;
 
 // This is one chunk within an archetype
-struct FArchetypeChunk
+struct FMassArchetypeChunk
 {
 private:
 	uint8* RawMemory = nullptr;
@@ -20,21 +20,21 @@ private:
 	TArray<FInstancedStruct> ChunkComponentData;
 
 public:
-	explicit FArchetypeChunk(int32 AllocSize, TConstArrayView<FInstancedStruct> InChunkComponentTemplates)
+	explicit FMassArchetypeChunk(int32 AllocSize, TConstArrayView<FInstancedStruct> InChunkComponentTemplates)
 		: ChunkComponentData(InChunkComponentTemplates)
 	{
 		RawMemory = (uint8*)FMemory::Malloc(AllocSize);
 	}
 
-	~FArchetypeChunk()
+	~FMassArchetypeChunk()
 	{
 		FMemory::Free(RawMemory);
 	}
 
 	// Returns the Entity array element at the specified index
-	FLWEntity& GetEntityArrayElementRef(int32 ChunkBase, int32 IndexWithinChunk)
+	FMassEntityHandle& GetEntityArrayElementRef(int32 ChunkBase, int32 IndexWithinChunk)
 	{
-		return ((FLWEntity*)(RawMemory + ChunkBase))[IndexWithinChunk];
+		return ((FMassEntityHandle*)(RawMemory + ChunkBase))[IndexWithinChunk];
 	}
 
 	uint8* GetRawMemory() const
@@ -98,7 +98,7 @@ public:
 };
 
 // Information for a single component type in an archetype
-struct FArchetypeComponentConfig
+struct FMassArchetypeFragmentConfig
 {
 	const UScriptStruct* ComponentType = nullptr;
 	int32 ArrayOffsetWithinChunk = 0;
@@ -110,18 +110,18 @@ struct FArchetypeComponentConfig
 };
 
 // An archetype is defined by a collection of unique component types (no duplicates).
-// Order doesn't matter, there will only ever be one FArchetypeData per unique set of component types per entity manager subsystem
-struct FArchetypeData
+// Order doesn't matter, there will only ever be one FMassArchetypeData per unique set of component types per entity manager subsystem
+struct FMassArchetypeData
 {
 private:
 	// One-stop-shop variable describing the archetype's component and tag composition 
-	FLWCompositionDescriptor CompositionDescriptor;
+	FMassCompositionDescriptor CompositionDescriptor;
 
 	TArray<FInstancedStruct> ChunkComponentTemplates;
 
-	TArray<FArchetypeComponentConfig, TInlineAllocator<16>> ComponentConfigs;
+	TArray<FMassArchetypeFragmentConfig, TInlineAllocator<16>> ComponentConfigs;
 	
-	TArray<FArchetypeChunk> Chunks;
+	TArray<FMassArchetypeChunk> Chunks;
 
 	// Entity ID to index within archetype
 	//@TODO: Could be folded into FEntityData in the entity manager at the expense of a bit
@@ -134,41 +134,41 @@ private:
 	int32 TotalBytesPerEntity;
 	int32 EntityListOffsetWithinChunk;
 
-	friend FLWComponentQuery;
+	friend FMassEntityQuery;
 	friend FArchetypeChunkCollection;
 
 public:
-	TConstArrayView<FArchetypeComponentConfig> GetComponentConfigs() const { return ComponentConfigs; }
-	const FLWComponentBitSet& GetComponentBitSet() const { return CompositionDescriptor.Components; }
-	const FLWTagBitSet& GetTagBitSet() const { return CompositionDescriptor.Tags; }
-	const FLWChunkComponentBitSet& GetChunkComponentBitSet() const { return CompositionDescriptor.ChunkComponents; }
-	const FLWCompositionDescriptor& GetCompositionDescriptor() const { return CompositionDescriptor; }
+	TConstArrayView<FMassArchetypeFragmentConfig> GetComponentConfigs() const { return ComponentConfigs; }
+	const FMassFragmentBitSet& GetComponentBitSet() const { return CompositionDescriptor.Components; }
+	const FMassTagBitSet& GetTagBitSet() const { return CompositionDescriptor.Tags; }
+	const FMassChunkFragmentBitSet& GetChunkComponentBitSet() const { return CompositionDescriptor.ChunkComponents; }
+	const FMassCompositionDescriptor& GetCompositionDescriptor() const { return CompositionDescriptor; }
 
 	/** Method to iterate on all the component types */
 	void ForEachComponentType(TFunction< void(const UScriptStruct* /*ComponentType*/)> Function) const;
 	bool HasComponentType(const UScriptStruct* ComponentType) const;
 	bool HasTagType(const UScriptStruct* ComponentType) const { check(ComponentType); return CompositionDescriptor.Tags.Contains(*ComponentType); }
 	bool IsEquivalent(TConstArrayView<const UScriptStruct*> OtherList) const;
-	bool IsEquivalent(const FLWComponentBitSet& InComponentBitSet, const FLWTagBitSet& InTagBitSet, const FLWChunkComponentBitSet& InChunkComponentsBitSet) const
+	bool IsEquivalent(const FMassFragmentBitSet& InComponentBitSet, const FMassTagBitSet& InTagBitSet, const FMassChunkFragmentBitSet& InChunkComponentsBitSet) const
 	{ 
 		return CompositionDescriptor.IsEquivalent(InComponentBitSet, InTagBitSet, InChunkComponentsBitSet);
 	}
-	bool IsEquivalent(const FLWCompositionDescriptor& OtherCompositionDescriptor) const
+	bool IsEquivalent(const FMassCompositionDescriptor& OtherCompositionDescriptor) const
 	{
 		return CompositionDescriptor.IsEquivalent(OtherCompositionDescriptor);
 	}
 
-	void Initialize(const FLWComponentBitSet& Components, const FLWTagBitSet& Tags, const FLWChunkComponentBitSet& ChunkComponents);
+	void Initialize(const FMassFragmentBitSet& Components, const FMassTagBitSet& Tags, const FMassChunkFragmentBitSet& ChunkComponents);
 
 	/** 
 	 * A special way of initializing an archetype resulting in a copy of SiblingArchetype's setup with OverrideTags
 	 * replacing original tags of SiblingArchetype
 	 */
-	void InitializeWithSibling(const FArchetypeData& SiblingArchetype, const FLWTagBitSet& OverrideTags);
+	void InitializeWithSibling(const FMassArchetypeData& SiblingArchetype, const FMassTagBitSet& OverrideTags);
 
-	void AddEntity(FLWEntity Entity);
-	void RemoveEntity(FLWEntity Entity);
-	void BatchDestroyEntityChunks(const FArchetypeChunkCollection& ChunkCollection, TArray<FLWEntity>& OutEntitiesRemoved);
+	void AddEntity(FMassEntityHandle Entity);
+	void RemoveEntity(FMassEntityHandle Entity);
+	void BatchDestroyEntityChunks(const FArchetypeChunkCollection& ChunkCollection, TArray<FMassEntityHandle>& OutEntitiesRemoved);
 
 	bool HasComponentDataForEntity(const UScriptStruct* ComponentType, int32 EntityIndex) const;
 	void* GetComponentDataForEntityChecked(const UScriptStruct* ComponentType, int32 EntityIndex) const;
@@ -183,10 +183,10 @@ public:
 
 	int32 GetChunkCount() const { return Chunks.Num(); }
 
-	void ExecuteFunction(FLWComponentSystemExecutionContext& RunContext, const FLWComponentSystemExecuteFunction& Function, const FLWRequirementIndicesMapping& RequirementMapping, const FArchetypeChunkCollection& ChunkCollection);
-	void ExecuteFunction(FLWComponentSystemExecutionContext& RunContext, const FLWComponentSystemExecuteFunction& Function, const FLWRequirementIndicesMapping& RequirementMapping, const FLWComponentSystemChunkConditionFunction& ChunkCondition = FLWComponentSystemChunkConditionFunction());
+	void ExecuteFunction(FMassExecutionContext& RunContext, const FMassExecuteFunction& Function, const FMassQueryRequirementIndicesMapping& RequirementMapping, const FArchetypeChunkCollection& ChunkCollection);
+	void ExecuteFunction(FMassExecutionContext& RunContext, const FMassExecuteFunction& Function, const FMassQueryRequirementIndicesMapping& RequirementMapping, const FMassChunkConditionFunction& ChunkCondition = FMassChunkConditionFunction());
 
-	void ExecutionFunctionForChunk(FLWComponentSystemExecutionContext RunContext, const FLWComponentSystemExecuteFunction& Function, const FLWRequirementIndicesMapping& RequirementMapping, const FArchetypeChunkCollection::FChunkInfo& ChunkInfo, const FLWComponentSystemChunkConditionFunction& ChunkCondition = FLWComponentSystemChunkConditionFunction());
+	void ExecutionFunctionForChunk(FMassExecutionContext RunContext, const FMassExecuteFunction& Function, const FMassQueryRequirementIndicesMapping& RequirementMapping, const FArchetypeChunkCollection::FChunkInfo& ChunkInfo, const FMassChunkConditionFunction& ChunkCondition = FMassChunkConditionFunction());
 
 	/**
 	 * Compacts entities to fill up chunks as much as possible
@@ -198,14 +198,14 @@ public:
 	 * @param Entity is the entity to move
 	 * @param NewArchetype the archetype to move to
 	 */
-	void MoveEntityToAnotherArchetype(const FLWEntity Entity, FArchetypeData& NewArchetype);
+	void MoveEntityToAnotherArchetype(const FMassEntityHandle Entity, FMassArchetypeData& NewArchetype);
 
 	/**
 	 * Set all component sources data on specified entity, will check if there are component sources type that does not exist in the archetype
 	 * @param Entity is the entity to set the data of all components
 	 * @param ComponentSources are the components to copy the data from
 	 */
-	void SetComponentsData(const FLWEntity Entity, TArrayView<const FInstancedStruct> ComponentSources);
+	void SetComponentsData(const FMassEntityHandle Entity, TArrayView<const FInstancedStruct> ComponentSources);
 
 	/** For all entities indicated by ChunkCollection the function sets the value of component of type
 	 *  ComponentSource.GetScriptStruct to the value represented by ComponentSource.GetMemory */
@@ -214,10 +214,10 @@ public:
 	void SetDefaultChunkComponentValue(FConstStructView InstancedStruct);
 
 	/** Returns conversion from given Requirements to archetype's component indices */
-	void GetRequirementsComponentMapping(TConstArrayView<FLWComponentRequirement> Requirements, FLWComponentIndicesMapping& OutComponentIndices);
+	void GetRequirementsComponentMapping(TConstArrayView<FMassFragmentRequirement> Requirements, FMassFragmentIndicesMapping& OutComponentIndices);
 
 	/** Returns conversion from given ChunkRequirements to archetype's chunk component indices */
-	void GetRequirementsChunkComponentMapping(TConstArrayView<FLWComponentRequirement> ChunkRequirements, FLWComponentIndicesMapping& OutComponentIndices);
+	void GetRequirementsChunkComponentMapping(TConstArrayView<FMassFragmentRequirement> ChunkRequirements, FMassFragmentIndicesMapping& OutComponentIndices);
 
 	SIZE_T GetAllocatedSize() const;
 
@@ -236,7 +236,7 @@ public:
 	 * @param Ar The output device
 	 * @param InPrefix Optional prefix to remove from component names
 	 */
-	void DebugPrintEntity(FLWEntity Entity, FOutputDevice& Ar, const TCHAR* InPrefix = TEXT("")) const;
+	void DebugPrintEntity(FMassEntityHandle Entity, FOutputDevice& Ar, const TCHAR* InPrefix = TEXT("")) const;
 #endif // WITH_MASSENTITY_DEBUG
 
 	void REMOVEME_GetArrayViewForComponentInChunk(int32 ChunkIndex, const UScriptStruct* ComponentType, void*& OutChunkBase, int32& OutNumEntities);
@@ -259,7 +259,7 @@ public:
 		return FInternalEntityHandle(Chunks[ChunkIndex].GetRawMemory(), AbsoluteIndex % NumEntitiesPerChunk); 
 	}
 
-	FORCEINLINE FInternalEntityHandle MakeEntityHandle(FLWEntity Entity) const
+	FORCEINLINE FInternalEntityHandle MakeEntityHandle(FMassEntityHandle Entity) const
 	{
 		return MakeEntityHandle(Entity.Index); 
 	}
@@ -272,10 +272,10 @@ protected:
 		return ComponentConfigs[ComponentIndex].GetComponentData(ChunkRawMemory, IndexWithinChunk);
 	}
 
-	void BindEntityRequirements(FLWComponentSystemExecutionContext& RunContext, const FLWComponentIndicesMapping& EntityComponentsMapping, FArchetypeChunk& Chunk, const int32 SubchunkStart, const int32 SubchunkLength);
-	void BindChunkComponentRequirements(FLWComponentSystemExecutionContext& RunContext, const FLWComponentIndicesMapping& ChunkComponentsMapping, FArchetypeChunk& Chunk);
+	void BindEntityRequirements(FMassExecutionContext& RunContext, const FMassFragmentIndicesMapping& EntityComponentsMapping, FMassArchetypeChunk& Chunk, const int32 SubchunkStart, const int32 SubchunkLength);
+	void BindChunkComponentRequirements(FMassExecutionContext& RunContext, const FMassFragmentIndicesMapping& ChunkComponentsMapping, FMassArchetypeChunk& Chunk);
 
 private:
-	int32 AddEntityInternal(FLWEntity Entity, const bool bInitializeComponents);
+	int32 AddEntityInternal(FMassEntityHandle Entity, const bool bInitializeComponents);
 	void RemoveEntityInternal(const int32 AbsoluteIndex, const bool bDestroyComponents);
 };
diff --git a/Engine/Plugins/Runtime/MassEntity/Source/MassEntity/Private/ArchetypeTypes.cpp b/Engine/Plugins/Runtime/MassEntity/Source/MassEntity/Private/ArchetypeTypes.cpp
--- a/Engine/Plugins/Runtime/MassEntity/Source/MassEntity/Private/ArchetypeTypes.cpp
+++ b/Engine/Plugins/Runtime/MassEntity/Source/MassEntity/Private/ArchetypeTypes.cpp
@@ -8,7 +8,7 @@
 
 //////////////////////////////////////////////////////////////////////
 // FArchetypeHandle
-bool FArchetypeHandle::operator==(const FArchetypeData* Other) const
+bool FArchetypeHandle::operator==(const FMassArchetypeData* Other) const
 {
 	return DataPtr.Get() == Other;
 }
@@ -21,18 +21,18 @@ uint32 GetTypeHash(const FArchetypeHandle& Instance)
 //////////////////////////////////////////////////////////////////////
 // FArchetypeChunkCollection
 
-FArchetypeChunkCollection::FArchetypeChunkCollection(const FArchetypeHandle& InArchetype, TConstArrayView<FLWEntity> InEntities)
+FArchetypeChunkCollection::FArchetypeChunkCollection(const FArchetypeHandle& InArchetype, TConstArrayView<FMassEntityHandle> InEntities)
 	: Archetype(InArchetype)
 {
 	check(InArchetype.IsValid());
-	const FArchetypeData& ArchetypeData = *InArchetype.DataPtr.Get();
+	const FMassArchetypeData& ArchetypeData = *InArchetype.DataPtr.Get();
 	const int32 NumEntitiesPerChunk = ArchetypeData.GetNumEntitiesPerChunk();
 
 	// InEntities has a real chance of not being sorted by AbsoluteIndex. We gotta fix that to optimize how we process the data 
 	TArray<int32> TrueIndices;
 	TrueIndices.AddUninitialized(InEntities.Num());
 	int32 i = 0;
-	for (const FLWEntity& Entity : InEntities)
+	for (const FMassEntityHandle& Entity : InEntities)
 	{
 		TrueIndices[i] = ArchetypeData.GetInternalIndexForEntity(Entity.Index);
 		++i;
@@ -73,12 +73,12 @@ FArchetypeChunkCollection::FArchetypeChunkCollection(FArchetypeHandle& InArchety
 	GatherChunksFromArchetype(InArchetypeHandle.DataPtr);
 }
 
-FArchetypeChunkCollection::FArchetypeChunkCollection(TSharedPtr<FArchetypeData>& InArchetype)
+FArchetypeChunkCollection::FArchetypeChunkCollection(TSharedPtr<FMassArchetypeData>& InArchetype)
 {
 	GatherChunksFromArchetype(InArchetype);
 }
 
-void FArchetypeChunkCollection::GatherChunksFromArchetype(TSharedPtr<FArchetypeData>& InArchetype)
+void FArchetypeChunkCollection::GatherChunksFromArchetype(TSharedPtr<FMassArchetypeData>& InArchetype)
 {
 	check(InArchetype.IsValid());
 	Archetype.DataPtr = InArchetype;
diff --git a/Engine/Plugins/Runtime/MassEntity/Source/MassEntity/Private/EntityQuery.cpp b/Engine/Plugins/Runtime/MassEntity/Source/MassEntity/Private/EntityQuery.cpp
--- a/Engine/Plugins/Runtime/MassEntity/Source/MassEntity/Private/EntityQuery.cpp
+++ b/Engine/Plugins/Runtime/MassEntity/Source/MassEntity/Private/EntityQuery.cpp
@@ -8,33 +8,33 @@
 #include "Async/ParallelFor.h"
 
 //////////////////////////////////////////////////////////////////////
-// FLWComponentQuery
+// FMassEntityQuery
 
 
-FLWComponentQuery::FLWComponentQuery()
+FMassEntityQuery::FMassEntityQuery()
 {
 	ReadCommandlineParams();
 }
 
-FLWComponentQuery::FLWComponentQuery(std::initializer_list<UScriptStruct*> InitList)
-	: FLWComponentQuery()
+FMassEntityQuery::FMassEntityQuery(std::initializer_list<UScriptStruct*> InitList)
+	: FMassEntityQuery()
 {
 	for (const UScriptStruct* ComponentType : InitList)
 	{
-		AddRequirement(ComponentType, ELWComponentAccess::ReadWrite, ELWComponentPresence::All);
+		AddRequirement(ComponentType, EMassFragmentAccess::ReadWrite, EMassFragmentPresence::All);
 	}
 }
 
-FLWComponentQuery::FLWComponentQuery(TConstArrayView<const UScriptStruct*> InitList)
-	: FLWComponentQuery()
+FMassEntityQuery::FMassEntityQuery(TConstArrayView<const UScriptStruct*> InitList)
+	: FMassEntityQuery()
 {
 	for (const UScriptStruct* ComponentType : InitList)
 	{
-		AddRequirement(ComponentType, ELWComponentAccess::ReadWrite, ELWComponentPresence::All);
+		AddRequirement(ComponentType, EMassFragmentAccess::ReadWrite, EMassFragmentPresence::All);
 	}
 }
 
-void FLWComponentQuery::ReadCommandlineParams()
+void FMassEntityQuery::ReadCommandlineParams()
 {
 	int AllowParallelQueries = -1;
 	if (FParse::Value(FCommandLine::Get(), TEXT("ParallelMassQueries="), AllowParallelQueries))
@@ -43,17 +43,17 @@ void FLWComponentQuery::ReadCommandlineParams()
 	}
 }
 
-void FLWComponentQuery::SortRequirements()
+void FMassEntityQuery::SortRequirements()
 {
-	// we're sorting the Requirements the same way ArchetypeData's ComponentConfig is sorted (see FArchetypeData::Initialize)
-	// so that when we access ArchetypeData.ComponentConfigs in FArchetypeData::BindRequirementsWithMapping
+	// we're sorting the Requirements the same way ArchetypeData's ComponentConfig is sorted (see FMassArchetypeData::Initialize)
+	// so that when we access ArchetypeData.ComponentConfigs in FMassArchetypeData::BindRequirementsWithMapping
 	// (via GetComponentData call) the access is sequential (i.e. not random) and there's a higher chance the memory
 	// ComponentConfigs we want to access have already been fetched and are available in processor cache.
-	Requirements.Sort(FLWComponentSorterOperator<FLWComponentRequirement>());
-	ChunkRequirements.Sort(FLWComponentSorterOperator<FLWComponentRequirement>());
+	Requirements.Sort(FMassSorterOperator<FMassFragmentRequirement>());
+	ChunkRequirements.Sort(FMassSorterOperator<FMassFragmentRequirement>());
 }
 
-void FLWComponentQuery::CacheArchetypes(UMassEntitySubsystem& InEntitySubsystem)
+void FMassEntityQuery::CacheArchetypes(UMassEntitySubsystem& InEntitySubsystem)
 {
 	const uint32 InEntitySubsystemHash = PointerHash(&InEntitySubsystem);
 	if (EntitySubsystemHash != InEntitySubsystemHash || InEntitySubsystem.GetArchetypeDataVersion() != ArchetypeDataVersion)
@@ -68,7 +68,7 @@ void FLWComponentQuery::CacheArchetypes(UMassEntitySubsystem& InEntitySubsystem)
 			ArchetypeDataVersion = InEntitySubsystem.GetArchetypeDataVersion();
 
 			TRACE_CPUPROFILER_EVENT_SCOPE_STR("Pipe RequirementsBinding")
-			const TConstArrayView<FLWComponentRequirement> LocalRequirements = GetRequirements();
+			const TConstArrayView<FMassFragmentRequirement> LocalRequirements = GetRequirements();
 			ArchetypeComponentMapping.Reset(ValidArchetypes.Num());
 			ArchetypeComponentMapping.AddDefaulted(ValidArchetypes.Num());
 			for (int i = 0; i < ValidArchetypes.Num(); ++i)
@@ -82,23 +82,23 @@ void FLWComponentQuery::CacheArchetypes(UMassEntitySubsystem& InEntitySubsystem)
 		}
 		else
 		{
-			UE_VLOG_UELOG(&InEntitySubsystem, LogAggregateTicking, Error, TEXT("FLWComponentQuery::CacheArchetypes: requirements not valid: %s"), *DebugGetDescription());
+			UE_VLOG_UELOG(&InEntitySubsystem, LogAggregateTicking, Error, TEXT("FMassEntityQuery::CacheArchetypes: requirements not valid: %s"), *DebugGetDescription());
 		}
 	}
 }
 
-bool FLWComponentQuery::CheckValidity() const
+bool FMassEntityQuery::CheckValidity() const
 {
 	return RequiredAllComponents.IsEmpty() == false || RequiredAnyComponents.IsEmpty() == false || RequiredOptionalComponents.IsEmpty() == false;
 }
 
-bool FLWComponentQuery::DoesArchetypeMatchRequirements(const FArchetypeHandle& ArchetypeHandle) const
+bool FMassEntityQuery::DoesArchetypeMatchRequirements(const FArchetypeHandle& ArchetypeHandle) const
 {
 	check(ArchetypeHandle.IsValid());
-	const FArchetypeData* Archetype = ArchetypeHandle.DataPtr.Get();
+	const FMassArchetypeData* Archetype = ArchetypeHandle.DataPtr.Get();
 	CA_ASSUME(Archetype);
 
-	for (const FLWComponentRequirement& Requirement : GetRequirements())
+	for (const FMassFragmentRequirement& Requirement : GetRequirements())
 	{
 		check(Requirement.StructType);
 		if (Requirement.IsOptional())
@@ -107,7 +107,7 @@ bool FLWComponentQuery::DoesArchetypeMatchRequirements(const FArchetypeHandle& A
 			continue;
 		}
 
-		const bool bNeedsComponent = (Requirement.Presence == ELWComponentPresence::All);
+		const bool bNeedsComponent = (Requirement.Presence == EMassFragmentPresence::All);
 		const bool bHasComponent = Archetype->HasComponentType(Requirement.StructType);
 		if (bNeedsComponent != bHasComponent)
 		{
@@ -117,7 +117,7 @@ bool FLWComponentQuery::DoesArchetypeMatchRequirements(const FArchetypeHandle& A
 	return true;
 }
 
-void FLWComponentQuery::ForEachEntityChunk(const FArchetypeChunkCollection& Chunks, UMassEntitySubsystem& EntitySubsystem, FLWComponentSystemExecutionContext& ExecutionContext, const FLWComponentSystemExecuteFunction& ExecuteFunction)
+void FMassEntityQuery::ForEachEntityChunk(const FArchetypeChunkCollection& Chunks, UMassEntitySubsystem& EntitySubsystem, FMassExecutionContext& ExecutionContext, const FMassExecuteFunction& ExecuteFunction)
 {
 	// mz@todo I don't like that we're copying data here.
 	ExecutionContext.SetChunkCollection(Chunks);
@@ -125,7 +125,7 @@ void FLWComponentQuery::ForEachEntityChunk(const FArchetypeChunkCollection& Chun
 	ExecutionContext.ClearChunkCollection();
 }
 
-void FLWComponentQuery::ForEachEntityChunk(UMassEntitySubsystem& EntitySubsystem, FLWComponentSystemExecutionContext& ExecutionContext, const FLWComponentSystemExecuteFunction& ExecuteFunction)
+void FMassEntityQuery::ForEachEntityChunk(UMassEntitySubsystem& EntitySubsystem, FMassExecutionContext& ExecutionContext, const FMassExecuteFunction& ExecuteFunction)
 {
 #if WITH_MASSENTITY_DEBUG
 	int32 NumEntitiesToProcess = 0;
@@ -138,7 +138,7 @@ void FLWComponentQuery::ForEachEntityChunk(UMassEntitySubsystem& EntitySubsystem
 		if (DoesArchetypeMatchRequirements(ExecutionContext.GetChunkCollection().GetArchetype()) == false)
 		{
 			// mz@todo add a unit test for this message
-			UE_VLOG_UELOG(&EntitySubsystem, LogAggregateTicking, Error, TEXT("Attempted to execute FLWComponentQuery with an incompatible Archetype"));
+			UE_VLOG_UELOG(&EntitySubsystem, LogAggregateTicking, Error, TEXT("Attempted to execute FMassEntityQuery with an incompatible Archetype"));
 			return;
 		}
 		ExecutionContext.SetRequirements(Requirements, ChunkRequirements);
@@ -175,7 +175,7 @@ void FLWComponentQuery::ForEachEntityChunk(UMassEntitySubsystem& EntitySubsystem
 	ExecutionContext.FlushDeferred(EntitySubsystem);
 }
 
-void FLWComponentQuery::ParallelForEachEntityChunk(UMassEntitySubsystem& EntitySubsystem, FLWComponentSystemExecutionContext& ExecutionContext, const FLWComponentSystemExecuteFunction& ExecuteFunction)
+void FMassEntityQuery::ParallelForEachEntityChunk(UMassEntitySubsystem& EntitySubsystem, FMassExecutionContext& ExecutionContext, const FMassExecuteFunction& ExecuteFunction)
 {
 	if (bAllowParallelExecution == false)
 	{
@@ -185,7 +185,7 @@ void FLWComponentQuery::ParallelForEachEntityChunk(UMassEntitySubsystem& EntityS
 
 	struct FChunkJob
 	{
-		FArchetypeData& Archetype;
+		FMassArchetypeData& Archetype;
 		const int32 ArchetypeIndex;
 		const FArchetypeChunkCollection::FChunkInfo ChunkInfo;
 	};
@@ -199,13 +199,13 @@ void FLWComponentQuery::ParallelForEachEntityChunk(UMassEntitySubsystem& EntityS
 		if (DoesArchetypeMatchRequirements(ArchetypeHandle) == false)
 		{
 			// mz@todo add a unit test for this message
-			UE_VLOG_UELOG(&EntitySubsystem, LogAggregateTicking, Error, TEXT("Attempted to execute FLWComponentQuery with an incompatible Archetype"));
+			UE_VLOG_UELOG(&EntitySubsystem, LogAggregateTicking, Error, TEXT("Attempted to execute FMassEntityQuery with an incompatible Archetype"));
 			return;
 		}
 
 		ExecutionContext.SetRequirements(Requirements, ChunkRequirements);
 		check(ArchetypeHandle.IsValid());
-		FArchetypeData& ArchetypeRef = *ArchetypeHandle.DataPtr.Get();
+		FMassArchetypeData& ArchetypeRef = *ArchetypeHandle.DataPtr.Get();
 		const FArchetypeChunkCollection AsChunkCollection(ArchetypeHandle.DataPtr);
 		for (const FArchetypeChunkCollection::FChunkInfo& ChunkInfo : AsChunkCollection.GetChunks())
 		{
@@ -220,7 +220,7 @@ void FLWComponentQuery::ParallelForEachEntityChunk(UMassEntitySubsystem& EntityS
 		{
 			FArchetypeHandle& Archetype = ValidArchetypes[ArchetypeIndex];
 			check(Archetype.IsValid());
-			FArchetypeData& ArchetypeRef = *Archetype.DataPtr.Get();
+			FMassArchetypeData& ArchetypeRef = *Archetype.DataPtr.Get();
 			const FArchetypeChunkCollection AsChunkCollection(Archetype.DataPtr);
 			for (const FArchetypeChunkCollection::FChunkInfo& ChunkInfo : AsChunkCollection.GetChunks())
 			{
@@ -232,7 +232,7 @@ void FLWComponentQuery::ParallelForEachEntityChunk(UMassEntitySubsystem& EntityS
 	ParallelFor(Jobs.Num(), [this, ExecutionContext, &ExecuteFunction, Jobs](const int32 JobIndex)
 	{
 		Jobs[JobIndex].Archetype.ExecutionFunctionForChunk(ExecutionContext, ExecuteFunction
-			, Jobs[JobIndex].ArchetypeIndex != INDEX_NONE ? ArchetypeComponentMapping[Jobs[JobIndex].ArchetypeIndex] : FLWRequirementIndicesMapping()
+			, Jobs[JobIndex].ArchetypeIndex != INDEX_NONE ? ArchetypeComponentMapping[Jobs[JobIndex].ArchetypeIndex] : FMassQueryRequirementIndicesMapping()
 			, Jobs[JobIndex].ChunkInfo
 			, ChunkCondition);
 	});
@@ -241,13 +241,13 @@ void FLWComponentQuery::ParallelForEachEntityChunk(UMassEntitySubsystem& EntityS
 	ExecutionContext.FlushDeferred(EntitySubsystem);
 }
 
-int32 FLWComponentQuery::GetMatchingEntitiesNum(UMassEntitySubsystem& InEntitySubsystem)
+int32 FMassEntityQuery::GetMatchingEntitiesNum(UMassEntitySubsystem& InEntitySubsystem)
 {
 	CacheArchetypes(InEntitySubsystem);
 	int32 TotalEntities = 0;
 	for (FArchetypeHandle& ArchetypeHandle : ValidArchetypes)
 	{
-		if (const FArchetypeData* Archetype = ArchetypeHandle.DataPtr.Get())
+		if (const FMassArchetypeData* Archetype = ArchetypeHandle.DataPtr.Get())
 		{
 			TotalEntities += Archetype->GetNumEntities();
 		}
@@ -255,13 +255,13 @@ int32 FLWComponentQuery::GetMatchingEntitiesNum(UMassEntitySubsystem& InEntitySu
 	return TotalEntities;
 }
 
-bool FLWComponentQuery::HasMatchingEntities(UMassEntitySubsystem& InEntitySubsystem)
+bool FMassEntityQuery::HasMatchingEntities(UMassEntitySubsystem& InEntitySubsystem)
 {
 	CacheArchetypes(InEntitySubsystem);
 
 	for (FArchetypeHandle& ArchetypeHandle : ValidArchetypes)
 	{
-		const FArchetypeData* Archetype = ArchetypeHandle.DataPtr.Get();
+		const FMassArchetypeData* Archetype = ArchetypeHandle.DataPtr.Get();
 		if (Archetype && Archetype->GetNumEntities() > 0)
 		{
 			return true;
@@ -270,14 +270,14 @@ bool FLWComponentQuery::HasMatchingEntities(UMassEntitySubsystem& InEntitySubsys
 	return false;
 }
 
-FString FLWComponentQuery::DebugGetDescription() const
+FString FMassEntityQuery::DebugGetDescription() const
 {
 #if WITH_MASSENTITY_DEBUG
 	TStringBuilder<256> StringBuilder;
 	StringBuilder.Append(TEXT("<"));
 
 	bool bNeedsComma = false;
-	for (const FLWComponentRequirement& Requirement : Requirements)
+	for (const FMassFragmentRequirement& Requirement : Requirements)
 	{
 		if (bNeedsComma)
 		{
@@ -295,12 +295,12 @@ FString FLWComponentQuery::DebugGetDescription() const
 }
 
 //////////////////////////////////////////////////////////////////////
-// FLWComponentRequirement
+// FMassFragmentRequirement
 
-FString FLWComponentRequirement::DebugGetDescription() const
+FString FMassFragmentRequirement::DebugGetDescription() const
 {
 #if WITH_MASSENTITY_DEBUG
-	return FString::Printf(TEXT("%s%s[%s]"), IsOptional() ? TEXT("?") : (Presence == ELWComponentPresence::None ? TEXT("-") : TEXT("+"))
+	return FString::Printf(TEXT("%s%s[%s]"), IsOptional() ? TEXT("?") : (Presence == EMassFragmentPresence::None ? TEXT("-") : TEXT("+"))
 		, *GetNameSafe(StructType), *UE::AggregateTicking::DebugGetComponentAccessString(AccessMode));
 #else
 	return {};
diff --git a/Engine/Plugins/Runtime/MassEntity/Source/MassEntity/Private/EntityView.cpp b/Engine/Plugins/Runtime/MassEntity/Source/MassEntity/Private/EntityView.cpp
--- a/Engine/Plugins/Runtime/MassEntity/Source/MassEntity/Private/EntityView.cpp
+++ b/Engine/Plugins/Runtime/MassEntity/Source/MassEntity/Private/EntityView.cpp
@@ -6,8 +6,8 @@
 
 
 //////////////////////////////////////////////////////////////////////
-// FEntityView
-FEntityView::FEntityView(const FArchetypeHandle& ArchetypeHandle, FLWEntity InEntity)
+// FMassEntityView
+FMassEntityView::FMassEntityView(const FArchetypeHandle& ArchetypeHandle, FMassEntityHandle InEntity)
 {
 	Entity = InEntity;
 	check(ArchetypeHandle.IsValid());
@@ -15,7 +15,7 @@ FEntityView::FEntityView(const FArchetypeHandle& ArchetypeHandle, FLWEntity InEn
 	EntityHandle = Archetype->MakeEntityHandle(Entity);
 }
 
-FEntityView::FEntityView(const UMassEntitySubsystem& EntitySubsystem, FLWEntity InEntity)
+FMassEntityView::FMassEntityView(const UMassEntitySubsystem& EntitySubsystem, FMassEntityHandle InEntity)
 {
 	Entity = InEntity;
 	const FArchetypeHandle ArchetypeHandle = EntitySubsystem.GetArchetypeForEntity(Entity);
@@ -24,7 +24,7 @@ FEntityView::FEntityView(const UMassEntitySubsystem& EntitySubsystem, FLWEntity
 	EntityHandle = Archetype->MakeEntityHandle(Entity);
 }
 
-void* FEntityView::GetComponentPtr(const UScriptStruct& ComponentType) const
+void* FMassEntityView::GetComponentPtr(const UScriptStruct& ComponentType) const
 {
 	checkSlow(Archetype && EntityHandle.IsValid());
 	if (const int32* ComponentIndex = Archetype->GetComponentIndex(&ComponentType))
@@ -35,14 +35,14 @@ void* FEntityView::GetComponentPtr(const UScriptStruct& ComponentType) const
 	return nullptr;
 }
 
-void* FEntityView::GetComponentPtrChecked(const UScriptStruct& ComponentType) const
+void* FMassEntityView::GetComponentPtrChecked(const UScriptStruct& ComponentType) const
 {
 	checkSlow(Archetype && EntityHandle.IsValid());
 	const int32 ComponentIndex = Archetype->GetComponentIndexChecked(&ComponentType);
 	return Archetype->GetComponentData(ComponentIndex, EntityHandle);
 }
 
-bool FEntityView::HasTag(const UScriptStruct& TagType) const
+bool FMassEntityView::HasTag(const UScriptStruct& TagType) const
 {
 	checkSlow(Archetype && EntityHandle.IsValid());
 	return Archetype->HasTagType(&TagType);
diff --git a/Engine/Plugins/Runtime/MassEntity/Source/MassEntity/Private/LWCCommandBuffer.cpp b/Engine/Plugins/Runtime/MassEntity/Source/MassEntity/Private/LWCCommandBuffer.cpp
--- a/Engine/Plugins/Runtime/MassEntity/Source/MassEntity/Private/LWCCommandBuffer.cpp
+++ b/Engine/Plugins/Runtime/MassEntity/Source/MassEntity/Private/LWCCommandBuffer.cpp
@@ -5,7 +5,7 @@
 #include "ProfilingDebugging/CsvProfiler.h"
 
 //////////////////////////////////////////////////////////////////////
-// FLWCCommandBuffer
+// FMassCommandBuffer
 
 CSV_DEFINE_CATEGORY(MassEntities, true);
 CSV_DEFINE_CATEGORY(MassEntitiesCounters, true);
@@ -66,7 +66,7 @@ void GetCommandStatNames(FStructView Entry, FString& OutName, ANSIName*& OutANSI
 #endif
 } // UE::FLWCCommand
 
-void FLWCCommandBuffer::ReplayBufferAgainstSystem(UMassEntitySubsystem* System)
+void FMassCommandBuffer::ReplayBufferAgainstSystem(UMassEntitySubsystem* System)
 {
 	check(System);
 
@@ -103,7 +103,7 @@ void FLWCCommandBuffer::ReplayBufferAgainstSystem(UMassEntitySubsystem* System)
 	}
 }
 
-void FLWCCommandBuffer::MoveAppend(FLWCCommandBuffer& Other)
+void FMassCommandBuffer::MoveAppend(FMassCommandBuffer& Other)
 {
 	// @todo optimize, there surely a way to do faster then this.
 	UE_MT_SCOPED_READ_ACCESS(Other.PendingCommandsDetector);
diff --git a/Engine/Plugins/Runtime/MassEntity/Source/MassEntity/Private/MassEntityDebug.cpp b/Engine/Plugins/Runtime/MassEntity/Source/MassEntity/Private/MassEntityDebug.cpp
--- a/Engine/Plugins/Runtime/MassEntity/Source/MassEntity/Private/MassEntityDebug.cpp
+++ b/Engine/Plugins/Runtime/MassEntity/Source/MassEntity/Private/MassEntityDebug.cpp
@@ -6,17 +6,17 @@
 #include "Engine/World.h"
 #include "Engine/Engine.h"
 
-DEFINE_ENUM_TO_STRING(EPipeProcessingPhase);
+DEFINE_ENUM_TO_STRING(EMassProcessingPhase);
 
 namespace UE::Pipe::Debug
 {
 
-void DebugOutputDescription(TConstArrayView<UPipeProcessor*> Processors, FOutputDevice& Ar)
+void DebugOutputDescription(TConstArrayView<UMassProcessor*> Processors, FOutputDevice& Ar)
 {
 #if WITH_MASSENTITY_DEBUG
 	const bool bAutoLineEnd = Ar.GetAutoEmitLineTerminator();
 	Ar.SetAutoEmitLineTerminator(false);
-	for (const UPipeProcessor* Proc : Processors)
+	for (const UMassProcessor* Proc : Processors)
 	{
 		if (Proc)
 		{
@@ -118,7 +118,7 @@ FAutoConsoleCommandWithWorldArgsAndOutputDevice LogFragmentSizes(
 	TEXT("Logs all the fragment types being used along with their sizes."),
 	FConsoleCommandWithWorldArgsAndOutputDeviceDelegate::CreateLambda([](const TArray<FString>& Params, UWorld* World, FOutputDevice& Ar)
 		{
-			for (const TWeakObjectPtr<const UScriptStruct>& WeakStruct : FLWComponentBitSet::DebugGetAllStructTypes())
+			for (const TWeakObjectPtr<const UScriptStruct>& WeakStruct : FMassFragmentBitSet::DebugGetAllStructTypes())
 			{
 				if (const UScriptStruct* StructType = WeakStruct.Get())
 				{
diff --git a/Engine/Plugins/Runtime/MassEntity/Source/MassEntity/Private/MassEntitySettings.cpp b/Engine/Plugins/Runtime/MassEntity/Source/MassEntity/Private/MassEntitySettings.cpp
--- a/Engine/Plugins/Runtime/MassEntity/Source/MassEntity/Private/MassEntitySettings.cpp
+++ b/Engine/Plugins/Runtime/MassEntity/Source/MassEntity/Private/MassEntitySettings.cpp
@@ -11,26 +11,26 @@
 #endif // WITH_EDITOR
 
 //----------------------------------------------------------------------//
-//  UPipeSettings
+//  UMassSettings
 //----------------------------------------------------------------------//
-UPipeSettings::UPipeSettings(const FObjectInitializer& ObjectInitializer)
+UMassSettings::UMassSettings(const FObjectInitializer& ObjectInitializer)
 	: Super(ObjectInitializer)
 {
-	for (int i = 0; i < (int)EPipeProcessingPhase::MAX; ++i)
+	for (int i = 0; i < (int)EMassProcessingPhase::MAX; ++i)
 	{
-		ProcessingPhasesConfig[i].PhaseName = *EnumToString(EPipeProcessingPhase(i));
+		ProcessingPhasesConfig[i].PhaseName = *EnumToString(EMassProcessingPhase(i));
 	}
 
-	FCoreDelegates::OnPostEngineInit.AddUObject(this, &UPipeSettings::BuildProcessorListAndPhases);
+	FCoreDelegates::OnPostEngineInit.AddUObject(this, &UMassSettings::BuildProcessorListAndPhases);
 }
 
-void UPipeSettings::BeginDestroy()
+void UMassSettings::BeginDestroy()
 {
 	FCoreDelegates::OnPostEngineInit.RemoveAll(this);
 	Super::BeginDestroy();
 }
 
-void UPipeSettings::BuildProcessorListAndPhases()
+void UMassSettings::BuildProcessorListAndPhases()
 {
 	if (bInitialized)
 	{
@@ -42,16 +42,16 @@ void UPipeSettings::BuildProcessorListAndPhases()
 	bInitialized = true;
 }
 
-void UPipeSettings::BuildPhases()
+void UMassSettings::BuildPhases()
 {
 #if WITH_EDITOR
-	for (int i = 0; i < int(EPipeProcessingPhase::MAX); ++i)
+	for (int i = 0; i < int(EMassProcessingPhase::MAX); ++i)
 	{
-		FPipeProcessingPhaseConfig& PhaseConfig = ProcessingPhasesConfig[i];
-		PhaseConfig.PhaseProcessor = NewObject<UPipeCompositeProcessor>(this, PhaseConfig.PhaseGroupClass
+		FMassProcessingPhaseConfig& PhaseConfig = ProcessingPhasesConfig[i];
+		PhaseConfig.PhaseProcessor = NewObject<UMassCompositeProcessor>(this, PhaseConfig.PhaseGroupClass
 			, *FString::Printf(TEXT("ProcessingPhase_%s"), *PhaseConfig.PhaseName.ToString()));
 		PhaseConfig.PhaseProcessor->SetGroupName(PhaseConfig.PhaseName);
-		PhaseConfig.PhaseProcessor->SetProcessingPhase(EPipeProcessingPhase(i));
+		PhaseConfig.PhaseProcessor->SetProcessingPhase(EMassProcessingPhase(i));
 		const FString PhaseDumpDependencyGraphFileName = !DumpDependencyGraphFileName.IsEmpty() ? DumpDependencyGraphFileName + TEXT("_") + PhaseConfig.PhaseName.ToString() : FString();
 		PhaseConfig.PhaseProcessor->CopyAndSort(PhaseConfig, PhaseDumpDependencyGraphFileName);
 		
@@ -62,16 +62,16 @@ void UPipeSettings::BuildPhases()
 #endif // WITH_EDITOR
 }
 
-void UPipeSettings::BuildProcessorList()
+void UMassSettings::BuildProcessorList()
 {
 	ProcessorCDOs.Reset();
-	for (FPipeProcessingPhaseConfig& PhaseConfig : ProcessingPhasesConfig)
+	for (FMassProcessingPhaseConfig& PhaseConfig : ProcessingPhasesConfig)
 	{
 		PhaseConfig.ProcessorCDOs.Reset();
 	}
 
 	TArray<UClass*> SubClassess;
-	GetDerivedClasses(UPipeProcessor::StaticClass(), SubClassess);
+	GetDerivedClasses(UMassProcessor::StaticClass(), SubClassess);
 
 	for (int i = SubClassess.Num() - 1; i; --i)
 	{
@@ -80,9 +80,9 @@ void UPipeSettings::BuildProcessorList()
 			continue;
 		}
 
-		UPipeProcessor* ProcessorCDO = GetMutableDefault<UPipeProcessor>(SubClassess[i]);
-		// we explicitly restrict adding UPipeCompositeProcessor. If needed by specific project a derived class can be added
-		if (ProcessorCDO && SubClassess[i] != UPipeCompositeProcessor::StaticClass()
+		UMassProcessor* ProcessorCDO = GetMutableDefault<UMassProcessor>(SubClassess[i]);
+		// we explicitly restrict adding UMassCompositeProcessor. If needed by specific project a derived class can be added
+		if (ProcessorCDO && SubClassess[i] != UMassCompositeProcessor::StaticClass()
 #if WITH_EDITOR
 			&& ProcessorCDO->ShouldShowUpInSettings()
 #endif // WITH_EDITOR
@@ -96,16 +96,16 @@ void UPipeSettings::BuildProcessorList()
 		}
 	}
 
-	ProcessorCDOs.Sort([](UPipeProcessor& LHS, UPipeProcessor& RHS) {
+	ProcessorCDOs.Sort([](UMassProcessor& LHS, UMassProcessor& RHS) {
 		return LHS.GetName().Compare(RHS.GetName()) < 0;
 	});
 }
 
-void UPipeSettings::AddToActiveProcessorsList(TSubclassOf<UPipeProcessor> ProcessorClass)
+void UMassSettings::AddToActiveProcessorsList(TSubclassOf<UMassProcessor> ProcessorClass)
 {
-	if (UPipeProcessor* ProcessorCDO = GetMutableDefault<UPipeProcessor>(ProcessorClass))
+	if (UMassProcessor* ProcessorCDO = GetMutableDefault<UMassProcessor>(ProcessorClass))
 	{
-		if (ProcessorClass == UPipeCompositeProcessor::StaticClass())
+		if (ProcessorClass == UMassCompositeProcessor::StaticClass())
 		{
 			UE_VLOG_UELOG(this, LogPipe, Log, TEXT("%s adding PipeCompositeProcessor to the global processor list is unsupported"), ANSI_TO_TCHAR(__FUNCTION__));
 		}
@@ -126,19 +126,19 @@ void UPipeSettings::AddToActiveProcessorsList(TSubclassOf<UPipeProcessor> Proces
 	}
 }
 
-const FPipeProcessingPhaseConfig* UPipeSettings::GetProcessingPhasesConfig()
+const FMassProcessingPhaseConfig* UMassSettings::GetProcessingPhasesConfig()
 {
 	BuildProcessorListAndPhases();
 	return ProcessingPhasesConfig;
 }
 
 #if WITH_EDITOR
-void UPipeSettings::PostEditChangeProperty(FPropertyChangedEvent& PropertyChangedEvent)
+void UMassSettings::PostEditChangeProperty(FPropertyChangedEvent& PropertyChangedEvent)
 {
 	// @todo replacing a valid entry in the array is not handled right now. 
 
 	static const FName PipeName = TEXT("Pipe");
-	static const FName ProcessorCDOsName = GET_MEMBER_NAME_CHECKED(UPipeSettings, ProcessorCDOs);
+	static const FName ProcessorCDOsName = GET_MEMBER_NAME_CHECKED(UMassSettings, ProcessorCDOs);
 
 	Super::PostEditChangeProperty(PropertyChangedEvent);
 
@@ -161,7 +161,7 @@ void UPipeSettings::PostEditChangeProperty(FPropertyChangedEvent& PropertyChange
 	}
 }
 
-void UPipeSettings::PostEditChangeChainProperty(struct FPropertyChangedChainEvent& PropertyChangedEvent)
+void UMassSettings::PostEditChangeChainProperty(struct FPropertyChangedChainEvent& PropertyChangedEvent)
 {
 	static const FName AutoRegisterName = TEXT("bAutoRegisterWithProcessingPhases");
 
