AggregateTicking and Pipe type renames:

diff --git a/Engine/Plugins/Runtime/MassEntity/Source/MassEntity/Private/MassEntitySubsystem.cpp b/Engine/Plugins/Runtime/MassEntity/Source/MassEntity/Private/MassEntitySubsystem.cpp
--- a/Engine/Plugins/Runtime/MassEntity/Source/MassEntity/Private/MassEntitySubsystem.cpp
+++ b/Engine/Plugins/Runtime/MassEntity/Source/MassEntity/Private/MassEntitySubsystem.cpp
@@ -7,22 +7,22 @@
 #include "Engine/World.h"
 #include "UObject/UObjectIterator.h"
 
-const FLWEntity UMassEntitySubsystem::InvalidEntity;
+const FMassEntityHandle UMassEntitySubsystem::InvalidEntity;
 
 DEFINE_LOG_CATEGORY(LogAggregateTicking);
 
 namespace UE { namespace AggregateTicking {
 
-FString DebugGetComponentAccessString(ELWComponentAccess Access)
+FString DebugGetComponentAccessString(EMassFragmentAccess Access)
 {
 #if WITH_MASSENTITY_DEBUG
 	switch (Access)
 	{
-	case ELWComponentAccess::None:	return TEXT("--");
-	case ELWComponentAccess::ReadOnly:	return TEXT("RO");
-	case ELWComponentAccess::ReadWrite:	return TEXT("RW");
+	case EMassFragmentAccess::None:	return TEXT("--");
+	case EMassFragmentAccess::ReadOnly:	return TEXT("RO");
+	case EMassFragmentAccess::ReadWrite:	return TEXT("RW");
 	default:
-		ensureMsgf(false, TEXT("Missing string conversion for ELWComponentAccess=%d"), Access);
+		ensureMsgf(false, TEXT("Missing string conversion for EMassFragmentAccess=%d"), Access);
 		break;
 	}
 #endif // WITH_MASSENTITY_DEBUG
@@ -57,7 +57,7 @@ void UMassEntitySubsystem::GetResourceSizeEx(FResourceSizeEx& CumulativeResource
 
 	for (const auto& KVP : ComponentHashToArchetypeMap)
 	{
-		for (const TSharedPtr<FArchetypeData>& ArchetypePtr : KVP.Value)
+		for (const TSharedPtr<FMassArchetypeData>& ArchetypePtr : KVP.Value)
 		{
 			CumulativeResourceSize.AddDedicatedSystemMemoryBytes(ArchetypePtr->GetAllocatedSize());
 		}
@@ -70,25 +70,25 @@ void UMassEntitySubsystem::Initialize(FSubsystemCollectionBase& Collection)
 	Entities.Add();
 	SerialNumberGenerator.fetch_add(FMath::Max(1,NumReservedEntities));
 
-	DeferredCommandBuffer = MakeShareable(new FLWCCommandBuffer());
+	DeferredCommandBuffer = MakeShareable(new FMassCommandBuffer());
 
 	// initialize the bitsets
 
-	FLWComponentBitSet Components;
-	FLWTagBitSet Tags;
-	FLWChunkComponentBitSet ChunkComponents;
+	FMassFragmentBitSet Components;
+	FMassTagBitSet Tags;
+	FMassChunkFragmentBitSet ChunkComponents;
 
 	for (TObjectIterator<UStruct> StructIt; StructIt; ++StructIt)
 	{
-		if (StructIt->IsChildOf(FLWComponentData::StaticStruct()))
+		if (StructIt->IsChildOf(FMassFragment::StaticStruct()))
 		{
 			Components.Add(*CastChecked<UScriptStruct>(*StructIt));
 		}
-		else if (StructIt->IsChildOf(FComponentTag::StaticStruct()))
+		else if (StructIt->IsChildOf(FMassTag::StaticStruct()))
 		{
 			Tags.Add(*CastChecked<UScriptStruct>(*StructIt));
 		}
-		else if (StructIt->IsChildOf(FLWChunkComponent::StaticStruct()))
+		else if (StructIt->IsChildOf(FMassChunkFragment::StaticStruct()))
 		{
 			ChunkComponents.Add(*CastChecked<UScriptStruct>(*StructIt));
 		}
@@ -102,22 +102,22 @@ void UMassEntitySubsystem::Deinitialize()
 
 FArchetypeHandle UMassEntitySubsystem::CreateArchetype(TConstArrayView<const UScriptStruct*> ComponentsAngTagsList)
 {
-	FLWChunkComponentBitSet ChunkComponents;
-	FLWTagBitSet Tags;
+	FMassChunkFragmentBitSet ChunkComponents;
+	FMassTagBitSet Tags;
 	TArray<const UScriptStruct*, TInlineAllocator<16>> ComponentList;
 	ComponentList.Reserve(ComponentsAngTagsList.Num());
 
 	for (const UScriptStruct* Type : ComponentsAngTagsList)
 	{
-		if (Type->IsChildOf(FLWComponentData::StaticStruct()))
+		if (Type->IsChildOf(FMassFragment::StaticStruct()))
 		{
 			ComponentList.Add(Type);
 		}
-		else if (Type->IsChildOf(FComponentTag::StaticStruct()))
+		else if (Type->IsChildOf(FMassTag::StaticStruct()))
 		{
 			Tags.Add(*Type);
 		}
-		else if (Type->IsChildOf(FLWChunkComponent::StaticStruct()))
+		else if (Type->IsChildOf(FMassChunkFragment::StaticStruct()))
 		{
 			ChunkComponents.Add(*Type);
 		}
@@ -128,30 +128,30 @@ FArchetypeHandle UMassEntitySubsystem::CreateArchetype(TConstArrayView<const USc
 		}
 	}
 
-	return CreateArchetype(FLWCompositionDescriptor(FLWComponentBitSet(ComponentList), Tags, ChunkComponents));
+	return CreateArchetype(FMassCompositionDescriptor(FMassFragmentBitSet(ComponentList), Tags, ChunkComponents));
 }
 
-FArchetypeHandle UMassEntitySubsystem::CreateArchetype(TConstArrayView<const UScriptStruct*> ComponentList, const FLWTagBitSet& Tags, const FLWChunkComponentBitSet& ChunkComponents)
+FArchetypeHandle UMassEntitySubsystem::CreateArchetype(TConstArrayView<const UScriptStruct*> ComponentList, const FMassTagBitSet& Tags, const FMassChunkFragmentBitSet& ChunkComponents)
 {
-	return CreateArchetype(FLWCompositionDescriptor(ComponentList, Tags, ChunkComponents));
+	return CreateArchetype(FMassCompositionDescriptor(ComponentList, Tags, ChunkComponents));
 }
 
-FArchetypeHandle UMassEntitySubsystem::CreateArchetype(const TSharedPtr<FArchetypeData>& SourceArchetype, const FLWComponentBitSet& NewComponents)
+FArchetypeHandle UMassEntitySubsystem::CreateArchetype(const TSharedPtr<FMassArchetypeData>& SourceArchetype, const FMassFragmentBitSet& NewComponents)
 {
 	check(SourceArchetype.IsValid());
 	checkf(NewComponents.IsEmpty() == false, TEXT("%s Adding an empty component list to an archetype is not supported."), ANSI_TO_TCHAR(__FUNCTION__));
 
-	return CreateArchetype(FLWCompositionDescriptor(NewComponents + SourceArchetype->GetComponentBitSet(), SourceArchetype->GetTagBitSet(), SourceArchetype->GetChunkComponentBitSet()));
+	return CreateArchetype(FMassCompositionDescriptor(NewComponents + SourceArchetype->GetComponentBitSet(), SourceArchetype->GetTagBitSet(), SourceArchetype->GetChunkComponentBitSet()));
 }
 
-FArchetypeHandle UMassEntitySubsystem::CreateArchetype(const FLWCompositionDescriptor& Composition)
+FArchetypeHandle UMassEntitySubsystem::CreateArchetype(const FMassCompositionDescriptor& Composition)
 {
 	const uint32 TypeHash = Composition.CalculateHash();
 
-	TArray<TSharedPtr<FArchetypeData>>& HashRow = ComponentHashToArchetypeMap.FindOrAdd(TypeHash);
+	TArray<TSharedPtr<FMassArchetypeData>>& HashRow = ComponentHashToArchetypeMap.FindOrAdd(TypeHash);
 
 	FArchetypeHandle Result;
-	for (TSharedPtr<FArchetypeData>& Ptr : HashRow)
+	for (TSharedPtr<FMassArchetypeData>& Ptr : HashRow)
 	{
 		if (Ptr->IsEquivalent(Composition))
 		{
@@ -163,11 +163,11 @@ FArchetypeHandle UMassEntitySubsystem::CreateArchetype(const FLWCompositionDescr
 	if (!Result.DataPtr.IsValid())
 	{
 		// Create a new archetype
-		TSharedPtr<FArchetypeData> NewArchetype = MakeShareable(new FArchetypeData);
+		TSharedPtr<FMassArchetypeData> NewArchetype = MakeShareable(new FMassArchetypeData);
 		NewArchetype->Initialize(Composition.Components, Composition.Tags, Composition.ChunkComponents);
 		HashRow.Add(NewArchetype);
 
-		for (const FArchetypeComponentConfig& ComponentConfig : NewArchetype->GetComponentConfigs())
+		for (const FMassArchetypeFragmentConfig& ComponentConfig : NewArchetype->GetComponentConfigs())
 		{
 			checkSlow(ComponentConfig.ComponentType)
 			ComponentTypeToArchetypeMap.FindOrAdd(ComponentConfig.ComponentType).Add(NewArchetype);
@@ -181,16 +181,16 @@ FArchetypeHandle UMassEntitySubsystem::CreateArchetype(const FLWCompositionDescr
 	return MoveTemp(Result);
 }
 
-FArchetypeHandle UMassEntitySubsystem::InternalCreateSiblingArchetype(const TSharedPtr<FArchetypeData>& SourceArchetype, const FLWTagBitSet& OverrideTags)
+FArchetypeHandle UMassEntitySubsystem::InternalCreateSiblingArchetype(const TSharedPtr<FMassArchetypeData>& SourceArchetype, const FMassTagBitSet& OverrideTags)
 {
 	checkSlow(SourceArchetype.IsValid());
-	const FArchetypeData& SourceArchetypeRef = *SourceArchetype.Get();
-	const uint32 TypeHash = FLWCompositionDescriptor::CalculateHash(SourceArchetypeRef.GetComponentBitSet(), OverrideTags, SourceArchetypeRef.GetChunkComponentBitSet());
+	const FMassArchetypeData& SourceArchetypeRef = *SourceArchetype.Get();
+	const uint32 TypeHash = FMassCompositionDescriptor::CalculateHash(SourceArchetypeRef.GetComponentBitSet(), OverrideTags, SourceArchetypeRef.GetChunkComponentBitSet());
 
-	TArray<TSharedPtr<FArchetypeData>>& HashRow = ComponentHashToArchetypeMap.FindOrAdd(TypeHash);
+	TArray<TSharedPtr<FMassArchetypeData>>& HashRow = ComponentHashToArchetypeMap.FindOrAdd(TypeHash);
 
 	FArchetypeHandle Result;
-	for (TSharedPtr<FArchetypeData>& Ptr : HashRow)
+	for (TSharedPtr<FMassArchetypeData>& Ptr : HashRow)
 	{
 		if (Ptr->IsEquivalent(SourceArchetypeRef.GetComponentBitSet(), OverrideTags, SourceArchetypeRef.GetChunkComponentBitSet()))
 		{
@@ -202,11 +202,11 @@ FArchetypeHandle UMassEntitySubsystem::InternalCreateSiblingArchetype(const TSha
 	if (!Result.DataPtr.IsValid())
 	{
 		// Create a new archetype
-		TSharedPtr<FArchetypeData> NewArchetype = MakeShareable(new FArchetypeData);
+		TSharedPtr<FMassArchetypeData> NewArchetype = MakeShareable(new FMassArchetypeData);
 		NewArchetype->InitializeWithSibling(SourceArchetypeRef, OverrideTags);
 		HashRow.Add(NewArchetype);
 
-		for (const FArchetypeComponentConfig& ComponentConfig : NewArchetype->GetComponentConfigs())
+		for (const FMassArchetypeFragmentConfig& ComponentConfig : NewArchetype->GetComponentConfigs())
 		{
 			checkSlow(ComponentConfig.ComponentType)
 			ComponentTypeToArchetypeMap.FindOrAdd(ComponentConfig.ComponentType).Add(NewArchetype);
@@ -220,7 +220,7 @@ FArchetypeHandle UMassEntitySubsystem::InternalCreateSiblingArchetype(const TSha
 	return MoveTemp(Result);
 }
 
-FArchetypeHandle UMassEntitySubsystem::GetArchetypeForEntity(FLWEntity Entity) const
+FArchetypeHandle UMassEntitySubsystem::GetArchetypeForEntity(FMassEntityHandle Entity) const
 {
 	FArchetypeHandle Result;
 	if (IsEntityValid(Entity))
@@ -233,14 +233,14 @@ FArchetypeHandle UMassEntitySubsystem::GetArchetypeForEntity(FLWEntity Entity) c
 void UMassEntitySubsystem::ForEachArchetypeComponentType(const FArchetypeHandle Archetype, TFunction< void(const UScriptStruct* /*ComponentType*/)> Function)
 {
 	check(Archetype.DataPtr.IsValid());
-	const FArchetypeData& ArchetypeData = *Archetype.DataPtr.Get();
+	const FMassArchetypeData& ArchetypeData = *Archetype.DataPtr.Get();
 	ArchetypeData.ForEachComponentType(Function);
 }
 
 void UMassEntitySubsystem::SetDefaultChunkComponentValue(const FArchetypeHandle Archetype, FConstStructView InstancedStruct)
 {
 	check(Archetype.DataPtr.IsValid());
-	FArchetypeData& ArchetypeData = *Archetype.DataPtr.Get();
+	FMassArchetypeData& ArchetypeData = *Archetype.DataPtr.Get();
 	ArchetypeData.SetDefaultChunkComponentValue(InstancedStruct);
 }
 
@@ -251,7 +251,7 @@ void UMassEntitySubsystem::DoEntityCompaction(const double TimeAllowed)
 	bool bReachedTimeLimit = false;
 	for (const auto& KVP : ComponentHashToArchetypeMap)
 	{
-		for (const TSharedPtr<FArchetypeData>& ArchetypePtr : KVP.Value)
+		for (const TSharedPtr<FMassArchetypeData>& ArchetypePtr : KVP.Value)
 		{
 			const double TimeAllowedLeft = TimeAllowedEnd - FPlatformTime::Seconds();
 			bReachedTimeLimit = TimeAllowedLeft <= 0.0;
@@ -268,23 +268,23 @@ void UMassEntitySubsystem::DoEntityCompaction(const double TimeAllowed)
 	}
 }
 
-FLWEntity UMassEntitySubsystem::CreateEntity(const FArchetypeHandle Archetype)
+FMassEntityHandle UMassEntitySubsystem::CreateEntity(const FArchetypeHandle Archetype)
 {
 	check(Archetype.DataPtr.IsValid());
 
-	FLWEntity Entity = ReserveEntity();
+	FMassEntityHandle Entity = ReserveEntity();
 	InternalBuildEntity(Entity, Archetype);
 	return Entity;
 }
 
-FLWEntity UMassEntitySubsystem::CreateEntity(TConstArrayView<FInstancedStruct> ComponentInstanceList)
+FMassEntityHandle UMassEntitySubsystem::CreateEntity(TConstArrayView<FInstancedStruct> ComponentInstanceList)
 {
 	check(ComponentInstanceList.Num() > 0);
 
-	FArchetypeHandle Archetype = CreateArchetype(FLWCompositionDescriptor(ComponentInstanceList, FLWTagBitSet(), FLWChunkComponentBitSet()));
+	FArchetypeHandle Archetype = CreateArchetype(FMassCompositionDescriptor(ComponentInstanceList, FMassTagBitSet(), FMassChunkFragmentBitSet()));
 	check(Archetype.DataPtr.IsValid());
 
-	FLWEntity Entity = ReserveEntity();
+	FMassEntityHandle Entity = ReserveEntity();
 	InternalBuildEntity(Entity, Archetype);
 
 	FEntityData& EntityData = Entities[Entity.Index];
@@ -293,10 +293,10 @@ FLWEntity UMassEntitySubsystem::CreateEntity(TConstArrayView<FInstancedStruct> C
 	return Entity;
 }
 
-FLWEntity UMassEntitySubsystem::ReserveEntity()
+FMassEntityHandle UMassEntitySubsystem::ReserveEntity()
 {
 	// @todo: Need to add thread safety to the reservation of an entity
-	FLWEntity Result;
+	FMassEntityHandle Result;
 	Result.Index = (EntityFreeIndexList.Num() > 0) ? EntityFreeIndexList.Pop(/*bAllowShrinking=*/ false) : Entities.Add();
 	Result.SerialNumber = SerialNumberGenerator.fetch_add(1);
 	Entities[Result.Index].SerialNumber = Result.SerialNumber;
@@ -304,14 +304,14 @@ FLWEntity UMassEntitySubsystem::ReserveEntity()
 	return Result;
 }
 
-void UMassEntitySubsystem::ReleaseReservedEntity(FLWEntity Entity)
+void UMassEntitySubsystem::ReleaseReservedEntity(FMassEntityHandle Entity)
 {
 	checkf(!IsEntityBuilt(Entity), TEXT("Entity is already built, use DestroyEntity() instead"));
 
 	InternalReleaseEntity(Entity);
 }
 
-void UMassEntitySubsystem::BuildEntity(FLWEntity Entity, FArchetypeHandle Archetype)
+void UMassEntitySubsystem::BuildEntity(FMassEntityHandle Entity, FArchetypeHandle Archetype)
 {
 	checkf(!IsEntityBuilt(Entity), TEXT("Expecting an entity that is not already built"));
 	check(Archetype.DataPtr.IsValid());
@@ -319,12 +319,12 @@ void UMassEntitySubsystem::BuildEntity(FLWEntity Entity, FArchetypeHandle Archet
 	InternalBuildEntity(Entity, Archetype);
 }
 
-void UMassEntitySubsystem::BuildEntity(FLWEntity Entity, TConstArrayView<FInstancedStruct> ComponentInstanceList)
+void UMassEntitySubsystem::BuildEntity(FMassEntityHandle Entity, TConstArrayView<FInstancedStruct> ComponentInstanceList)
 {
 	check(ComponentInstanceList.Num() > 0);
 	checkf(!IsEntityBuilt(Entity), TEXT("Expecting an entity that is not already built"));
 
-	FArchetypeHandle Archetype = CreateArchetype(FLWCompositionDescriptor(ComponentInstanceList, FLWTagBitSet(), FLWChunkComponentBitSet()));
+	FArchetypeHandle Archetype = CreateArchetype(FMassCompositionDescriptor(ComponentInstanceList, FMassTagBitSet(), FMassChunkFragmentBitSet()));
 	check(Archetype.DataPtr.IsValid());
 
 	InternalBuildEntity(Entity, Archetype);
@@ -333,9 +333,9 @@ void UMassEntitySubsystem::BuildEntity(FLWEntity Entity, TConstArrayView<FInstan
 	EntityData.CurrentArchetype->SetComponentsData(Entity, ComponentInstanceList);
 }
 
-int32 UMassEntitySubsystem::BatchCreateEntities(const FArchetypeHandle Archetype, const int32 Count, TArray<FLWEntity>& OutEntities)
+int32 UMassEntitySubsystem::BatchCreateEntities(const FArchetypeHandle Archetype, const int32 Count, TArray<FMassEntityHandle>& OutEntities)
 {
-	FArchetypeData* ArchetypePtr = Archetype.DataPtr.Get();
+	FMassArchetypeData* ArchetypePtr = Archetype.DataPtr.Get();
 	check(ArchetypePtr);
 	check(Count >= 0);
 	
@@ -345,7 +345,7 @@ int32 UMassEntitySubsystem::BatchCreateEntities(const FArchetypeHandle Archetype
 	// @todo optimize
 	for (; Index < OutEntities.Num(); ++Index)
 	{
-		FLWEntity& Result = OutEntities[Index];
+		FMassEntityHandle& Result = OutEntities[Index];
 		Result.Index = (EntityFreeIndexList.Num() > 0) ? EntityFreeIndexList.Pop(/*bAllowShrinking=*/ false) : Entities.Add();
 		Result.SerialNumber = SerialNumberGenerator++;
 
@@ -359,14 +359,14 @@ int32 UMassEntitySubsystem::BatchCreateEntities(const FArchetypeHandle Archetype
 	return Count;
 }
 
-void UMassEntitySubsystem::DestroyEntity(FLWEntity Entity)
+void UMassEntitySubsystem::DestroyEntity(FMassEntityHandle Entity)
 {
 	checkf(IsProcessing() == false, TEXT("Synchronous API function %s called during mass processing. Use asynchronous API instead."), ANSI_TO_TCHAR(__FUNCTION__));
 	
 	CheckIfEntityIsActive(Entity);
 
 	FEntityData& EntityData = Entities[Entity.Index];
-	FArchetypeData* Archetype = EntityData.CurrentArchetype.Get();
+	FMassArchetypeData* Archetype = EntityData.CurrentArchetype.Get();
 	check(Archetype);
 	Archetype->RemoveEntity(Entity);
 
@@ -378,7 +378,7 @@ void UMassEntitySubsystem::DestroyEntity(FLWEntity Entity)
 	}
 }
 
-void UMassEntitySubsystem::BatchDestroyEntities(TConstArrayView<FLWEntity> InEntities)
+void UMassEntitySubsystem::BatchDestroyEntities(TConstArrayView<FMassEntityHandle> InEntities)
 {
 	TRACE_CPUPROFILER_EVENT_SCOPE_STR("BatchDestroyEntities");
 
@@ -387,7 +387,7 @@ void UMassEntitySubsystem::BatchDestroyEntities(TConstArrayView<FLWEntity> InEnt
 	EntityFreeIndexList.Reserve(EntityFreeIndexList.Num() + InEntities.Num());
 
 	// @todo optimize, we can make savings by implementing Archetype->RemoveEntities()
-	for (const FLWEntity Entity : InEntities)
+	for (const FMassEntityHandle Entity : InEntities)
 	{
 		if (Entities.IsValidIndex(Entity.Index) == false)
 		{
@@ -400,7 +400,7 @@ void UMassEntitySubsystem::BatchDestroyEntities(TConstArrayView<FLWEntity> InEnt
 			continue;
 		}
 
-		FArchetypeData* Archetype = EntityData.CurrentArchetype.Get();
+		FMassArchetypeData* Archetype = EntityData.CurrentArchetype.Get();
 		check(Archetype);
 		Archetype->RemoveEntity(Entity);
 
@@ -416,12 +416,12 @@ void UMassEntitySubsystem::BatchDestroyEntityChunks(const FArchetypeChunkCollect
 	checkf(IsProcessing() == false, TEXT("Synchronous API function %s called during mass processing. Use asynchronous API instead."), ANSI_TO_TCHAR(__FUNCTION__));
 
 	check(Chunks.GetArchetype().IsValid());
-	TArray<FLWEntity> EntitiesRemoved;
+	TArray<FMassEntityHandle> EntitiesRemoved;
 	Chunks.GetArchetype().DataPtr->BatchDestroyEntityChunks(Chunks, EntitiesRemoved);
 	
 	EntityFreeIndexList.Reserve(EntityFreeIndexList.Num() + EntitiesRemoved.Num());
 
-	for (const FLWEntity& Entity : EntitiesRemoved)
+	for (const FMassEntityHandle& Entity : EntitiesRemoved)
 	{
 		check(Entities.IsValidIndex(Entity.Index));
 
@@ -434,29 +434,29 @@ void UMassEntitySubsystem::BatchDestroyEntityChunks(const FArchetypeChunkCollect
 	}
 }
 
-void UMassEntitySubsystem::AddComponentToEntity(FLWEntity Entity, const UScriptStruct* ComponentType)
+void UMassEntitySubsystem::AddComponentToEntity(FMassEntityHandle Entity, const UScriptStruct* ComponentType)
 {
 	checkf(ComponentType, TEXT("Null component type passed in to %s"), ANSI_TO_TCHAR(__FUNCTION__));
 	checkf(IsProcessing() == false, TEXT("Synchronous API function %s called during mass processing. Use asynchronous API instead."), ANSI_TO_TCHAR(__FUNCTION__));
 
 	CheckIfEntityIsActive(Entity);
 
-	InternalAddComponentListToEntityChecked(Entity, FLWComponentBitSet(*ComponentType));
+	InternalAddComponentListToEntityChecked(Entity, FMassFragmentBitSet(*ComponentType));
 }
 
-void UMassEntitySubsystem::AddComponentListToEntity(FLWEntity Entity, TConstArrayView<const UScriptStruct*> ComponentList)
+void UMassEntitySubsystem::AddComponentListToEntity(FMassEntityHandle Entity, TConstArrayView<const UScriptStruct*> ComponentList)
 {
 	CheckIfEntityIsActive(Entity);
 
-	InternalAddComponentListToEntityChecked(Entity, FLWComponentBitSet(ComponentList));
+	InternalAddComponentListToEntityChecked(Entity, FMassFragmentBitSet(ComponentList));
 }
 
-void UMassEntitySubsystem::AddCompositionToEntity_GetDelta(FLWEntity Entity, FLWCompositionDescriptor& Composition)
+void UMassEntitySubsystem::AddCompositionToEntity_GetDelta(FMassEntityHandle Entity, FMassCompositionDescriptor& Composition)
 {
 	CheckIfEntityIsActive(Entity);
 
 	FEntityData& EntityData = Entities[Entity.Index];
-	FArchetypeData* OldArchetype = EntityData.CurrentArchetype.Get();
+	FMassArchetypeData* OldArchetype = EntityData.CurrentArchetype.Get();
 	check(OldArchetype);
 
 	Composition.Components -= OldArchetype->GetCompositionDescriptor().Components;
@@ -464,7 +464,7 @@ void UMassEntitySubsystem::AddCompositionToEntity_GetDelta(FLWEntity Entity, FLW
 
 	if (Composition.IsEmpty() == false)
 	{
-		FLWCompositionDescriptor NewDescriptor = OldArchetype->GetCompositionDescriptor();
+		FMassCompositionDescriptor NewDescriptor = OldArchetype->GetCompositionDescriptor();
 		NewDescriptor.Components += Composition.Components;
 		NewDescriptor.Tags+= Composition.Tags;
 
@@ -473,7 +473,7 @@ void UMassEntitySubsystem::AddCompositionToEntity_GetDelta(FLWEntity Entity, FLW
 		if (ensure(NewArchetypeHandle.DataPtr != EntityData.CurrentArchetype))
 		{
 			// Move the entity over
-			FArchetypeData* NewArchetype = NewArchetypeHandle.DataPtr.Get();
+			FMassArchetypeData* NewArchetype = NewArchetypeHandle.DataPtr.Get();
 			check(NewArchetype);
 			EntityData.CurrentArchetype->MoveEntityToAnotherArchetype(Entity, *NewArchetype);
 			EntityData.CurrentArchetype = NewArchetypeHandle.DataPtr;
@@ -481,15 +481,15 @@ void UMassEntitySubsystem::AddCompositionToEntity_GetDelta(FLWEntity Entity, FLW
 	}
 }
 
-void UMassEntitySubsystem::RemoveCompositionFromEntity(FLWEntity Entity, const FLWCompositionDescriptor& InDescriptor)
+void UMassEntitySubsystem::RemoveCompositionFromEntity(FMassEntityHandle Entity, const FMassCompositionDescriptor& InDescriptor)
 {
 	CheckIfEntityIsActive(Entity);
 
 	FEntityData& EntityData = Entities[Entity.Index];
-	FArchetypeData* OldArchetype = EntityData.CurrentArchetype.Get();
+	FMassArchetypeData* OldArchetype = EntityData.CurrentArchetype.Get();
 	check(OldArchetype);
 
-	FLWCompositionDescriptor NewDescriptor = OldArchetype->GetCompositionDescriptor();
+	FMassCompositionDescriptor NewDescriptor = OldArchetype->GetCompositionDescriptor();
 	NewDescriptor.Components -= InDescriptor.Components;
 	NewDescriptor.Tags -= InDescriptor.Tags;
 
@@ -500,7 +500,7 @@ void UMassEntitySubsystem::RemoveCompositionFromEntity(FLWEntity Entity, const F
 		if (ensure(NewArchetypeHandle.DataPtr != EntityData.CurrentArchetype))
 		{
 			// Move the entity over
-			FArchetypeData* NewArchetype = NewArchetypeHandle.DataPtr.Get();
+			FMassArchetypeData* NewArchetype = NewArchetypeHandle.DataPtr.Get();
 			check(NewArchetype);
 			EntityData.CurrentArchetype->MoveEntityToAnotherArchetype(Entity, *NewArchetype);
 			EntityData.CurrentArchetype = NewArchetypeHandle.DataPtr;
@@ -508,36 +508,36 @@ void UMassEntitySubsystem::RemoveCompositionFromEntity(FLWEntity Entity, const F
 	}
 }
 
-void UMassEntitySubsystem::InternalBuildEntity(FLWEntity Entity, const FArchetypeHandle Archetype)
+void UMassEntitySubsystem::InternalBuildEntity(FMassEntityHandle Entity, const FArchetypeHandle Archetype)
 {
 	FEntityData& EntityData = Entities[Entity.Index];
 	EntityData.CurrentArchetype = Archetype.DataPtr;
 	EntityData.CurrentArchetype->AddEntity(Entity);
 }
 
-void UMassEntitySubsystem::InternalReleaseEntity(FLWEntity Entity)
+void UMassEntitySubsystem::InternalReleaseEntity(FMassEntityHandle Entity)
 {
 	FEntityData& EntityData = Entities[Entity.Index];
 	EntityData.Reset();
 	EntityFreeIndexList.Add(Entity.Index);
 }
 
-void UMassEntitySubsystem::InternalAddComponentListToEntityChecked(FLWEntity Entity, const FLWComponentBitSet& InComponents)
+void UMassEntitySubsystem::InternalAddComponentListToEntityChecked(FMassEntityHandle Entity, const FMassFragmentBitSet& InComponents)
 {
 	FEntityData& EntityData = Entities[Entity.Index];
-	FArchetypeData* OldArchetype = EntityData.CurrentArchetype.Get();
+	FMassArchetypeData* OldArchetype = EntityData.CurrentArchetype.Get();
 	check(OldArchetype);
 	
 	ensureMsgf(OldArchetype->GetComponentBitSet().HasAny(InComponents) == false, TEXT("Trying to add a new component type to an entity, but it already has some of them."));
 
-	const FLWComponentBitSet NewComponents = InComponents - OldArchetype->GetComponentBitSet();
+	const FMassFragmentBitSet NewComponents = InComponents - OldArchetype->GetComponentBitSet();
 	if (NewComponents.IsEmpty() == false)
 	{
 		InternalAddComponentListToEntity(Entity, NewComponents);
 	}
 }
 
-void UMassEntitySubsystem::InternalAddComponentListToEntity(FLWEntity Entity, const FLWComponentBitSet& NewComponents)
+void UMassEntitySubsystem::InternalAddComponentListToEntity(FMassEntityHandle Entity, const FMassFragmentBitSet& NewComponents)
 {
 	checkf(NewComponents.IsEmpty() == false, TEXT("%s is intended for internal calls with non empty NewComponents parameter"), ANSI_TO_TCHAR(__FUNCTION__));
 	check(Entities.IsValidIndex(Entity.Index));
@@ -550,70 +550,70 @@ void UMassEntitySubsystem::InternalAddComponentListToEntity(FLWEntity Entity, co
 	if (NewArchetypeHandle.DataPtr != EntityData.CurrentArchetype)
 	{
 		// Move the entity over
-		FArchetypeData* NewArchetype = NewArchetypeHandle.DataPtr.Get();
+		FMassArchetypeData* NewArchetype = NewArchetypeHandle.DataPtr.Get();
 		check(NewArchetype);
 		EntityData.CurrentArchetype->MoveEntityToAnotherArchetype(Entity, *NewArchetype);
 		EntityData.CurrentArchetype = NewArchetypeHandle.DataPtr;
 	}
 }
 
-void UMassEntitySubsystem::AddComponentInstanceListToEntity(FLWEntity Entity, TConstArrayView<FInstancedStruct> ComponentInstanceList)
+void UMassEntitySubsystem::AddComponentInstanceListToEntity(FMassEntityHandle Entity, TConstArrayView<FInstancedStruct> ComponentInstanceList)
 {
 	checkf(IsProcessing() == false, TEXT("Synchronous API function %s called during mass processing. Use asynchronous API instead."), ANSI_TO_TCHAR(__FUNCTION__));
 
 	CheckIfEntityIsActive(Entity);
 	checkf(ComponentInstanceList.Num() > 0, TEXT("Need to specify at least one component instances for this operation"));
 
-	InternalAddComponentListToEntityChecked(Entity, FLWComponentBitSet(ComponentInstanceList));
+	InternalAddComponentListToEntityChecked(Entity, FMassFragmentBitSet(ComponentInstanceList));
 
 	FEntityData& EntityData = Entities[Entity.Index];
 	EntityData.CurrentArchetype->SetComponentsData(Entity, ComponentInstanceList);
 }
 
-void UMassEntitySubsystem::RemoveComponentFromEntity(FLWEntity Entity, const UScriptStruct* ComponentType)
+void UMassEntitySubsystem::RemoveComponentFromEntity(FMassEntityHandle Entity, const UScriptStruct* ComponentType)
 {
 	RemoveComponentListFromEntity(Entity, MakeArrayView(&ComponentType, 1));
 }
 
-void UMassEntitySubsystem::RemoveComponentListFromEntity(FLWEntity Entity, TConstArrayView<const UScriptStruct*> ComponentList)
+void UMassEntitySubsystem::RemoveComponentListFromEntity(FMassEntityHandle Entity, TConstArrayView<const UScriptStruct*> ComponentList)
 {
 	checkf(IsProcessing() == false, TEXT("Synchronous API function %s called during mass processing. Use asynchronous API instead."), ANSI_TO_TCHAR(__FUNCTION__));
 
 	CheckIfEntityIsActive(Entity);
 	
 	FEntityData& EntityData = Entities[Entity.Index];
-	FArchetypeData* OldArchetype = EntityData.CurrentArchetype.Get();
+	FMassArchetypeData* OldArchetype = EntityData.CurrentArchetype.Get();
 	check(OldArchetype);
 
-	const FLWComponentBitSet ComponentsToRemove(ComponentList);
+	const FMassFragmentBitSet ComponentsToRemove(ComponentList);
 
 	if (ensureMsgf(OldArchetype->GetComponentBitSet().HasAny(ComponentsToRemove), TEXT("Trying to remove a list of components from an entity but none of the components given was found.")))
 	{
 		// If all the components got removed this will result in fetching of the empty archetype
-		const FArchetypeHandle NewArchetypeHandle = CreateArchetype(FLWCompositionDescriptor(OldArchetype->GetComponentBitSet() - ComponentsToRemove, OldArchetype->GetTagBitSet(), OldArchetype->GetChunkComponentBitSet()));
+		const FArchetypeHandle NewArchetypeHandle = CreateArchetype(FMassCompositionDescriptor(OldArchetype->GetComponentBitSet() - ComponentsToRemove, OldArchetype->GetTagBitSet(), OldArchetype->GetChunkComponentBitSet()));
 
 		// Move the entity over
-		FArchetypeData* NewArchetype = NewArchetypeHandle.DataPtr.Get();
+		FMassArchetypeData* NewArchetype = NewArchetypeHandle.DataPtr.Get();
 		check(NewArchetype);
 		OldArchetype->MoveEntityToAnotherArchetype(Entity, *NewArchetype);
 		EntityData.CurrentArchetype = NewArchetypeHandle.DataPtr;
 	}
 }
 
-void UMassEntitySubsystem::SwapTagsForEntity(FLWEntity Entity, const UScriptStruct* OldTagType, const UScriptStruct* NewTagType)
+void UMassEntitySubsystem::SwapTagsForEntity(FMassEntityHandle Entity, const UScriptStruct* OldTagType, const UScriptStruct* NewTagType)
 {
 	checkf(IsProcessing() == false, TEXT("Synchronous API function %s called during mass processing. Use asynchronous API instead."), ANSI_TO_TCHAR(__FUNCTION__));
 
 	CheckIfEntityIsActive(Entity);
 
-	checkf((OldTagType != nullptr) && OldTagType->IsChildOf(FComponentTag::StaticStruct()), TEXT("%s works only with tags while '%s' is not one."), ANSI_TO_TCHAR(__FUNCTION__), *GetPathNameSafe(OldTagType));
-	checkf((NewTagType != nullptr) && NewTagType->IsChildOf(FComponentTag::StaticStruct()), TEXT("%s works only with tags while '%s' is not one."), ANSI_TO_TCHAR(__FUNCTION__), *GetPathNameSafe(NewTagType));
+	checkf((OldTagType != nullptr) && OldTagType->IsChildOf(FMassTag::StaticStruct()), TEXT("%s works only with tags while '%s' is not one."), ANSI_TO_TCHAR(__FUNCTION__), *GetPathNameSafe(OldTagType));
+	checkf((NewTagType != nullptr) && NewTagType->IsChildOf(FMassTag::StaticStruct()), TEXT("%s works only with tags while '%s' is not one."), ANSI_TO_TCHAR(__FUNCTION__), *GetPathNameSafe(NewTagType));
 
 	FEntityData& EntityData = Entities[Entity.Index];
-	FArchetypeData* CurrentArchetype = EntityData.CurrentArchetype.Get();
+	FMassArchetypeData* CurrentArchetype = EntityData.CurrentArchetype.Get();
 	check(CurrentArchetype);
 
-	FLWTagBitSet NewTagBitSet = CurrentArchetype->GetTagBitSet();
+	FMassTagBitSet NewTagBitSet = CurrentArchetype->GetTagBitSet();
 	NewTagBitSet.Remove(*OldTagType);
 	NewTagBitSet.Add(*NewTagType);
 	
@@ -628,20 +628,20 @@ void UMassEntitySubsystem::SwapTagsForEntity(FLWEntity Entity, const UScriptStru
 	}
 }
 
-void UMassEntitySubsystem::AddTagToEntity(FLWEntity Entity, const UScriptStruct* TagType)
+void UMassEntitySubsystem::AddTagToEntity(FMassEntityHandle Entity, const UScriptStruct* TagType)
 {
-	checkf((TagType != nullptr) && TagType->IsChildOf(FComponentTag::StaticStruct()), TEXT("%s works only with tags while '%s' is not one."), ANSI_TO_TCHAR(__FUNCTION__), *GetPathNameSafe(TagType));
+	checkf((TagType != nullptr) && TagType->IsChildOf(FMassTag::StaticStruct()), TEXT("%s works only with tags while '%s' is not one."), ANSI_TO_TCHAR(__FUNCTION__), *GetPathNameSafe(TagType));
 
 	CheckIfEntityIsActive(Entity);
 
 	FEntityData& EntityData = Entities[Entity.Index];
-	FArchetypeData* CurrentArchetype = EntityData.CurrentArchetype.Get();
+	FMassArchetypeData* CurrentArchetype = EntityData.CurrentArchetype.Get();
 	check(CurrentArchetype);
 
 	if (CurrentArchetype->HasTagType(TagType) == false)
 	{
-		//FLWTagBitSet NewTags = CurrentArchetype->GetTagBitSet() - *TagType;
-		FLWTagBitSet NewTags = CurrentArchetype->GetTagBitSet();
+		//FMassTagBitSet NewTags = CurrentArchetype->GetTagBitSet() - *TagType;
+		FMassTagBitSet NewTags = CurrentArchetype->GetTagBitSet();
 		NewTags.Add(*TagType);
 		FArchetypeHandle NewArchetypeHandle = InternalCreateSiblingArchetype(EntityData.CurrentArchetype, NewTags);
 		checkSlow(NewArchetypeHandle.IsValid());
@@ -652,20 +652,20 @@ void UMassEntitySubsystem::AddTagToEntity(FLWEntity Entity, const UScriptStruct*
 	}
 }
 	
-void UMassEntitySubsystem::RemoveTagFromEntity(FLWEntity Entity, const UScriptStruct* TagType)
+void UMassEntitySubsystem::RemoveTagFromEntity(FMassEntityHandle Entity, const UScriptStruct* TagType)
 {
-	checkf((TagType != nullptr) && TagType->IsChildOf(FComponentTag::StaticStruct()), TEXT("%s works only with tags while '%s' is not one."), ANSI_TO_TCHAR(__FUNCTION__), *GetPathNameSafe(TagType));
+	checkf((TagType != nullptr) && TagType->IsChildOf(FMassTag::StaticStruct()), TEXT("%s works only with tags while '%s' is not one."), ANSI_TO_TCHAR(__FUNCTION__), *GetPathNameSafe(TagType));
 
 	CheckIfEntityIsActive(Entity);
 
 	FEntityData& EntityData = Entities[Entity.Index];
-	FArchetypeData* CurrentArchetype = EntityData.CurrentArchetype.Get();
+	FMassArchetypeData* CurrentArchetype = EntityData.CurrentArchetype.Get();
 	check(CurrentArchetype);
 
 	if (CurrentArchetype->HasTagType(TagType))
 	{
 		// CurrentArchetype->GetTagBitSet() -  *TagType
-		FLWTagBitSet NewTags = CurrentArchetype->GetTagBitSet();
+		FMassTagBitSet NewTags = CurrentArchetype->GetTagBitSet();
 		NewTags.Remove(*TagType);
 		FArchetypeHandle NewArchetypeHandle = InternalCreateSiblingArchetype(EntityData.CurrentArchetype, NewTags);
 		checkSlow(NewArchetypeHandle.IsValid());
@@ -677,11 +677,11 @@ void UMassEntitySubsystem::RemoveTagFromEntity(FLWEntity Entity, const UScriptSt
 }
 
 
-void UMassEntitySubsystem::MoveEntityToAnotherArchetype(FLWEntity Entity, FArchetypeHandle NewArchetypeHandle)
+void UMassEntitySubsystem::MoveEntityToAnotherArchetype(FMassEntityHandle Entity, FArchetypeHandle NewArchetypeHandle)
 {
 	CheckIfEntityIsActive(Entity);
 
-	FArchetypeData* NewArchetype = NewArchetypeHandle.DataPtr.Get();
+	FMassArchetypeData* NewArchetype = NewArchetypeHandle.DataPtr.Get();
 	check(NewArchetype);
 
 	// Move the entity over
@@ -690,7 +690,7 @@ void UMassEntitySubsystem::MoveEntityToAnotherArchetype(FLWEntity Entity, FArche
 	EntityData.CurrentArchetype = NewArchetypeHandle.DataPtr;
 }
 
-void UMassEntitySubsystem::SetEntityComponentsValues(FLWEntity Entity, TArrayView<const FInstancedStruct> ComponentInstanceList)
+void UMassEntitySubsystem::SetEntityComponentsValues(FMassEntityHandle Entity, TArrayView<const FInstancedStruct> ComponentInstanceList)
 {
 	CheckIfEntityIsActive(Entity);
 
@@ -700,7 +700,7 @@ void UMassEntitySubsystem::SetEntityComponentsValues(FLWEntity Entity, TArrayVie
 
 void UMassEntitySubsystem::BatchSetEntityComponentsValues(const FArchetypeChunkCollection& SparseEntities, TArrayView<const FInstancedStruct> ComponentInstanceList)
 {
-	FArchetypeData* Archetype = SparseEntities.GetArchetype().DataPtr.Get();
+	FMassArchetypeData* Archetype = SparseEntities.GetArchetype().DataPtr.Get();
 	check(Archetype);
 
 	for (const FInstancedStruct& ComponentTemplate : ComponentInstanceList)
@@ -709,57 +709,57 @@ void UMassEntitySubsystem::BatchSetEntityComponentsValues(const FArchetypeChunkC
 	}
 }
 
-void* UMassEntitySubsystem::InternalGetComponentDataChecked(FLWEntity Entity, const UScriptStruct* ComponentType) const
+void* UMassEntitySubsystem::InternalGetComponentDataChecked(FMassEntityHandle Entity, const UScriptStruct* ComponentType) const
 {
 	CheckIfEntityIsActive(Entity);
 
-	checkf((ComponentType != nullptr) && ComponentType->IsChildOf(FLWComponentData::StaticStruct()), TEXT("InternalGetComponentDataChecked called with an invalid component type '%s'"), *GetPathNameSafe(ComponentType));
+	checkf((ComponentType != nullptr) && ComponentType->IsChildOf(FMassFragment::StaticStruct()), TEXT("InternalGetComponentDataChecked called with an invalid component type '%s'"), *GetPathNameSafe(ComponentType));
 	const FEntityData& EntityData = Entities[Entity.Index];
 	return EntityData.CurrentArchetype->GetComponentDataForEntityChecked(ComponentType, Entity.Index);
 }
 
-void* UMassEntitySubsystem::InternalGetComponentDataPtr(FLWEntity Entity, const UScriptStruct* ComponentType) const
+void* UMassEntitySubsystem::InternalGetComponentDataPtr(FMassEntityHandle Entity, const UScriptStruct* ComponentType) const
 {
 	CheckIfEntityIsActive(Entity);
-	checkf((ComponentType != nullptr) && ComponentType->IsChildOf(FLWComponentData::StaticStruct()), TEXT("InternalGetComponentData called with an invalid component type '%s'"), *GetPathNameSafe(ComponentType));
+	checkf((ComponentType != nullptr) && ComponentType->IsChildOf(FMassFragment::StaticStruct()), TEXT("InternalGetComponentData called with an invalid component type '%s'"), *GetPathNameSafe(ComponentType));
 	const FEntityData& EntityData = Entities[Entity.Index];
 	return EntityData.CurrentArchetype->GetComponentDataForEntity(ComponentType, Entity.Index);
 }
 
-bool UMassEntitySubsystem::IsEntityValid(FLWEntity Entity) const
+bool UMassEntitySubsystem::IsEntityValid(FMassEntityHandle Entity) const
 {
 	return (Entity.Index > 0) && Entities.IsValidIndex(Entity.Index) && (Entities[Entity.Index].SerialNumber == Entity.SerialNumber);
 }
 
-bool UMassEntitySubsystem::IsEntityBuilt(FLWEntity Entity) const
+bool UMassEntitySubsystem::IsEntityBuilt(FMassEntityHandle Entity) const
 {
 	CheckIfEntityIsValid(Entity);
 	return Entities[Entity.Index].CurrentArchetype.IsValid();
 }
 
-void UMassEntitySubsystem::CheckIfEntityIsValid(FLWEntity Entity) const
+void UMassEntitySubsystem::CheckIfEntityIsValid(FMassEntityHandle Entity) const
 {
 	checkf(IsEntityValid(Entity), TEXT("Invalid entity (ID: %d, SN:%d, %s)"), Entity.Index, Entity.SerialNumber,
 		   (Entity.Index == 0) ? TEXT("was never initialized") : TEXT("already destroyed"));
 }
 
-void UMassEntitySubsystem::CheckIfEntityIsActive(FLWEntity Entity) const
+void UMassEntitySubsystem::CheckIfEntityIsActive(FMassEntityHandle Entity) const
 {
 	checkf(IsEntityBuilt(Entity), TEXT("Entity not yet created(ID: %d, SN:%d)"));
 }
 
-void UMassEntitySubsystem::GetValidArchetypes(const FLWComponentQuery& Query, TArray<FArchetypeHandle>& OutValidArchetypes)
+void UMassEntitySubsystem::GetValidArchetypes(const FMassEntityQuery& Query, TArray<FArchetypeHandle>& OutValidArchetypes)
 {
 	//@TODO: Not optimized yet, but we call this rarely now, so not a big deal.
 
 	// First get set of all archetypes that contain *any* component
-	TSet<TSharedPtr<FArchetypeData>> AnyArchetypes;
-	for (const FLWComponentRequirement& Requirement : Query.GetRequirements())
+	TSet<TSharedPtr<FMassArchetypeData>> AnyArchetypes;
+	for (const FMassFragmentRequirement& Requirement : Query.GetRequirements())
 	{
 		check(Requirement.StructType);
-		if (Requirement.Presence != ELWComponentPresence::None)
+		if (Requirement.Presence != EMassFragmentPresence::None)
 		{
-			if (TArray<TSharedPtr<FArchetypeData>>* pData = ComponentTypeToArchetypeMap.Find(Requirement.StructType))
+			if (TArray<TSharedPtr<FMassArchetypeData>>* pData = ComponentTypeToArchetypeMap.Find(Requirement.StructType))
 			{
 				AnyArchetypes.Append(*pData);
 			}
@@ -767,15 +767,15 @@ void UMassEntitySubsystem::GetValidArchetypes(const FLWComponentQuery& Query, TA
 	}
 
 	// Then verify that they contain *all* required components
-	for (TSharedPtr<FArchetypeData>& ArchetypePtr : AnyArchetypes)
+	for (TSharedPtr<FMassArchetypeData>& ArchetypePtr : AnyArchetypes)
 	{
-		FArchetypeData& Archetype = *(ArchetypePtr.Get());
+		FMassArchetypeData& Archetype = *(ArchetypePtr.Get());
 
 		if (Archetype.GetTagBitSet().HasAll(Query.GetRequiredAllTags()) == false)
 		{
 			// missing some required tags, skip.
 #if WITH_MASSENTITY_DEBUG
-			const FLWTagBitSet UnsatisfiedTags = Query.GetRequiredAllTags() - Archetype.GetTagBitSet();
+			const FMassTagBitSet UnsatisfiedTags = Query.GetRequiredAllTags() - Archetype.GetTagBitSet();
 			FStringOutputDevice Description;
 			UnsatisfiedTags.DebugGetStringDesc(Description);
 			UE_LOG(LogAggregateTicking, VeryVerbose, TEXT("Archetype did not match due to missing tags: %s")
@@ -788,7 +788,7 @@ void UMassEntitySubsystem::GetValidArchetypes(const FLWComponentQuery& Query, TA
 		{
 			// has some tags required to be absent
 #if WITH_MASSENTITY_DEBUG
-			const FLWTagBitSet UnwantedTags = Query.GetRequiredAllTags().GetOverlap(Archetype.GetTagBitSet());
+			const FMassTagBitSet UnwantedTags = Query.GetRequiredAllTags().GetOverlap(Archetype.GetTagBitSet());
 			FStringOutputDevice Description;
 			UnwantedTags.DebugGetStringDesc(Description);
 			UE_LOG(LogAggregateTicking, VeryVerbose, TEXT("Archetype has tags required absent: %s")
@@ -813,7 +813,7 @@ void UMassEntitySubsystem::GetValidArchetypes(const FLWComponentQuery& Query, TA
 		{
 			// missing some required components, skip.
 #if WITH_MASSENTITY_DEBUG
-			const FLWComponentBitSet UnsatisfiedComponents = Query.GetRequiredAllComponents() - Archetype.GetComponentBitSet();
+			const FMassFragmentBitSet UnsatisfiedComponents = Query.GetRequiredAllComponents() - Archetype.GetComponentBitSet();
 			FStringOutputDevice Description;
 			UnsatisfiedComponents.DebugGetStringDesc(Description);
 			UE_LOG(LogAggregateTicking, VeryVerbose, TEXT("Archetype did not match due to missing Components: %s")
@@ -826,7 +826,7 @@ void UMassEntitySubsystem::GetValidArchetypes(const FLWComponentQuery& Query, TA
 		{
 			// has some Components required to be absent
 #if WITH_MASSENTITY_DEBUG
-			const FLWComponentBitSet UnwantedComponents = Query.GetRequiredAllComponents().GetOverlap(Archetype.GetComponentBitSet());
+			const FMassFragmentBitSet UnwantedComponents = Query.GetRequiredAllComponents().GetOverlap(Archetype.GetComponentBitSet());
 			FStringOutputDevice Description;
 			UnwantedComponents.DebugGetStringDesc(Description);
 			UE_LOG(LogAggregateTicking, VeryVerbose, TEXT("Archetype has Components required absent: %s")
@@ -851,7 +851,7 @@ void UMassEntitySubsystem::GetValidArchetypes(const FLWComponentQuery& Query, TA
 		{
 			// missing some required components, skip.
 #if WITH_MASSENTITY_DEBUG
-			const FLWChunkComponentBitSet UnsatisfiedComponents = Query.GetRequiredAllChunkComponents() - Archetype.GetChunkComponentBitSet();
+			const FMassChunkFragmentBitSet UnsatisfiedComponents = Query.GetRequiredAllChunkComponents() - Archetype.GetChunkComponentBitSet();
 			FStringOutputDevice Description;
 			UnsatisfiedComponents.DebugGetStringDesc(Description);
 			UE_LOG(LogAggregateTicking, VeryVerbose, TEXT("Archetype did not match due to missing Chunk Components: %s")
@@ -864,7 +864,7 @@ void UMassEntitySubsystem::GetValidArchetypes(const FLWComponentQuery& Query, TA
 		{
 			// has some Components required to be absent
 #if WITH_MASSENTITY_DEBUG
-			const FLWChunkComponentBitSet UnwantedComponents = Query.GetRequiredNoneChunkComponents().GetOverlap(Archetype.GetChunkComponentBitSet());
+			const FMassChunkFragmentBitSet UnwantedComponents = Query.GetRequiredNoneChunkComponents().GetOverlap(Archetype.GetChunkComponentBitSet());
 			FStringOutputDevice Description;
 			UnwantedComponents.DebugGetStringDesc(Description);
 			UE_LOG(LogAggregateTicking, VeryVerbose, TEXT("Archetype has Chunk Components required absent: %s")
@@ -877,9 +877,9 @@ void UMassEntitySubsystem::GetValidArchetypes(const FLWComponentQuery& Query, TA
 	}
 }
 
-FLWComponentSystemExecutionContext UMassEntitySubsystem::CreateExecutionContext(const float DeltaSeconds) const
+FMassExecutionContext UMassEntitySubsystem::CreateExecutionContext(const float DeltaSeconds) const
 {
-	FLWComponentSystemExecutionContext ExecutionContext(DeltaSeconds);
+	FMassExecutionContext ExecutionContext(DeltaSeconds);
 	ExecutionContext.SetDeferredCommandBuffer(DeferredCommandBuffer);
 	return MoveTemp(ExecutionContext);
 }
@@ -894,7 +894,7 @@ void UMassEntitySubsystem::DebugPrintArchetypes(FOutputDevice& Ar) const
 	int32 LongestArchetypeBucket = 0;
 	for (const auto& KVP : ComponentHashToArchetypeMap)
 	{
-		for (const TSharedPtr<FArchetypeData>& ArchetypePtr : KVP.Value)
+		for (const TSharedPtr<FMassArchetypeData>& ArchetypePtr : KVP.Value)
 		{
 			ArchetypePtr->DebugPrintArchetype(Ar);
 		}
@@ -923,13 +923,13 @@ void UMassEntitySubsystem::DebugPrintEntity(int32 Index, FOutputDevice& Ar, cons
 		Ar.Logf(ELogVerbosity::Log, TEXT("Unable to list components values for invalid entity in %s"), *GetPathNameSafe(this));
 	}
 
-	FLWEntity Entity;
+	FMassEntityHandle Entity;
 	Entity.Index = Index;
 	Entity.SerialNumber = EntityData.SerialNumber;
 	DebugPrintEntity(Entity, Ar, InPrefix);
 }
 
-void UMassEntitySubsystem::DebugPrintEntity(FLWEntity Entity, FOutputDevice& Ar, const TCHAR* InPrefix) const
+void UMassEntitySubsystem::DebugPrintEntity(FMassEntityHandle Entity, FOutputDevice& Ar, const TCHAR* InPrefix) const
 {
 	if (!IsEntityActive(Entity))
 	{
@@ -939,7 +939,7 @@ void UMassEntitySubsystem::DebugPrintEntity(FLWEntity Entity, FOutputDevice& Ar,
 	Ar.Logf(ELogVerbosity::Log, TEXT("Listing components values for Entity[%s] in %s"), *Entity.DebugGetDescription(), *GetPathNameSafe(this));
 
 	const FEntityData& EntityData = Entities[Entity.Index];
-	FArchetypeData* Archetype = EntityData.CurrentArchetype.Get();
+	FMassArchetypeData* Archetype = EntityData.CurrentArchetype.Get();
 	if (Archetype == nullptr)
 	{
 		Ar.Logf(ELogVerbosity::Log, TEXT("Unable to list components values for invalid entity in %s"), *GetPathNameSafe(this));
@@ -962,7 +962,7 @@ void UMassEntitySubsystem::DebugGetArchetypesStringDetails(FOutputDevice& Ar, co
 	for (auto Pair : ComponentHashToArchetypeMap)
 	{
 		Ar.Logf(ELogVerbosity::Log, TEXT("\n-----------------------------------\nHash: %u"), Pair.Key);
-		for (TSharedPtr<FArchetypeData> Archetype : Pair.Value)
+		for (TSharedPtr<FMassArchetypeData> Archetype : Pair.Value)
 		{
 			if (Archetype.IsValid() && (bIncludeEmpty == true || Archetype->GetChunkCount() > 0))
 			{
@@ -997,8 +997,8 @@ void UMassEntitySubsystem::DebugRemoveAllEntities()
 			// already dead
 			continue;
 		}
-		const TSharedPtr<FArchetypeData>& Archetype = EntityData.CurrentArchetype;
-		FLWEntity Entity;
+		const TSharedPtr<FMassArchetypeData>& Archetype = EntityData.CurrentArchetype;
+		FMassEntityHandle Entity;
 		Entity.Index = EntityIndex;
 		Entity.SerialNumber = EntityData.SerialNumber;
 		Archetype->RemoveEntity(Entity);
@@ -1015,10 +1015,10 @@ void UMassEntitySubsystem::DebugGetArchetypeStrings(const FArchetypeHandle& Arch
 		return;
 	}
 
-	FArchetypeData& ArchetypeRef = *Archetype.DataPtr.Get();
+	FMassArchetypeData& ArchetypeRef = *Archetype.DataPtr.Get();
 	
 	OutComponentNames.Reserve(ArchetypeRef.GetComponentConfigs().Num());
-	for (const FArchetypeComponentConfig& ComponentConfig : ArchetypeRef.GetComponentConfigs())
+	for (const FMassArchetypeFragmentConfig& ComponentConfig : ArchetypeRef.GetComponentConfigs())
 	{
 		checkSlow(ComponentConfig.ComponentType);
 		OutComponentNames.Add(ComponentConfig.ComponentType->GetFName());
@@ -1048,9 +1048,9 @@ FAutoConsoleCommandWithWorldArgsAndOutputDevice GPrintArchetypesCmd(
 #endif // WITH_MASSENTITY_DEBUG
 
 //////////////////////////////////////////////////////////////////////
-// FLWComponentSystemExecutionContext
+// FMassExecutionContext
 
-void FLWComponentSystemExecutionContext::FlushDeferred(UMassEntitySubsystem& EntitySystem) const
+void FMassExecutionContext::FlushDeferred(UMassEntitySubsystem& EntitySystem) const
 {
 	if (bFlushDeferredCommands && DeferredCommandBuffer)
 	{
@@ -1058,47 +1058,47 @@ void FLWComponentSystemExecutionContext::FlushDeferred(UMassEntitySubsystem& Ent
 	}
 }
 
-void FLWComponentSystemExecutionContext::ClearExecutionData()
+void FMassExecutionContext::ClearExecutionData()
 {
 	ComponentViews.Reset();
 	CurrentArchetypesTagBitSet.Reset();
 	ChunkSerialModificationNumber = -1;
 }
 
-void FLWComponentSystemExecutionContext::SetCurrentArchetypeData(FArchetypeData& ArchetypeData)
+void FMassExecutionContext::SetCurrentArchetypeData(FMassArchetypeData& ArchetypeData)
 {
 	CurrentArchetypesTagBitSet = ArchetypeData.GetTagBitSet();
 }
 
-void FLWComponentSystemExecutionContext::SetChunkCollection(const FArchetypeChunkCollection& InChunkCollection)
+void FMassExecutionContext::SetChunkCollection(const FArchetypeChunkCollection& InChunkCollection)
 {
 	check(ChunkCollection.IsEmpty());
 	ChunkCollection = InChunkCollection;
 }
 
-void FLWComponentSystemExecutionContext::SetChunkCollection(FArchetypeChunkCollection&& InChunkCollection)
+void FMassExecutionContext::SetChunkCollection(FArchetypeChunkCollection&& InChunkCollection)
 {
 	check(ChunkCollection.IsEmpty());
 	ChunkCollection = MoveTemp(InChunkCollection);
 }
 
-void FLWComponentSystemExecutionContext::SetRequirements(TConstArrayView<FLWComponentRequirement> InRequirements, TConstArrayView<FLWComponentRequirement> InChunkRequirements)
+void FMassExecutionContext::SetRequirements(TConstArrayView<FMassFragmentRequirement> InRequirements, TConstArrayView<FMassFragmentRequirement> InChunkRequirements)
 { 
 	ComponentViews.Reset();
-	for (const FLWComponentRequirement& Requirement : InRequirements)
+	for (const FMassFragmentRequirement& Requirement : InRequirements)
 	{
 		if (Requirement.RequiresBinding())
 		{
-			ComponentViews.Add(FComponentView(Requirement));
+			ComponentViews.Add(FFragmentView(Requirement));
 		}
 	}
 
 	ChunkComponents.Reset();
-	for (const FLWComponentRequirement& Requirement : InChunkRequirements)
+	for (const FMassFragmentRequirement& Requirement : InChunkRequirements)
 	{
 		if (Requirement.RequiresBinding())
 		{
-			ChunkComponents.Add(FChunkComponentView(Requirement));
+			ChunkComponents.Add(FChunkFragmentView(Requirement));
 		}
 	}
 }
diff --git a/Engine/Plugins/Runtime/MassEntity/Source/MassEntity/Private/MassEntityTypes.cpp b/Engine/Plugins/Runtime/MassEntity/Source/MassEntity/Private/MassEntityTypes.cpp
--- a/Engine/Plugins/Runtime/MassEntity/Source/MassEntity/Private/MassEntityTypes.cpp
+++ b/Engine/Plugins/Runtime/MassEntity/Source/MassEntity/Private/MassEntityTypes.cpp
@@ -12,25 +12,25 @@
 DEFINE_LOG_CATEGORY(LogPipe);
 
 //----------------------------------------------------------------------//
-//  FPipeContext
+//  FMassProcessingContext
 //----------------------------------------------------------------------//
-FPipeContext::FPipeContext(UMassEntitySubsystem& InEntities, const float InDeltaSeconds)
+FMassProcessingContext::FMassProcessingContext(UMassEntitySubsystem& InEntities, const float InDeltaSeconds)
 	: EntitySubsystem(&InEntities), DeltaSeconds(InDeltaSeconds)
 {
 
 }
 
 //----------------------------------------------------------------------//
-//  FRuntimePipeline
+//  FMassRuntimePipeline
 //----------------------------------------------------------------------//
-void FRuntimePipeline::Reset()
+void FMassRuntimePipeline::Reset()
 {
 	Processors.Reset();
 }
 
-void FRuntimePipeline::Initialize(UObject& Owner)
+void FMassRuntimePipeline::Initialize(UObject& Owner)
 {
-	for (UPipeProcessor* Proc : Processors)
+	for (UMassProcessor* Proc : Processors)
 	{
 		if (Proc)
 		{
@@ -40,59 +40,59 @@ void FRuntimePipeline::Initialize(UObject& Owner)
 	}
 }
 
-void FRuntimePipeline::SetProcessors(TArray<UPipeProcessor*>&& InProcessors)
+void FMassRuntimePipeline::SetProcessors(TArray<UMassProcessor*>&& InProcessors)
 {
 	Processors = InProcessors;
 }
 
-void FRuntimePipeline::InitializeFromSchematics(TConstArrayView<TSoftObjectPtr<UPipeSchematic>> Schematics, UObject& InOwner)
+void FMassRuntimePipeline::InitializeFromSchematics(TConstArrayView<TSoftObjectPtr<UMassSchematic>> Schematics, UObject& InOwner)
 {
 	Reset();
 	
-	// @todo we'll sometimes end up with duplicated PipeProcessors in the resulting array. We need to come up with a consistent policy for handling that 
-	for (const TSoftObjectPtr<UPipeSchematic>& Schematic : Schematics)
+	// @todo we'll sometimes end up with duplicated MassProcessors in the resulting array. We need to come up with a consistent policy for handling that 
+	for (const TSoftObjectPtr<UMassSchematic>& Schematic : Schematics)
 	{
-		UPipeSchematic* Pipe = Schematic.LoadSynchronous();
+		UMassSchematic* Pipe = Schematic.LoadSynchronous();
 		if (Pipe)
 		{
 			AppendOrOverrideRuntimeProcessorCopies(Pipe->GetProcessors(), InOwner);
 		}
 		else
 		{
-			UE_LOG(LogPipe, Error, TEXT("Unable to resolve PipeSchematic %s while creating FRuntimePipeline"), *Schematic.GetLongPackageName());
+			UE_LOG(LogPipe, Error, TEXT("Unable to resolve MassSchematic %s while creating FMassRuntimePipeline"), *Schematic.GetLongPackageName());
 		}
 	}
 
 	Initialize(InOwner);
 }
 
-void FRuntimePipeline::CreateFromArray(TConstArrayView<const UPipeProcessor*> InProcessors, UObject& InOwner)
+void FMassRuntimePipeline::CreateFromArray(TConstArrayView<const UMassProcessor*> InProcessors, UObject& InOwner)
 {
 	Reset();
 	AppendOrOverrideRuntimeProcessorCopies(InProcessors, InOwner);
 }
 
-void FRuntimePipeline::InitializeFromArray(TConstArrayView<const UPipeProcessor*> InProcessors, UObject& InOwner)
+void FMassRuntimePipeline::InitializeFromArray(TConstArrayView<const UMassProcessor*> InProcessors, UObject& InOwner)
 {
 	CreateFromArray(InProcessors, InOwner);
 	Initialize(InOwner);
 }
 
-void FRuntimePipeline::InitializeFromClassArray(TConstArrayView<TSubclassOf<UPipeProcessor>> InProcessorClasses, UObject& InOwner)
+void FMassRuntimePipeline::InitializeFromClassArray(TConstArrayView<TSubclassOf<UMassProcessor>> InProcessorClasses, UObject& InOwner)
 {
 	Reset();
 
 	const UWorld* World = InOwner.GetWorld();
 	const EProcessorExecutionFlags WorldExecutionFlags = World ? UE::Pipe::Utils::GetProcessorExecutionFlagsForWold(*World) : EProcessorExecutionFlags::All;
 
-	for (const TSubclassOf<UPipeProcessor>& ProcessorClass : InProcessorClasses)
+	for (const TSubclassOf<UMassProcessor>& ProcessorClass : InProcessorClasses)
 	{
 		if (ProcessorClass)
 		{
-			UPipeProcessor* CDO = ProcessorClass.GetDefaultObject();
+			UMassProcessor* CDO = ProcessorClass.GetDefaultObject();
 			if (CDO && CDO->ShouldExecute(WorldExecutionFlags))
 			{
-				UPipeProcessor* ProcInstance = NewObject<UPipeProcessor>(&InOwner, ProcessorClass);
+				UMassProcessor* ProcInstance = NewObject<UMassProcessor>(&InOwner, ProcessorClass);
 				Processors.Add(ProcInstance);
 			}
 			else
@@ -105,25 +105,25 @@ void FRuntimePipeline::InitializeFromClassArray(TConstArrayView<TSubclassOf<UPip
 	Initialize(InOwner);
 }
 
-bool FRuntimePipeline::HasProcessorOfExactClass(TSubclassOf<UPipeProcessor> InClass) const
+bool FMassRuntimePipeline::HasProcessorOfExactClass(TSubclassOf<UMassProcessor> InClass) const
 {
 	UClass* TestClass = InClass.Get();
-	return Processors.FindByPredicate([TestClass](const UPipeProcessor* Proc){ return Proc != nullptr && Proc->GetClass() == TestClass; }) != nullptr;
+	return Processors.FindByPredicate([TestClass](const UMassProcessor* Proc){ return Proc != nullptr && Proc->GetClass() == TestClass; }) != nullptr;
 }
 
-void FRuntimePipeline::AppendUniqueRuntimeProcessorCopies(TConstArrayView<const UPipeProcessor*> InProcessors, UObject& InOwner)
+void FMassRuntimePipeline::AppendUniqueRuntimeProcessorCopies(TConstArrayView<const UMassProcessor*> InProcessors, UObject& InOwner)
 {
 	const UWorld* World = InOwner.GetWorld();
 	const EProcessorExecutionFlags WorldExecutionFlags = World ? UE::Pipe::Utils::GetProcessorExecutionFlagsForWold(*World) : EProcessorExecutionFlags::All;
 	const int32 StartingCount = Processors.Num();
 		
-	for (const UPipeProcessor* Proc : InProcessors)
+	for (const UMassProcessor* Proc : InProcessors)
 	{
 		if (Proc && Proc->ShouldExecute(WorldExecutionFlags)
 			&& (Proc->AllowDuplicates() || (HasProcessorOfExactClass(Proc->GetClass()) == false)))
 		{
 			// unfortunately the const cast is required since NewObject doesn't support const Template object
-			UPipeProcessor* ProcCopy = NewObject<UPipeProcessor>(&InOwner, Proc->GetClass(), FName(), RF_NoFlags, const_cast<UPipeProcessor*>(Proc));
+			UMassProcessor* ProcCopy = NewObject<UMassProcessor>(&InOwner, Proc->GetClass(), FName(), RF_NoFlags, const_cast<UMassProcessor*>(Proc));
 			Processors.Add(ProcCopy);
 		}
 #if WITH_MASSENTITY_DEBUG
@@ -143,24 +143,24 @@ void FRuntimePipeline::AppendUniqueRuntimeProcessorCopies(TConstArrayView<const
 
 	for (int32 NewProcIndex = StartingCount; NewProcIndex < Processors.Num(); ++NewProcIndex)
 	{
-		UPipeProcessor* Proc = Processors[NewProcIndex];
+		UMassProcessor* Proc = Processors[NewProcIndex];
 		check(Proc);
 		REDIRECT_OBJECT_TO_VLOG(Proc, &InOwner);
 		Proc->Initialize(InOwner);
 	}
 }
 
-void FRuntimePipeline::AppendOrOverrideRuntimeProcessorCopies(TConstArrayView<const UPipeProcessor*> InProcessors, UObject& InOwner)
+void FMassRuntimePipeline::AppendOrOverrideRuntimeProcessorCopies(TConstArrayView<const UMassProcessor*> InProcessors, UObject& InOwner)
 {
 	const UWorld* World = InOwner.GetWorld();
 	const EProcessorExecutionFlags WorldExecutionFlags = World ? UE::Pipe::Utils::GetProcessorExecutionFlagsForWold(*World) : EProcessorExecutionFlags::All;
 	
-	for (const UPipeProcessor* Proc : InProcessors)
+	for (const UMassProcessor* Proc : InProcessors)
 	{
 		if (Proc && Proc->ShouldExecute(WorldExecutionFlags))
 		{
 			// unfortunately the const cast is required since NewObject doesn't support const Template object
-			UPipeProcessor* ProcCopy = NewObject<UPipeProcessor>(&InOwner, Proc->GetClass(), FName(), RF_NoFlags, const_cast<UPipeProcessor*>(Proc));
+			UMassProcessor* ProcCopy = NewObject<UMassProcessor>(&InOwner, Proc->GetClass(), FName(), RF_NoFlags, const_cast<UMassProcessor*>(Proc));
 			check(ProcCopy);
 
 			if (ProcCopy->AllowDuplicates())
@@ -171,7 +171,7 @@ void FRuntimePipeline::AppendOrOverrideRuntimeProcessorCopies(TConstArrayView<co
 			else 
 			{
 				const UClass* TestClass = Proc->GetClass();
-				UPipeProcessor** PrevProcessor = Processors.FindByPredicate([TestClass, ProcCopy](const UPipeProcessor* Proc) {
+				UMassProcessor** PrevProcessor = Processors.FindByPredicate([TestClass, ProcCopy](const UMassProcessor* Proc) {
 					return Proc != nullptr && Proc->GetClass() == TestClass;
 				});
 
@@ -192,16 +192,16 @@ void FRuntimePipeline::AppendOrOverrideRuntimeProcessorCopies(TConstArrayView<co
 	}
 }
 
-void FRuntimePipeline::AppendProcessor(UPipeProcessor& Processor)
+void FMassRuntimePipeline::AppendProcessor(UMassProcessor& Processor)
 {
 	Processors.Add(&Processor);
 }
 
-UPipeCompositeProcessor* FRuntimePipeline::FindTopLevelGroupByName(FName GroupName)
+UMassCompositeProcessor* FMassRuntimePipeline::FindTopLevelGroupByName(FName GroupName)
 {
-	for (UPipeProcessor* Processor : Processors)
+	for (UMassProcessor* Processor : Processors)
 	{
-		UPipeCompositeProcessor* CompositeProcessor = Cast<UPipeCompositeProcessor>(Processor);
+		UMassCompositeProcessor* CompositeProcessor = Cast<UMassCompositeProcessor>(Processor);
 		if (CompositeProcessor && CompositeProcessor->GetGroupName() == GroupName)
 		{
 			return CompositeProcessor;
@@ -210,15 +210,15 @@ UPipeCompositeProcessor* FRuntimePipeline::FindTopLevelGroupByName(FName GroupNa
 	return nullptr;
 }
 
-void FRuntimePipeline::DebugOutputDescription(FOutputDevice& Ar) const
+void FMassRuntimePipeline::DebugOutputDescription(FOutputDevice& Ar) const
 {
 	UE::Pipe::Debug::DebugOutputDescription(Processors, Ar);
 }
 
-uint32 GetTypeHash(const FRuntimePipeline& Instance)
+uint32 GetTypeHash(const FMassRuntimePipeline& Instance)
 { 
 	uint32 Hash = 0;
-	for (const UPipeProcessor* Proc : Instance.Processors)
+	for (const UMassProcessor* Proc : Instance.Processors)
 	{
 		Hash = HashCombine(Hash, PointerHash(Proc));
 	}
diff --git a/Engine/Plugins/Runtime/MassEntity/Source/MassEntity/Private/MassExecutor.cpp b/Engine/Plugins/Runtime/MassEntity/Source/MassEntity/Private/MassExecutor.cpp
--- a/Engine/Plugins/Runtime/MassEntity/Source/MassEntity/Private/MassExecutor.cpp
+++ b/Engine/Plugins/Runtime/MassEntity/Source/MassEntity/Private/MassExecutor.cpp
@@ -8,7 +8,7 @@
 
 namespace UE::Pipe::Executor
 {
-void Run(FRuntimePipeline& RuntimePipeline, FPipeContext& PipeContext)
+void Run(FMassRuntimePipeline& RuntimePipeline, FMassProcessingContext& PipeContext)
 {
 	if (!ensure(PipeContext.EntitySubsystem) || 
 		!ensure(PipeContext.DeltaSeconds >= 0.f) ||
@@ -21,7 +21,7 @@ void Run(FRuntimePipeline& RuntimePipeline, FPipeContext& PipeContext)
 	RunProcessorsView(RuntimePipeline.Processors, PipeContext);
 }
 
-void RunSparse(FRuntimePipeline& RuntimePipeline, FPipeContext& PipeContext, FArchetypeHandle Archetype, TConstArrayView<FLWEntity> Entities)
+void RunSparse(FMassRuntimePipeline& RuntimePipeline, FMassProcessingContext& PipeContext, FArchetypeHandle Archetype, TConstArrayView<FMassEntityHandle> Entities)
 {
 	if (!ensure(PipeContext.EntitySubsystem) ||
 		!ensure(RuntimePipeline.Processors.Find(nullptr) == INDEX_NONE) ||
@@ -37,7 +37,7 @@ void RunSparse(FRuntimePipeline& RuntimePipeline, FPipeContext& PipeContext, FAr
 	RunProcessorsView(RuntimePipeline.Processors, PipeContext, &ChunkCollection);
 }
 
-void RunSparse(FRuntimePipeline& RuntimePipeline, FPipeContext& PipeContext, const FArchetypeChunkCollection& ChunkCollection)
+void RunSparse(FMassRuntimePipeline& RuntimePipeline, FMassProcessingContext& PipeContext, const FArchetypeChunkCollection& ChunkCollection)
 {
 	if (!ensure(PipeContext.EntitySubsystem) ||
 		!ensure(RuntimePipeline.Processors.Find(nullptr) == INDEX_NONE) ||
@@ -52,7 +52,7 @@ void RunSparse(FRuntimePipeline& RuntimePipeline, FPipeContext& PipeContext, con
 	RunProcessorsView(RuntimePipeline.Processors, PipeContext, &ChunkCollection);
 }
 
-void Run(UPipeProcessor& Processor, FPipeContext& PipeContext)
+void Run(UMassProcessor& Processor, FMassProcessingContext& PipeContext)
 {
 	if (!ensure(PipeContext.EntitySubsystem) || !ensure(PipeContext.DeltaSeconds >= 0.f))
 	{
@@ -61,11 +61,11 @@ void Run(UPipeProcessor& Processor, FPipeContext& PipeContext)
 
 	TRACE_CPUPROFILER_EVENT_SCOPE_STR("PipeExecutor Run")
 
-	UPipeProcessor* ProcPtr = &Processor;
+	UMassProcessor* ProcPtr = &Processor;
 	RunProcessorsView(MakeArrayView(&ProcPtr, 1), PipeContext);
 }
 
-void RunProcessorsView(TArrayView<UPipeProcessor*> Processors, FPipeContext& PipeContext, const FArchetypeChunkCollection* ChunkCollection)
+void RunProcessorsView(TArrayView<UMassProcessor*> Processors, FMassProcessingContext& PipeContext, const FArchetypeChunkCollection* ChunkCollection)
 {
 	TRACE_CPUPROFILER_EVENT_SCOPE(RunProcessorsView);
 
@@ -84,13 +84,13 @@ void RunProcessorsView(TArrayView<UPipeProcessor*> Processors, FPipeContext& Pip
 
 	TRACE_CPUPROFILER_EVENT_SCOPE_STR("PipeExecutor RunProcessorsView")
 
-	FLWComponentSystemExecutionContext ExecutionContext(PipeContext.DeltaSeconds);
+	FMassExecutionContext ExecutionContext(PipeContext.DeltaSeconds);
 	if (ChunkCollection)
 	{
 		ExecutionContext.SetChunkCollection(*ChunkCollection);
 	}
 	// manually creating a new command buffer to let the default one still be used by code unaware of pipe processing
-	ExecutionContext.SetDeferredCommandBuffer(MakeShareable(new FLWCCommandBuffer()));
+	ExecutionContext.SetDeferredCommandBuffer(MakeShareable(new FMassCommandBuffer()));
 	ExecutionContext.SetFlushDeferredCommands(false);
 	ExecutionContext.SetAuxData(PipeContext.AuxData);
 
@@ -99,7 +99,7 @@ void RunProcessorsView(TArrayView<UPipeProcessor*> Processors, FPipeContext& Pip
 		
 		UMassEntitySubsystem::FScopedProcessing ProcessingScope = PipeContext.EntitySubsystem->NewProcessingScope();
 
-		for (UPipeProcessor* Proc : Processors)
+		for (UMassProcessor* Proc : Processors)
 		{
 			Proc->CallExecute(*PipeContext.EntitySubsystem, ExecutionContext);
 		}
@@ -117,7 +117,7 @@ void RunProcessorsView(TArrayView<UPipeProcessor*> Processors, FPipeContext& Pip
 
 struct FPipeExecutorDoneTask
 {
-	FPipeExecutorDoneTask(const FLWComponentSystemExecutionContext& InExecutionContext, UMassEntitySubsystem& InEntitySubsystem, TFunction<void()> InOnDoneNotification, const FString& InDebugName)
+	FPipeExecutorDoneTask(const FMassExecutionContext& InExecutionContext, UMassEntitySubsystem& InEntitySubsystem, TFunction<void()> InOnDoneNotification, const FString& InDebugName)
 		: ExecutionContext(InExecutionContext)
 		, EntitySubsystem(InEntitySubsystem)
 		, OnDoneNotification(InOnDoneNotification)
@@ -147,13 +147,13 @@ struct FPipeExecutorDoneTask
 		OnDoneNotification();
 	}
 private:
-	FLWComponentSystemExecutionContext ExecutionContext;
+	FMassExecutionContext ExecutionContext;
 	UMassEntitySubsystem& EntitySubsystem;
 	TFunction<void()> OnDoneNotification;
 	FString DebugName;
 };
 
-FGraphEventRef TriggerParallelTasks(UPipeProcessor& Processor, FPipeContext& PipeContext, TFunction<void()> OnDoneNotification)
+FGraphEventRef TriggerParallelTasks(UMassProcessor& Processor, FMassProcessingContext& PipeContext, TFunction<void()> OnDoneNotification)
 {
 	TRACE_CPUPROFILER_EVENT_SCOPE(RunProcessorsView);
 
@@ -166,8 +166,8 @@ FGraphEventRef TriggerParallelTasks(UPipeProcessor& Processor, FPipeContext& Pip
 	TRACE_CPUPROFILER_EVENT_SCOPE_STR("PipeExecutor RunParallel")
 
 	// not going through UMassEntitySubsystem::CreateExecutionContext on purpose - we do need a separate command buffer
-	FLWComponentSystemExecutionContext ExecutionContext(PipeContext.DeltaSeconds);
-	ExecutionContext.SetDeferredCommandBuffer(MakeShareable(new FLWCCommandBuffer()));
+	FMassExecutionContext ExecutionContext(PipeContext.DeltaSeconds);
+	ExecutionContext.SetDeferredCommandBuffer(MakeShareable(new FMassCommandBuffer()));
 	ExecutionContext.SetFlushDeferredCommands(false);
 	ExecutionContext.SetAuxData(PipeContext.AuxData);
 
diff --git a/Engine/Plugins/Runtime/MassEntity/Source/MassEntity/Private/MassProcessingPhase.cpp b/Engine/Plugins/Runtime/MassEntity/Source/MassEntity/Private/MassProcessingPhase.cpp
--- a/Engine/Plugins/Runtime/MassEntity/Source/MassEntity/Private/MassProcessingPhase.cpp
+++ b/Engine/Plugins/Runtime/MassEntity/Source/MassEntity/Private/MassProcessingPhase.cpp
@@ -11,9 +11,9 @@
 #include "Engine/World.h"
 #include "LWCCommandBuffer.h"
 
-#define LOCTEXT_NAMESPACE "Pipe"
+#define LOCTEXT_NAMESPACE "Mass"
 
-DECLARE_CYCLE_STAT(TEXT("Pipe Phase Done"), STAT_PipePhaseDone, STATGROUP_TaskGraphTasks);
+DECLARE_CYCLE_STAT(TEXT("Mass Phase Done"), STAT_MassPhaseDone, STATGROUP_TaskGraphTasks);
 
 namespace FPipeTweakables
 {
@@ -25,22 +25,22 @@ namespace FPipeTweakables
 }
 
 //----------------------------------------------------------------------//
-//  FPipeProcessingPhase
+//  FMassProcessingPhase
 //----------------------------------------------------------------------//
-FPipeProcessingPhase::FPipeProcessingPhase()
+FMassProcessingPhase::FMassProcessingPhase()
 {
 	bCanEverTick = true;
 	bStartWithTickEnabled = false;
 }
 
-void FPipeProcessingPhase::ExecuteTick(float DeltaTime, ELevelTick TickType, ENamedThreads::Type CurrentThread, const FGraphEventRef& MyCompletionGraphEvent)
+void FMassProcessingPhase::ExecuteTick(float DeltaTime, ELevelTick TickType, ENamedThreads::Type CurrentThread, const FGraphEventRef& MyCompletionGraphEvent)
 {
 	if (TickType == LEVELTICK_ViewportsOnly || TickType == LEVELTICK_PauseTick)
 	{
 		return;
 	}
 
-	checkf(Manager, TEXT("Manager is null which is not a supported case. Either this FPipeProcessingPhase has not been initialized properly or it's been left dangling after the FPipeProcessingPhase owner got destroyed."));
+	checkf(Manager, TEXT("Manager is null which is not a supported case. Either this FMassProcessingPhase has not been initialized properly or it's been left dangling after the FMassProcessingPhase owner got destroyed."));
 
 	Manager->OnPhaseStart(*this);
 	
@@ -48,7 +48,7 @@ void FPipeProcessingPhase::ExecuteTick(float DeltaTime, ELevelTick TickType, ENa
 
 	check(PhaseProcessor);
 	
-	FPipeContext Context(Manager->GetEntitySubsystemRef(), DeltaTime);
+	FMassProcessingContext Context(Manager->GetEntitySubsystemRef(), DeltaTime);
 
 	if (bRunInParallelMode)
 	{
@@ -77,7 +77,7 @@ void FPipeProcessingPhase::ExecuteTick(float DeltaTime, ELevelTick TickType, ENa
 	}
 }
 
-void FPipeProcessingPhase::OnParallelExecutionDone(const float DeltaTime)
+void FMassProcessingPhase::OnParallelExecutionDone(const float DeltaTime)
 {
 	bIsDuringPipeProcessing = false;
 	OnPhaseEnd.Broadcast(DeltaTime);
@@ -85,20 +85,20 @@ void FPipeProcessingPhase::OnParallelExecutionDone(const float DeltaTime)
 	Manager->OnPhaseEnd(*this);
 }
 
-FString FPipeProcessingPhase::DiagnosticMessage()
+FString FMassProcessingPhase::DiagnosticMessage()
 {
 	return (Manager ? Manager->GetFullName() : TEXT("NULL-PipeProcessingPhaseManager")) + TEXT("[ProcessorTick]");
 }
 
-FName FPipeProcessingPhase::DiagnosticContext(bool bDetailed)
+FName FMassProcessingPhase::DiagnosticContext(bool bDetailed)
 {
 	return Manager ? Manager->GetClass()->GetFName() : TEXT("NULL-PipeProcessingPhaseManager");
 }
 
 //----------------------------------------------------------------------//
-// UPipeProcessingPhaseManager  
+// UMassProcessingPhaseManager  
 //----------------------------------------------------------------------//
-void UPipeProcessingPhaseManager::PostInitProperties()
+void UMassProcessingPhaseManager::PostInitProperties()
 {
 	Super::PostInitProperties();
 	if (HasAnyFlags(RF_ClassDefaultObject) == false)
@@ -107,47 +107,47 @@ void UPipeProcessingPhaseManager::PostInitProperties()
 		UWorld* World = GetWorld();
 		if (World && World->IsGameWorld() == false)
 		{
-			UPipeSettings* PipeSettings = GetMutableDefault<UPipeSettings>();
+			UMassSettings* PipeSettings = GetMutableDefault<UMassSettings>();
 			check(PipeSettings);
-			PipeSettingsChangeHandle = PipeSettings->GetOnSettingsChange().AddUObject(this, &UPipeProcessingPhaseManager::OnPipeSettingsChange);
+			PipeSettingsChangeHandle = PipeSettings->GetOnSettingsChange().AddUObject(this, &UMassProcessingPhaseManager::OnPipeSettingsChange);
 		}
 #endif // WITH_EDITOR
 	}
 	CreatePhases();
 }
 
-void UPipeProcessingPhaseManager::BeginDestroy()
+void UMassProcessingPhaseManager::BeginDestroy()
 {
 	Stop();
 	Super::BeginDestroy();
 
 #if WITH_EDITOR
-	if (UPipeSettings* PipeSettings = GetMutableDefault<UPipeSettings>())
+	if (UMassSettings* PipeSettings = GetMutableDefault<UMassSettings>())
 	{
 		PipeSettings->GetOnSettingsChange().Remove(PipeSettingsChangeHandle);
 	}
 #endif // WITH_EDITOR
 }
 
-void UPipeProcessingPhaseManager::InitializePhases(UObject& InProcessorOwner)
+void UMassProcessingPhaseManager::InitializePhases(UObject& InProcessorOwner)
 {
-	const FPipeProcessingPhaseConfig* ProcessingPhasesConfig = GET_PIPE_CONFIG_VALUE(GetProcessingPhasesConfig());
+	const FMassProcessingPhaseConfig* ProcessingPhasesConfig = GET_PIPE_CONFIG_VALUE(GetProcessingPhasesConfig());
 
 	FString DependencyGraphFileName;
 
 #if WITH_EDITOR
 	const UWorld* World = InProcessorOwner.GetWorld();
-	const UPipeSettings* PipeSettings = GetMutableDefault<UPipeSettings>();
+	const UMassSettings* PipeSettings = GetMutableDefault<UMassSettings>();
 	if (World != nullptr && PipeSettings != nullptr && !PipeSettings->DumpDependencyGraphFileName.IsEmpty())
 	{
 		DependencyGraphFileName = FString::Printf(TEXT("%s_%s"), *PipeSettings->DumpDependencyGraphFileName,*ToString(World->GetNetMode()));
 	}
 #endif // WITH_EDITOR
 
-	for (int i = 0; i < int(EPipeProcessingPhase::MAX); ++i)
+	for (int i = 0; i < int(EMassProcessingPhase::MAX); ++i)
 	{
-		const FPipeProcessingPhaseConfig& PhaseConfig = ProcessingPhasesConfig[i];
-		UPipeCompositeProcessor* PhaseProcessor = ProcessingPhases[i].PhaseProcessor;
+		const FMassProcessingPhaseConfig& PhaseConfig = ProcessingPhasesConfig[i];
+		UMassCompositeProcessor* PhaseProcessor = ProcessingPhases[i].PhaseProcessor;
 		FString FileName = !DependencyGraphFileName.IsEmpty() ? FString::Printf(TEXT("%s_%s"), *DependencyGraphFileName, *PhaseConfig.PhaseName.ToString()) : FString();
 		PhaseProcessor->CopyAndSort(PhaseConfig, FileName);
 		PhaseProcessor->Initialize(InProcessorOwner);
@@ -158,7 +158,7 @@ void UPipeProcessingPhaseManager::InitializePhases(UObject& InProcessorOwner)
 	UE_VLOG_UELOG(this, LogPipe, Verbose, TEXT("Phases initialization done. Current composition:"));
 
 	FStringOutputDevice OutDescription;
-	for (int i = 0; i < int(EPipeProcessingPhase::MAX); ++i)
+	for (int i = 0; i < int(EMassProcessingPhase::MAX); ++i)
 	{
 		if (ProcessingPhases[i].PhaseProcessor)
 		{
@@ -170,7 +170,7 @@ void UPipeProcessingPhaseManager::InitializePhases(UObject& InProcessorOwner)
 #endif // WITH_MASSENTITY_DEBUG
 }
 
-void UPipeProcessingPhaseManager::Start(UWorld& World)
+void UMassProcessingPhaseManager::Start(UWorld& World)
 {
 	EntitySubsystem = UWorld::GetSubsystem<UMassEntitySubsystem>(&World);
 
@@ -184,7 +184,7 @@ void UPipeProcessingPhaseManager::Start(UWorld& World)
 	}
 }
 
-void UPipeProcessingPhaseManager::Start(UMassEntitySubsystem& InEntitySubsystem)
+void UMassProcessingPhaseManager::Start(UMassEntitySubsystem& InEntitySubsystem)
 {
 	UWorld* World = InEntitySubsystem.GetWorld();
 	check(World);
@@ -192,10 +192,10 @@ void UPipeProcessingPhaseManager::Start(UMassEntitySubsystem& InEntitySubsystem)
 	EnableTickFunctions(*World);
 }
 
-void UPipeProcessingPhaseManager::EnableTickFunctions(const UWorld& World)
+void UMassProcessingPhaseManager::EnableTickFunctions(const UWorld& World)
 {
 	check(EntitySubsystem);
-	for (FPipeProcessingPhase& Phase : ProcessingPhases)
+	for (FMassProcessingPhase& Phase : ProcessingPhases)
 	{
 		Phase.Manager = this;
 		Phase.RegisterTickFunction(World.PersistentLevel);
@@ -205,11 +205,11 @@ void UPipeProcessingPhaseManager::EnableTickFunctions(const UWorld& World)
 		, *GetNameSafe(GetOuter()), *GetName());
 }
 
-void UPipeProcessingPhaseManager::Stop()
+void UMassProcessingPhaseManager::Stop()
 {
 	EntitySubsystem = nullptr;
 	
-	for (FPipeProcessingPhase& Phase : ProcessingPhases)
+	for (FMassProcessingPhase& Phase : ProcessingPhases)
 	{
 		Phase.SetTickFunctionEnable(false);
 	}
@@ -218,41 +218,41 @@ void UPipeProcessingPhaseManager::Stop()
 		, *GetNameSafe(GetOuter()), *GetName());
 }
 
-void UPipeProcessingPhaseManager::CreatePhases() 
+void UMassProcessingPhaseManager::CreatePhases() 
 {
-	static ETickingGroup PhaseToTickingGroup[int(EPipeProcessingPhase::MAX)]
+	static ETickingGroup PhaseToTickingGroup[int(EMassProcessingPhase::MAX)]
 	{
-		ETickingGroup::TG_PrePhysics, // EPipeProcessingPhase::PrePhysics
-		ETickingGroup::TG_StartPhysics, // EPipeProcessingPhase::StartPhysics
-		ETickingGroup::TG_DuringPhysics, // EPipeProcessingPhase::DuringPhysics
-		ETickingGroup::TG_EndPhysics,	// EPipeProcessingPhase::EndPhysics
-		ETickingGroup::TG_PostPhysics,	// EPipeProcessingPhase::PostPhysics
-		ETickingGroup::TG_LastDemotable, // EPipeProcessingPhase::FrameEnd
+		ETickingGroup::TG_PrePhysics, // EMassProcessingPhase::PrePhysics
+		ETickingGroup::TG_StartPhysics, // EMassProcessingPhase::StartPhysics
+		ETickingGroup::TG_DuringPhysics, // EMassProcessingPhase::DuringPhysics
+		ETickingGroup::TG_EndPhysics,	// EMassProcessingPhase::EndPhysics
+		ETickingGroup::TG_PostPhysics,	// EMassProcessingPhase::PostPhysics
+		ETickingGroup::TG_LastDemotable, // EMassProcessingPhase::FrameEnd
 	};
 
 	// @todo copy from settings instead of blindly creating from scratch
-	for (int i = 0; i < int(EPipeProcessingPhase::MAX); ++i)
+	for (int i = 0; i < int(EMassProcessingPhase::MAX); ++i)
 	{
-		ProcessingPhases[i].Phase = EPipeProcessingPhase(i);
+		ProcessingPhases[i].Phase = EMassProcessingPhase(i);
 		ProcessingPhases[i].TickGroup = PhaseToTickingGroup[i];
-		ProcessingPhases[i].PhaseProcessor = NewObject<UPipeCompositeProcessor>(this, UPipeCompositeProcessor::StaticClass()
-			, *FString::Printf(TEXT("ProcessingPhase_%s"), *EnumToString(EPipeProcessingPhase(i))));
+		ProcessingPhases[i].PhaseProcessor = NewObject<UMassCompositeProcessor>(this, UMassCompositeProcessor::StaticClass()
+			, *FString::Printf(TEXT("ProcessingPhase_%s"), *EnumToString(EMassProcessingPhase(i))));
 		REDIRECT_OBJECT_TO_VLOG(ProcessingPhases[i].PhaseProcessor, this);
-		ProcessingPhases[i].PhaseProcessor->SetGroupName(*EnumToString(EPipeProcessingPhase(i)));
-		ProcessingPhases[i].PhaseProcessor->SetProcessingPhase(EPipeProcessingPhase(i));
+		ProcessingPhases[i].PhaseProcessor->SetGroupName(*EnumToString(EMassProcessingPhase(i)));
+		ProcessingPhases[i].PhaseProcessor->SetProcessingPhase(EMassProcessingPhase(i));
 	}
 }
 
-void UPipeProcessingPhaseManager::OnPhaseStart(const FPipeProcessingPhase& Phase)
+void UMassProcessingPhaseManager::OnPhaseStart(const FMassProcessingPhase& Phase)
 {
-	ensure(CurrentPhase == EPipeProcessingPhase::MAX);
+	ensure(CurrentPhase == EMassProcessingPhase::MAX);
 	CurrentPhase = Phase.Phase;
 }
 
-void UPipeProcessingPhaseManager::OnPhaseEnd(FPipeProcessingPhase& Phase)
+void UMassProcessingPhaseManager::OnPhaseEnd(FMassProcessingPhase& Phase)
 {
 	ensure(CurrentPhase == Phase.Phase);
-	CurrentPhase = EPipeProcessingPhase::MAX;
+	CurrentPhase = EMassProcessingPhase::MAX;
 
 	// switch between parallel and single-thread versions only after a given batch of processing has been wrapped up	
 	if (Phase.IsConfiguredForParallelMode() != FPipeTweakables::bFullyParallel)
@@ -269,7 +269,7 @@ void UPipeProcessingPhaseManager::OnPhaseEnd(FPipeProcessingPhase& Phase)
 }
 
 #if WITH_EDITOR
-void UPipeProcessingPhaseManager::OnPipeSettingsChange(const FPropertyChangedEvent& PropertyChangedEvent)
+void UMassProcessingPhaseManager::OnPipeSettingsChange(const FPropertyChangedEvent& PropertyChangedEvent)
 {
 	check(GetOuter());
 	InitializePhases(*GetOuter());
diff --git a/Engine/Plugins/Runtime/MassEntity/Source/MassEntity/Private/MassProcessor.cpp b/Engine/Plugins/Runtime/MassEntity/Source/MassEntity/Private/MassProcessor.cpp
--- a/Engine/Plugins/Runtime/MassEntity/Source/MassEntity/Private/MassProcessor.cpp
+++ b/Engine/Plugins/Runtime/MassEntity/Source/MassEntity/Private/MassProcessor.cpp
@@ -7,7 +7,7 @@
 #include "Engine/World.h"
 #include "LWCCommandBuffer.h"
 
-DECLARE_CYCLE_STAT(TEXT("PipeProcessor Group Completed"), Pipe_GroupCompletedTask, STATGROUP_TaskGraphTasks);
+DECLARE_CYCLE_STAT(TEXT("MassProcessor Group Completed"), Pipe_GroupCompletedTask, STATGROUP_TaskGraphTasks);
 
 #define PARALLELIZED_TRAFFIC_HACK 1
 
@@ -40,16 +40,16 @@ namespace FPipeTweakables
 	bool bParallelGroups = false;
 	float PostponedTaskWaitTimeWarningLevel = 0.002f;
 
-	FAutoConsoleVariableRef CVarsPipeProcessor[] = {
+	FAutoConsoleVariableRef CVarsMassProcessor[] = {
 		{TEXT("pipe.ParallelGroups"), bParallelGroups, TEXT("Enables pipe processing groups distribution to all available threads (via the task graph)")},
 		{TEXT("pipe.PostponedTaskWaitTimeWarningLevel"), PostponedTaskWaitTimeWarningLevel, TEXT("if waiting for postponed task\'s dependencies exceeds this number an error will be logged")},
 	};
 }
 
-class FPipeProcessorTask
+class FMassProcessorTask
 {
 public:
-	FPipeProcessorTask(UMassEntitySubsystem& InEntitySubsystem, const FLWComponentSystemExecutionContext& InExecutionContext, UPipeProcessor& InProc, bool bInManageCommandBuffer = true)
+	FMassProcessorTask(UMassEntitySubsystem& InEntitySubsystem, const FMassExecutionContext& InExecutionContext, UMassProcessor& InProc, bool bInManageCommandBuffer = true)
 		: EntitySubsystem(&InEntitySubsystem)
 		, ExecutionContext(InExecutionContext)
 		, Processor(&InProc)
@@ -58,7 +58,7 @@ public:
 
 	static TStatId GetStatId()
 	{
-		RETURN_QUICK_DECLARE_CYCLE_STAT(FPipeProcessorTask, STATGROUP_TaskGraphTasks);
+		RETURN_QUICK_DECLARE_CYCLE_STAT(FMassProcessorTask, STATGROUP_TaskGraphTasks);
 	}
 
 	static ESubsequentsMode::Type GetSubsequentsMode() { return ESubsequentsMode::TrackSubsequents; }
@@ -81,8 +81,8 @@ public:
 		
 		if (bManageCommandBuffer)
 		{
-			TSharedPtr<FLWCCommandBuffer> MainSharedPtr = ExecutionContext.GetSharedDeferredCommandBuffer();
-			ExecutionContext.SetDeferredCommandBuffer(MakeShareable(new FLWCCommandBuffer()));
+			TSharedPtr<FMassCommandBuffer> MainSharedPtr = ExecutionContext.GetSharedDeferredCommandBuffer();
+			ExecutionContext.SetDeferredCommandBuffer(MakeShareable(new FMassCommandBuffer()));
 			Processor->CallExecute(*EntitySubsystem, ExecutionContext);
 			MainSharedPtr->MoveAppend(ExecutionContext.Defer());
 		}
@@ -95,8 +95,8 @@ public:
 
 private:
 	UMassEntitySubsystem* EntitySubsystem = nullptr;
-	FLWComponentSystemExecutionContext ExecutionContext;
-	UPipeProcessor* Processor = nullptr;
+	FMassExecutionContext ExecutionContext;
+	UMassProcessor* Processor = nullptr;
 	/** 
 	 * indicates whether this task is responsible for creation of a dedicated command buffer and transferring over the 
 	 * commands after processor's execution;
@@ -104,11 +104,11 @@ private:
 	bool bManageCommandBuffer = true;
 };
 
-class FPipeProcessorsTask_GameThread : public FPipeProcessorTask
+class FMassProcessorsTask_GameThread : public FMassProcessorTask
 {
 public:
-	FPipeProcessorsTask_GameThread(UMassEntitySubsystem& InEntitySubsystem, const FLWComponentSystemExecutionContext& InExecutionContext, UPipeProcessor& InProc)
-		: FPipeProcessorTask(InEntitySubsystem, InExecutionContext, InProc)
+	FMassProcessorsTask_GameThread(UMassEntitySubsystem& InEntitySubsystem, const FMassExecutionContext& InExecutionContext, UMassProcessor& InProc)
+		: FMassProcessorTask(InEntitySubsystem, InExecutionContext, InProc)
 	{}
 
 	static ENamedThreads::Type GetDesiredThread()
@@ -118,14 +118,14 @@ public:
 };
 
 //----------------------------------------------------------------------//
-// UPipeProcessor 
+// UMassProcessor 
 //----------------------------------------------------------------------//
-UPipeProcessor::UPipeProcessor(const FObjectInitializer& ObjectInitializer)
+UMassProcessor::UMassProcessor(const FObjectInitializer& ObjectInitializer)
 	: ExecutionFlags((int32)(EProcessorExecutionFlags::Server | EProcessorExecutionFlags::Standalone))
 {
 }
 
-void UPipeProcessor::SetShouldAutoRegisterWithGlobalList(const bool bAutoRegister)
+void UMassProcessor::SetShouldAutoRegisterWithGlobalList(const bool bAutoRegister)
 {	
 	if (ensureMsgf(HasAnyFlags(RF_ClassDefaultObject), TEXT("Setting bAutoRegisterWithProcessingPhases for non-CDOs has no effect")))
 	{
@@ -136,7 +136,7 @@ void UPipeProcessor::SetShouldAutoRegisterWithGlobalList(const bool bAutoRegiste
 	}
 }
 
-void UPipeProcessor::PostInitProperties()
+void UMassProcessor::PostInitProperties()
 {
 	Super::PostInitProperties();
 	if (HasAnyFlags(RF_ClassDefaultObject) == false) 
@@ -148,7 +148,7 @@ void UPipeProcessor::PostInitProperties()
 #endif
 }
 
-void UPipeProcessor::CallExecute(UMassEntitySubsystem& EntitySubsystem, FLWComponentSystemExecutionContext& Context)
+void UMassProcessor::CallExecute(UMassEntitySubsystem& EntitySubsystem, FMassExecutionContext& Context)
 {
 	TRACE_CPUPROFILER_EVENT_SCOPE_TEXT(*StatId);
 #if WITH_MASSENTITY_DEBUG
@@ -157,21 +157,21 @@ void UPipeProcessor::CallExecute(UMassEntitySubsystem& EntitySubsystem, FLWCompo
 	Execute(EntitySubsystem, Context);
 }
 
-FGraphEventRef UPipeProcessor::DispatchProcessorTasks(UMassEntitySubsystem& EntitySubsystem, FLWComponentSystemExecutionContext& ExecutionContext, const FGraphEventArray& Prerequisites)
+FGraphEventRef UMassProcessor::DispatchProcessorTasks(UMassEntitySubsystem& EntitySubsystem, FMassExecutionContext& ExecutionContext, const FGraphEventArray& Prerequisites)
 {
 	FGraphEventRef ReturnVal;
 	if (bRequiresGameThreadExecution)
 	{
-		ReturnVal = TGraphTask<FPipeProcessorsTask_GameThread>::CreateTask(&Prerequisites).ConstructAndDispatchWhenReady(EntitySubsystem, ExecutionContext, *this);
+		ReturnVal = TGraphTask<FMassProcessorsTask_GameThread>::CreateTask(&Prerequisites).ConstructAndDispatchWhenReady(EntitySubsystem, ExecutionContext, *this);
 	}
 	else
 	{
-		ReturnVal = TGraphTask<FPipeProcessorTask>::CreateTask(&Prerequisites).ConstructAndDispatchWhenReady(EntitySubsystem, ExecutionContext, *this);
+		ReturnVal = TGraphTask<FMassProcessorTask>::CreateTask(&Prerequisites).ConstructAndDispatchWhenReady(EntitySubsystem, ExecutionContext, *this);
 	}	
 	return ReturnVal;
 }
 
-void UPipeProcessor::DebugOutputDescription(FOutputDevice& Ar, int32 Indent) const
+void UMassProcessor::DebugOutputDescription(FOutputDevice& Ar, int32 Indent) const
 {
 #if WITH_MASSENTITY_DEBUG
 	Ar.Logf(TEXT("%*s%s"), Indent, TEXT(""), *GetProcessorName());
@@ -179,7 +179,7 @@ void UPipeProcessor::DebugOutputDescription(FOutputDevice& Ar, int32 Indent) con
 }
 
 #if WITH_EDITOR
-void UPipeProcessor::PostEditChangeProperty(FPropertyChangedEvent& PropertyChangedEvent)
+void UMassProcessor::PostEditChangeProperty(FPropertyChangedEvent& PropertyChangedEvent)
 {
 	Super::PostEditChangeProperty(PropertyChangedEvent);
 
@@ -192,9 +192,9 @@ void UPipeProcessor::PostEditChangeProperty(FPropertyChangedEvent& PropertyChang
 #endif // WITH_EDITOR
 
 //----------------------------------------------------------------------//
-//  UPipeCompositeProcessor
+//  UMassCompositeProcessor
 //----------------------------------------------------------------------//
-UPipeCompositeProcessor::UPipeCompositeProcessor()
+UMassCompositeProcessor::UMassCompositeProcessor()
 	: GroupName(TEXT("None"))
 	, bRunInSeparateThread(false)
 	, bHasOffThreadSubGroups(false)
@@ -205,18 +205,18 @@ UPipeCompositeProcessor::UPipeCompositeProcessor()
 	bAutoRegisterWithProcessingPhases = false;
 }
 
-void UPipeCompositeProcessor::SetChildProcessors(TArray<UPipeProcessor*>&& InProcessors)
+void UMassCompositeProcessor::SetChildProcessors(TArray<UMassProcessor*>&& InProcessors)
 {
 	ChildPipeline.SetProcessors(MoveTemp(InProcessors));
 }
 
-void UPipeCompositeProcessor::ConfigureQueries()
+void UMassCompositeProcessor::ConfigureQueries()
 {
 	// nothing to do here since ConfigureQueries will get automatically called for all the processors in ChildPipeline
 	// via their individual PostInitProperties call
 }
 
-FGraphEventRef UPipeCompositeProcessor::DispatchProcessorTasks(UMassEntitySubsystem& EntitySubsystem, FLWComponentSystemExecutionContext& ExecutionContext, const FGraphEventArray& InPrerequisites)
+FGraphEventRef UMassCompositeProcessor::DispatchProcessorTasks(UMassEntitySubsystem& EntitySubsystem, FMassExecutionContext& ExecutionContext, const FGraphEventArray& InPrerequisites)
 {
 	FGraphEventArray Events;
 	Events.Reserve(ProcessingFlatGraph.Num());
@@ -273,29 +273,29 @@ FGraphEventRef UPipeCompositeProcessor::DispatchProcessorTasks(UMassEntitySubsys
 	return CompletionEvent;
 }
 
-void UPipeCompositeProcessor::Execute(UMassEntitySubsystem& EntitySubsystem, FLWComponentSystemExecutionContext& Context)
+void UMassCompositeProcessor::Execute(UMassEntitySubsystem& EntitySubsystem, FMassExecutionContext& Context)
 {
 #if PARALLELIZED_TRAFFIC_HACK
 	if (FPipeTweakables::bParallelGroups == false
-		&& UE::MassTraffic::bParallelizeTraffic && GetProcessingPhase() == EPipeProcessingPhase::PrePhysics)
+		&& UE::MassTraffic::bParallelizeTraffic && GetProcessingPhase() == EMassProcessingPhase::PrePhysics)
 	{
 		static FName TrafficGroup(TEXT("Traffic"));
-		FLWComponentSystemExecutionContext TrafficExecutionContext;
+		FMassExecutionContext TrafficExecutionContext;
 		FGraphEventRef TrafficCompletionEvent;
 
-		for (UPipeProcessor* Proc : ChildPipeline.Processors)
+		for (UMassProcessor* Proc : ChildPipeline.Processors)
 		{
 			check(Proc);
-				if (GetProcessingPhase() == EPipeProcessingPhase::PrePhysics)
+				if (GetProcessingPhase() == EMassProcessingPhase::PrePhysics)
 			{
-				if (const UPipeCompositeProcessor* CompProcessor = Cast<UPipeCompositeProcessor>(Proc))
+				if (const UMassCompositeProcessor* CompProcessor = Cast<UMassCompositeProcessor>(Proc))
 				{
 					if (CompProcessor->GetGroupName() == TrafficGroup)
 					{
 						TrafficExecutionContext = Context;
 						// Needs its own command buffer as it will be in it own thread, will be merged after this loop
-						TrafficExecutionContext.SetDeferredCommandBuffer(MakeShareable(new FLWCCommandBuffer()));
-						TrafficCompletionEvent = TGraphTask<FPipeProcessorTask>::CreateTask().ConstructAndDispatchWhenReady(EntitySubsystem, TrafficExecutionContext, *Proc, /*bManageCommandBuffer=*/false);
+						TrafficExecutionContext.SetDeferredCommandBuffer(MakeShareable(new FMassCommandBuffer()));
+						TrafficCompletionEvent = TGraphTask<FMassProcessorTask>::CreateTask().ConstructAndDispatchWhenReady(EntitySubsystem, TrafficExecutionContext, *Proc, /*bManageCommandBuffer=*/false);
 						continue;
 					}
 				}
@@ -320,14 +320,14 @@ void UPipeCompositeProcessor::Execute(UMassEntitySubsystem& EntitySubsystem, FLW
 		CompletionStatus.AddDefaulted(ChildPipeline.Processors.Num());
 		TArray<int32> PostponedProcessors;
 
-		FLWComponentSystemExecutionContext SingleThreadContext = Context;
-		SingleThreadContext.SetDeferredCommandBuffer(MakeShareable(new FLWCCommandBuffer()));
+		FMassExecutionContext SingleThreadContext = Context;
+		SingleThreadContext.SetDeferredCommandBuffer(MakeShareable(new FMassCommandBuffer()));
 
 		for (int32 NodeIndex = 0; NodeIndex < ChildPipeline.Processors.Num(); ++NodeIndex)
 		{
-			UPipeProcessor* Proc = ChildPipeline.Processors[NodeIndex];
+			UMassProcessor* Proc = ChildPipeline.Processors[NodeIndex];
 			check(Proc);
-			UPipeCompositeProcessor* CompositeProc = Cast<UPipeCompositeProcessor>(Proc);
+			UMassCompositeProcessor* CompositeProc = Cast<UMassCompositeProcessor>(Proc);
 			if (CompositeProc == nullptr || CompositeProc->bRunInSeparateThread == false)
 			{
 				// check if all dependencies have been processed already
@@ -382,18 +382,18 @@ void UPipeCompositeProcessor::Execute(UMassEntitySubsystem& EntitySubsystem, FLW
 					, DependenciesDesc.Len() > 0 ? TEXT(" Dependencies: ") : TEXT(""), *DependenciesDesc);
 
 				// send off to another thread
-				CompletionStatus[NodeIndex].CompletionEvent = TGraphTask<FPipeProcessorTask>::CreateTask(&Prerequisites).ConstructAndDispatchWhenReady(EntitySubsystem, Context, *Proc);
+				CompletionStatus[NodeIndex].CompletionEvent = TGraphTask<FMassProcessorTask>::CreateTask(&Prerequisites).ConstructAndDispatchWhenReady(EntitySubsystem, Context, *Proc);
 				CompletionStatus[NodeIndex].Status = EProcessorCompletionStatus::Threaded;
 			}
 		}
 
 		{
-			TRACE_CPUPROFILER_EVENT_SCOPE_TEXT("Postponed PipeProcessors");
+			TRACE_CPUPROFILER_EVENT_SCOPE_TEXT("Postponed MassProcessors");
 
 			for (int32 i = 0; i < PostponedProcessors.Num(); ++i)
 			{
 				const int32 PostponedIndex = PostponedProcessors[i];
-				UPipeProcessor* Proc = ChildPipeline.Processors[PostponedIndex];
+				UMassProcessor* Proc = ChildPipeline.Processors[PostponedIndex];
 				check(Proc);
 				for (int32 j = Proc->TransientDependencyIndices.Num() - 1; j >= 0; --j)
 				{
@@ -428,7 +428,7 @@ void UPipeCompositeProcessor::Execute(UMassEntitySubsystem& EntitySubsystem, FLW
 	}
 	else
 	{
-		for (UPipeProcessor* Proc : ChildPipeline.Processors)
+		for (UMassProcessor* Proc : ChildPipeline.Processors)
 		{
 			check(Proc);
 			Proc->CallExecute(EntitySubsystem, Context);
@@ -436,25 +436,25 @@ void UPipeCompositeProcessor::Execute(UMassEntitySubsystem& EntitySubsystem, FLW
 	}
 }
 
-void UPipeCompositeProcessor::Initialize(UObject& Owner)
+void UMassCompositeProcessor::Initialize(UObject& Owner)
 {
 	// remove all nulls
-	ChildPipeline.Processors.RemoveAll([](const UPipeProcessor* Proc) { return Proc == nullptr; });
+	ChildPipeline.Processors.RemoveAll([](const UMassProcessor* Proc) { return Proc == nullptr; });
 
 	// from this point on we don't expect to have nulls in ChildPipeline.Processors
-	for (UPipeProcessor* Proc : ChildPipeline.Processors)
+	for (UMassProcessor* Proc : ChildPipeline.Processors)
 	{
 		REDIRECT_OBJECT_TO_VLOG(Proc, this);
 		Proc->Initialize(Owner);
 	}
 }
 
-void UPipeCompositeProcessor::CopyAndSort(const FPipeProcessingPhaseConfig& PhaseConfig, const FString& DependencyGraphFileName)
+void UMassCompositeProcessor::CopyAndSort(const FMassProcessingPhaseConfig& PhaseConfig, const FString& DependencyGraphFileName)
 {
 	check(GetOuter());
 	
 	// create processors via a temporary pipeline. This ensures consistency with other places we use FRuntimePipelines
-	FRuntimePipeline TmpPipeline;
+	FMassRuntimePipeline TmpPipeline;
 	TmpPipeline.CreateFromArray(PhaseConfig.ProcessorCDOs, *GetOuter());
 
 	// figure out dependencies
@@ -481,7 +481,7 @@ void UPipeCompositeProcessor::CopyAndSort(const FPipeProcessingPhaseConfig& Phas
 	}
 }
 
-int32 UPipeCompositeProcessor::Populate(TArray<FProcessorDependencySolver::FOrderInfo>& ProcessorsAndGroups, const int32 StartIndex)
+int32 UMassCompositeProcessor::Populate(TArray<FProcessorDependencySolver::FOrderInfo>& ProcessorsAndGroups, const int32 StartIndex)
 {
 	ChildPipeline.Processors.Reset();
 
@@ -492,7 +492,7 @@ int32 UPipeCompositeProcessor::Populate(TArray<FProcessorDependencySolver::FOrde
 	int32 Index = StartIndex;
 	TMap<FName, int32> NameToIndexMap;
 	bool bOffThreadGroupsFound = false;
-	const FPipeProcessingPhaseConfig& PhaseConfig = GET_PIPE_CONFIG_VALUE(GetProcessingPhaseConfig(ProcessingPhase));
+	const FMassProcessingPhaseConfig& PhaseConfig = GET_PIPE_CONFIG_VALUE(GetProcessingPhaseConfig(ProcessingPhase));
 
 	while (ProcessorsAndGroups.IsValidIndex(Index))
 	{
@@ -516,7 +516,7 @@ int32 UPipeCompositeProcessor::Populate(TArray<FProcessorDependencySolver::FOrde
 		}
 		else if (Element.NodeType == EDependencyNodeType::GroupStart)
 		{
-			UPipeCompositeProcessor* GroupProcessor = NewObject<UPipeCompositeProcessor>(GetOuter());
+			UMassCompositeProcessor* GroupProcessor = NewObject<UMassCompositeProcessor>(GetOuter());
 			GroupProcessor->SetGroupName(Element.Name);
 			GroupProcessor->SetProcessingPhase(ProcessingPhase);
 			if (PhaseConfig.OffGameThreadGroupNames.Find(Element.Name) != INDEX_NONE)
@@ -551,7 +551,7 @@ int32 UPipeCompositeProcessor::Populate(TArray<FProcessorDependencySolver::FOrde
 	return Index;
 }
 
-void UPipeCompositeProcessor::DebugOutputDescription(FOutputDevice& Ar, int32 Indent) const
+void UMassCompositeProcessor::DebugOutputDescription(FOutputDevice& Ar, int32 Indent) const
 {
 #if WITH_MASSENTITY_DEBUG
 	if (ChildPipeline.Processors.Num() == 0)
@@ -561,7 +561,7 @@ void UPipeCompositeProcessor::DebugOutputDescription(FOutputDevice& Ar, int32 In
 	else
 	{
 		Ar.Logf(TEXT("%*sGroup %s:"), Indent, TEXT(""), *GroupName.ToString());
-		for (UPipeProcessor* Proc : ChildPipeline.Processors)
+		for (UMassProcessor* Proc : ChildPipeline.Processors)
 		{
 			check(Proc);
 			Ar.Logf(TEXT("\n"));
@@ -571,16 +571,16 @@ void UPipeCompositeProcessor::DebugOutputDescription(FOutputDevice& Ar, int32 In
 #endif // WITH_MASSENTITY_DEBUG
 }
 
-void UPipeCompositeProcessor::SetProcessingPhase(EPipeProcessingPhase Phase)
+void UMassCompositeProcessor::SetProcessingPhase(EMassProcessingPhase Phase)
 {
 	Super::SetProcessingPhase(Phase);
-	for (UPipeProcessor* Proc : ChildPipeline.Processors)
+	for (UMassProcessor* Proc : ChildPipeline.Processors)
 	{
 		Proc->SetProcessingPhase(Phase);
 	}
 }
 
-void UPipeCompositeProcessor::SetGroupName(FName NewName)
+void UMassCompositeProcessor::SetGroupName(FName NewName)
 {
 	GroupName = NewName;
 #if CPUPROFILERTRACE_ENABLED
@@ -588,7 +588,7 @@ void UPipeCompositeProcessor::SetGroupName(FName NewName)
 #endif
 }
 
-void UPipeCompositeProcessor::AddGroupedProcessor(FName RequestedGroupName, UPipeProcessor& Processor)
+void UMassCompositeProcessor::AddGroupedProcessor(FName RequestedGroupName, UMassProcessor& Processor)
 {
 	if (RequestedGroupName.IsNone() || RequestedGroupName == GroupName)
 	{
@@ -597,15 +597,15 @@ void UPipeCompositeProcessor::AddGroupedProcessor(FName RequestedGroupName, UPip
 	else
 	{
 		FString RemainingGroupName;
-		UPipeCompositeProcessor* GroupProcessor = FindOrAddGroupProcessor(RequestedGroupName, &RemainingGroupName);
+		UMassCompositeProcessor* GroupProcessor = FindOrAddGroupProcessor(RequestedGroupName, &RemainingGroupName);
 		check(GroupProcessor);
 		GroupProcessor->AddGroupedProcessor(FName(*RemainingGroupName), Processor);
 	}
 }
 
-UPipeCompositeProcessor* UPipeCompositeProcessor::FindOrAddGroupProcessor(FName RequestedGroupName, FString* OutRemainingGroupName)
+UMassCompositeProcessor* UMassCompositeProcessor::FindOrAddGroupProcessor(FName RequestedGroupName, FString* OutRemainingGroupName)
 {
-	UPipeCompositeProcessor* GroupProcessor = nullptr;
+	UMassCompositeProcessor* GroupProcessor = nullptr;
 	const FString NameAsString = RequestedGroupName.ToString();
 	FString TopGroupName;
 	if (NameAsString.Split(TEXT("."), &TopGroupName, OutRemainingGroupName))
@@ -617,7 +617,7 @@ UPipeCompositeProcessor* UPipeCompositeProcessor::FindOrAddGroupProcessor(FName
 	if (GroupProcessor == nullptr)
 	{
 		check(GetOuter());
-		GroupProcessor = NewObject<UPipeCompositeProcessor>(GetOuter());
+		GroupProcessor = NewObject<UMassCompositeProcessor>(GetOuter());
 		GroupProcessor->SetGroupName(RequestedGroupName);
 		ChildPipeline.AppendProcessor(*GroupProcessor);
 	}
diff --git a/Engine/Plugins/Runtime/MassEntity/Source/MassEntity/Private/MassProcessorDependencySolver.cpp b/Engine/Plugins/Runtime/MassEntity/Source/MassEntity/Private/MassProcessorDependencySolver.cpp
--- a/Engine/Plugins/Runtime/MassEntity/Source/MassEntity/Private/MassProcessorDependencySolver.cpp
+++ b/Engine/Plugins/Runtime/MassEntity/Source/MassEntity/Private/MassProcessorDependencySolver.cpp
@@ -41,7 +41,7 @@ bool FProcessorDependencySolver::FNode::HasDependencies() const
 //----------------------------------------------------------------------//
 //  FProcessorDependencySolver
 //----------------------------------------------------------------------//
-FProcessorDependencySolver::FProcessorDependencySolver(TArrayView<UPipeProcessor*> InProcessors, const FName Name, const FString& InDependencyGraphFileName /*= FString()*/)
+FProcessorDependencySolver::FProcessorDependencySolver(TArrayView<UMassProcessor*> InProcessors, const FName Name, const FString& InDependencyGraphFileName /*= FString()*/)
 	: Processors(InProcessors)
 	, GroupRootNode({Name, nullptr})
 	, DependencyGraphFileName(InDependencyGraphFileName)
@@ -129,7 +129,7 @@ void FProcessorDependencySolver::CreateSubGroupNames(FName InGroupName, TArray<F
 	}
 }
 
-void FProcessorDependencySolver::AddNode(FName InGroupName, UPipeProcessor& Processor)
+void FProcessorDependencySolver::AddNode(FName InGroupName, UMassProcessor& Processor)
 {
 	check(Processor.GetClass());
 	const FName ProcName = Processor.GetClass()->GetFName();
@@ -672,7 +672,7 @@ struct FDumpGraphDependencyUtils
 void FProcessorDependencySolver::DumpGraph(FArchive& LogFile) const
 {
 	TSet<const FNode*> AllNodes;
-	LogFile.Logf(TEXT("digraph PipeProcessorGraph"));
+	LogFile.Logf(TEXT("digraph MassProcessorGraph"));
 	LogFile.Logf(TEXT("{"));
 	LogFile.Logf(TEXT("    compound = true;"));
 	LogFile.Logf(TEXT("    newrank = true;"));
@@ -761,7 +761,7 @@ void FProcessorDependencySolver::Solve(FNode& RootNode, TConstArrayView<const FN
 			UE_LOG(LogPipe, Error, TEXT("Detected processing dependency cycle:"));
 			for (const int32 Index : IndicesRemaining)
 			{
-				UPipeProcessor* Processor = RootNode.SubNodes[Index].Processor;
+				UMassProcessor* Processor = RootNode.SubNodes[Index].Processor;
 				if (Processor)
 				{
 					UE_LOG(LogPipe, Error, TEXT("\t%s, group: %s, before: %s, after %s")
@@ -822,7 +822,7 @@ void FProcessorDependencySolver::ResolveDependencies(TArray<FProcessorDependency
 	UE_LOG(LogPipe, Log, TEXT("Gathering dependencies data:"));
 
 	// gather the processors information first
-	for (UPipeProcessor* Processor : Processors)
+	for (UMassProcessor* Processor : Processors)
 	{
 		if (Processor == nullptr)
 		{
