AggregateTicking and Pipe type renames:

diff --git a/Engine/Plugins/Runtime/MassEntity/Source/MassEntity/Public/ArchetypeTypes.h b/Engine/Plugins/Runtime/MassEntity/Source/MassEntity/Public/ArchetypeTypes.h
--- a/Engine/Plugins/Runtime/MassEntity/Source/MassEntity/Public/ArchetypeTypes.h
+++ b/Engine/Plugins/Runtime/MassEntity/Source/MassEntity/Public/ArchetypeTypes.h
@@ -8,16 +8,16 @@
 #include "LWComponentTypes.h"
 
 class UMassEntitySubsystem;
-struct FArchetypeData;
-struct FLWComponentSystemExecutionContext;
-struct FLWComponentData;
+struct FMassArchetypeData;
+struct FMassExecutionContext;
+struct FMassFragment;
 struct FArchetypeChunkIterator;
-struct FLWComponentQuery;
+struct FMassEntityQuery;
 struct FArchetypeChunkCollection;
-struct FEntityView;
+struct FMassEntityView;
 
-typedef TFunction< void(FLWComponentSystemExecutionContext& /*ExecutionContext*/) > FLWComponentSystemExecuteFunction;
-typedef TFunction< bool(const FLWComponentSystemExecutionContext& /*ExecutionContext*/) > FLWComponentSystemChunkConditionFunction;
+typedef TFunction< void(FMassExecutionContext& /*ExecutionContext*/) > FMassExecuteFunction;
+typedef TFunction< bool(const FMassExecutionContext& /*ExecutionContext*/) > FMassChunkConditionFunction;
 
 
 //////////////////////////////////////////////////////////////////////
@@ -29,21 +29,21 @@ struct FArchetypeHandle final
 	FArchetypeHandle() = default;
 	bool IsValid() const { return DataPtr.IsValid(); }
 
-	MASSENTITY_API bool operator==(const FArchetypeData* Other) const;
+	MASSENTITY_API bool operator==(const FMassArchetypeData* Other) const;
 	bool operator==(const FArchetypeHandle& Other) const { return DataPtr == Other.DataPtr; }
 	bool operator!=(const FArchetypeHandle& Other) const { return DataPtr != Other.DataPtr; }
 
 	MASSENTITY_API friend uint32 GetTypeHash(const FArchetypeHandle& Instance);
 private:
-	FArchetypeHandle(const TSharedPtr<FArchetypeData>& InDataPtr)
+	FArchetypeHandle(const TSharedPtr<FMassArchetypeData>& InDataPtr)
 	: DataPtr(InDataPtr)
 	{}
-	TSharedPtr<FArchetypeData> DataPtr;
+	TSharedPtr<FMassArchetypeData> DataPtr;
 
 	friend UMassEntitySubsystem;
 	friend FArchetypeChunkCollection;
-	friend FLWComponentQuery;
-	friend FEntityView;
+	friend FMassEntityQuery;
+	friend FMassEntityView;
 };
 
 
@@ -76,9 +76,9 @@ private:
 
 public:
 	FArchetypeChunkCollection() = default;
-	FArchetypeChunkCollection(const FArchetypeHandle& InArchetype, TConstArrayView<FLWEntity> InEntities);
+	FArchetypeChunkCollection(const FArchetypeHandle& InArchetype, TConstArrayView<FMassEntityHandle> InEntities);
 	explicit FArchetypeChunkCollection(FArchetypeHandle& InArchetypeHandle);
-	explicit FArchetypeChunkCollection(TSharedPtr<FArchetypeData>& InArchetype);
+	explicit FArchetypeChunkCollection(TSharedPtr<FMassArchetypeData>& InArchetype);
 
 	TArrayView<const FChunkInfo> GetChunks() const { return Chunks; }
 	const FArchetypeHandle& GetArchetype() const { return Archetype; }
@@ -91,7 +91,7 @@ public:
 	}
 
 private:
-	void GatherChunksFromArchetype(TSharedPtr<FArchetypeData>& InArchetype);
+	void GatherChunksFromArchetype(TSharedPtr<FMassArchetypeData>& InArchetype);
 };
 
 //////////////////////////////////////////////////////////////////////
@@ -132,14 +132,14 @@ struct FInternalEntityHandle
 	int32 IndexWithinChunk = INDEX_NONE;
 };
 
-typedef TArray<int32, TInlineAllocator<16>> FLWComponentIndicesMapping;
-typedef TConstArrayView<int32> FLWComponentIndicesMappingView;
-struct FLWRequirementIndicesMapping
+typedef TArray<int32, TInlineAllocator<16>> FMassFragmentIndicesMapping;
+typedef TConstArrayView<int32> FMassFragmentIndicesMappingView;
+struct FMassQueryRequirementIndicesMapping
 {
-	FLWRequirementIndicesMapping() = default;
+	FMassQueryRequirementIndicesMapping() = default;
 
-	FLWComponentIndicesMapping EntityComponents;
-	FLWComponentIndicesMapping ChunkComponents;
+	FMassFragmentIndicesMapping EntityComponents;
+	FMassFragmentIndicesMapping ChunkComponents;
 	FORCEINLINE bool IsEmpty() const
 	{
 		return EntityComponents.Num() == 0 || ChunkComponents.Num() == 0;
@@ -147,7 +147,7 @@ struct FLWRequirementIndicesMapping
 };
 
 template<typename T>
-struct FLWComponentSorterOperator
+struct FMassSorterOperator
 {
 	bool operator()(const T& A, const T& B) const
 	{
diff --git a/Engine/Plugins/Runtime/MassEntity/Source/MassEntity/Public/EntityQuery.h b/Engine/Plugins/Runtime/MassEntity/Source/MassEntity/Public/EntityQuery.h
--- a/Engine/Plugins/Runtime/MassEntity/Source/MassEntity/Public/EntityQuery.h
+++ b/Engine/Plugins/Runtime/MassEntity/Source/MassEntity/Public/EntityQuery.h
@@ -11,67 +11,67 @@
 
 
 class UMassEntitySubsystem;
-struct FArchetypeData;
-struct FLWComponentSystemExecutionContext;
-struct FLWComponentData;
+struct FMassArchetypeData;
+struct FMassExecutionContext;
+struct FMassFragment;
 struct FArchetypeHandle;
 
-enum class ELWComponentAccess : uint8
+enum class EMassFragmentAccess : uint8
 {
 	// no binding required
 	None, 
 
-	// We want to read the data for the component
+	// We want to read the data for the fragment
 	ReadOnly,
 
-	// We want to read and write the data for the component
+	// We want to read and write the data for the fragment
 	ReadWrite,
 
 	MAX
 };
 
-enum class ELWComponentPresence : uint8
+enum class EMassFragmentPresence : uint8
 {
-	// All of the required components must be present
+	// All of the required fragments must be present
 	All,
 
-	// One of the required components must be present
+	// One of the required fragments must be present
 	Any,
 
-	// None of the required components can be present
+	// None of the required fragments can be present
 	None,
 
-	// If component is present we'll use it, but it missing stop processing of a given archetype
+	// If fragment is present we'll use it, but it missing stop processing of a given archetype
 	Optional,
 
 	MAX
 };
 
 
-struct MASSENTITY_API FLWComponentRequirement
+struct MASSENTITY_API FMassFragmentRequirement
 {
 	const UScriptStruct* StructType = nullptr;
-	ELWComponentAccess AccessMode = ELWComponentAccess::None;
-	ELWComponentPresence Presence = ELWComponentPresence::Optional;
+	EMassFragmentAccess AccessMode = EMassFragmentAccess::None;
+	EMassFragmentPresence Presence = EMassFragmentPresence::Optional;
 
 public:
-	FLWComponentRequirement(){}
-	FLWComponentRequirement(const UScriptStruct* InStruct, const ELWComponentAccess InAccessMode, const ELWComponentPresence InPresence)
+	FMassFragmentRequirement(){}
+	FMassFragmentRequirement(const UScriptStruct* InStruct, const EMassFragmentAccess InAccessMode, const EMassFragmentPresence InPresence)
 		: StructType(InStruct)
 		, AccessMode(InAccessMode)
 		, Presence(InPresence)
 	{
 		check(InStruct);
-		checkf((Presence != ELWComponentPresence::Any && Presence != ELWComponentPresence::Optional)
-			|| AccessMode == ELWComponentAccess::ReadOnly || AccessMode == ELWComponentAccess::ReadWrite, TEXT("Only ReadOnly and ReadWrite modes are suppored for optional requirements"));
+		checkf((Presence != EMassFragmentPresence::Any && Presence != EMassFragmentPresence::Optional)
+			|| AccessMode == EMassFragmentAccess::ReadOnly || AccessMode == EMassFragmentAccess::ReadWrite, TEXT("Only ReadOnly and ReadWrite modes are suppored for optional requirements"));
 	}
 
-	bool RequiresBinding() const { return (AccessMode != ELWComponentAccess::None); }
-	bool IsOptional() const { return (Presence == ELWComponentPresence::Optional || Presence == ELWComponentPresence::Any); }
+	bool RequiresBinding() const { return (AccessMode != EMassFragmentAccess::None); }
+	bool IsOptional() const { return (Presence == EMassFragmentPresence::Optional || Presence == EMassFragmentPresence::Any); }
 
 	FString DebugGetDescription() const;
 
-	// these functions are used for sorting. See FLWComponentSorterOperator
+	// these functions are used for sorting. See FMassSorterOperator
 	int32 GetStructureSize() const
 	{
 		return StructType->GetStructureSize();
@@ -85,61 +85,61 @@ public:
 
 
 /** 
- *  FLWComponentQuery is a structure that serves two main purposes:
+ *  FMassEntityQuery is a structure that serves two main purposes:
  *  1. Describe properties required of an archetype that's a subject of calculations
  *  2. Trigger calculations on cached set of valid archetypes as described by Requirements
  * 
- *  A query to be considered valid needs declared at least one ELWComponentPresence::All, ELWComponentPresence::Any 
- *  ELWComponentPresence::Optional component requirement.
+ *  A query to be considered valid needs declared at least one EMassFragmentPresence::All, EMassFragmentPresence::Any 
+ *  EMassFragmentPresence::Optional fragment requirement.
  */
 USTRUCT()
-struct MASSENTITY_API FLWComponentQuery
+struct MASSENTITY_API FMassEntityQuery
 {
 	GENERATED_BODY()
 
 public:
-	FLWComponentQuery();
-	FLWComponentQuery(std::initializer_list<UScriptStruct*> InitList);
-	FLWComponentQuery(TConstArrayView<const UScriptStruct*> InitList);
+	FMassEntityQuery();
+	FMassEntityQuery(std::initializer_list<UScriptStruct*> InitList);
+	FMassEntityQuery(TConstArrayView<const UScriptStruct*> InitList);
 
 	/** Runs ExecuteFunction on all entities matching Requirements */
-	void ForEachEntityChunk(UMassEntitySubsystem& EntitySubsystem, FLWComponentSystemExecutionContext& ExecutionContext, const FLWComponentSystemExecuteFunction& ExecuteFunction);
+	void ForEachEntityChunk(UMassEntitySubsystem& EntitySubsystem, FMassExecutionContext& ExecutionContext, const FMassExecuteFunction& ExecuteFunction);
 	
 	/** Will first verify that the archetype given with Chunks matches the query's requirements, and if so will run the other, more generic ForEachEntityChunk implementation */
-	void ForEachEntityChunk(const FArchetypeChunkCollection& Chunks, UMassEntitySubsystem& EntitySubsystem, FLWComponentSystemExecutionContext& ExecutionContext, const FLWComponentSystemExecuteFunction& ExecuteFunction);
+	void ForEachEntityChunk(const FArchetypeChunkCollection& Chunks, UMassEntitySubsystem& EntitySubsystem, FMassExecutionContext& ExecutionContext, const FMassExecuteFunction& ExecuteFunction);
 
 	/**
 	 * Attempts to process every chunk of every affected archetype in parallel.
 	 */
-	void ParallelForEachEntityChunk(UMassEntitySubsystem& EntitySubsystem, FLWComponentSystemExecutionContext& ExecutionContext, const FLWComponentSystemExecuteFunction& ExecuteFunction);
+	void ParallelForEachEntityChunk(UMassEntitySubsystem& EntitySubsystem, FMassExecutionContext& ExecutionContext, const FMassExecuteFunction& ExecuteFunction);
 
 	/** Will gather all archetypes from InEntitySubsystem matching this->Requirements.
 	 *  Note that no work will be done if the cached data is up to date (as tracked by EntitySubsystemHash and 
 	 *	ArchetypeDataVersion properties). */
 	void CacheArchetypes(UMassEntitySubsystem& InEntitySubsystem);
 
-	FLWComponentQuery& AddRequirement(const UScriptStruct* ComponentType, const ELWComponentAccess AccessMode, const ELWComponentPresence Presence = ELWComponentPresence::All)
+	FMassEntityQuery& AddRequirement(const UScriptStruct* ComponentType, const EMassFragmentAccess AccessMode, const EMassFragmentPresence Presence = EMassFragmentPresence::All)
 	{
-		checkf(Requirements.FindByPredicate([ComponentType](const FLWComponentRequirement& Item){ return Item.StructType == ComponentType; }) == nullptr
+		checkf(Requirements.FindByPredicate([ComponentType](const FMassFragmentRequirement& Item){ return Item.StructType == ComponentType; }) == nullptr
 			, TEXT("Duplicated requirements are not supported. %s already present"), *GetNameSafe(ComponentType));
 		
-		if (Presence != ELWComponentPresence::None)
+		if (Presence != EMassFragmentPresence::None)
 		{
 			Requirements.Emplace(ComponentType, AccessMode, Presence);
 		}
 
 		switch (Presence)
 		{
-		case ELWComponentPresence::All:
+		case EMassFragmentPresence::All:
 			RequiredAllComponents.Add(*ComponentType);
 			break;
-		case ELWComponentPresence::Any:
+		case EMassFragmentPresence::Any:
 			RequiredAnyComponents.Add(*ComponentType);
 			break;
-		case ELWComponentPresence::Optional:
+		case EMassFragmentPresence::Optional:
 			RequiredOptionalComponents.Add(*ComponentType);
 			break;
-		case ELWComponentPresence::None:
+		case EMassFragmentPresence::None:
 			RequiredNoneComponents.Add(*ComponentType);
 			break;
 		}
@@ -148,32 +148,32 @@ public:
 		return *this;
 	}
 
-	/** FLWComponentQuery ref returned for chaining */
+	/** FMassEntityQuery ref returned for chaining */
 	template<typename T>
-	FLWComponentQuery& AddRequirement(const ELWComponentAccess AccessMode, const ELWComponentPresence Presence = ELWComponentPresence::All)
+	FMassEntityQuery& AddRequirement(const EMassFragmentAccess AccessMode, const EMassFragmentPresence Presence = EMassFragmentPresence::All)
 	{
-		checkf(Requirements.FindByPredicate([](const FLWComponentRequirement& Item) { return Item.StructType == T::StaticStruct(); }) == nullptr
+		checkf(Requirements.FindByPredicate([](const FMassFragmentRequirement& Item) { return Item.StructType == T::StaticStruct(); }) == nullptr
 			, TEXT("Duplicated requirements are not supported. %s already present"), *T::StaticStruct()->GetName());
 
-		static_assert(TIsDerivedFrom<T, FLWComponentData>::IsDerived, "Given struct doesn't represent a valid fragment type. Make sure to inherit from FLWComponentData or one of its child-types.");
+		static_assert(TIsDerivedFrom<T, FMassFragment>::IsDerived, "Given struct doesn't represent a valid fragment type. Make sure to inherit from FMassFragment or one of its child-types.");
 		
-		if (Presence != ELWComponentPresence::None)
+		if (Presence != EMassFragmentPresence::None)
 		{
 			Requirements.Emplace(T::StaticStruct(), AccessMode, Presence);
 		}
 		
 		switch (Presence)
 		{
-		case ELWComponentPresence::All:
+		case EMassFragmentPresence::All:
 			RequiredAllComponents.Add<T>();
 			break;
-		case ELWComponentPresence::Any:
+		case EMassFragmentPresence::Any:
 			RequiredAnyComponents.Add<T>();
 			break;
-		case ELWComponentPresence::Optional:
+		case EMassFragmentPresence::Optional:
 			RequiredOptionalComponents.Add<T>();
 			break;
-		case ELWComponentPresence::None:
+		case EMassFragmentPresence::None:
 			RequiredNoneComponents.Add<T>();
 			break;
 		}
@@ -182,37 +182,37 @@ public:
 		return *this;
 	}
 
-	void AddTagRequirement(const UScriptStruct& ComponentType, const ELWComponentPresence Presence)
+	void AddTagRequirement(const UScriptStruct& TagType, const EMassFragmentPresence Presence)
 	{
-		checkfSlow(int(Presence) < int(ELWComponentPresence::Optional), TEXT("Optional and MAX presence are not valid calues for AddTagRequirement"));
+		checkfSlow(int(Presence) < int(EMassFragmentPresence::Optional), TEXT("Optional and MAX presence are not valid calues for AddTagRequirement"));
 		switch (Presence)
 		{
-		case ELWComponentPresence::All:
-			RequiredAllTags.Add(ComponentType);
+		case EMassFragmentPresence::All:
+			RequiredAllTags.Add(TagType);
 			break;
-		case ELWComponentPresence::Any:
-			RequiredAnyTags.Add(ComponentType);
+		case EMassFragmentPresence::Any:
+			RequiredAnyTags.Add(TagType);
 			break;
-		case ELWComponentPresence::None:
-			RequiredNoneTags.Add(ComponentType);
+		case EMassFragmentPresence::None:
+			RequiredNoneTags.Add(TagType);
 			break;
 		}
 	}
 
 	template<typename T>
-	FLWComponentQuery& AddTagRequirement(const ELWComponentPresence Presence)
+	FMassEntityQuery& AddTagRequirement(const EMassFragmentPresence Presence)
 	{
-		checkfSlow(int(Presence) < int(ELWComponentPresence::Optional), TEXT("Optional and MAX presence are not valid calues for AddTagRequirement"));
-		static_assert(TIsDerivedFrom<T, FComponentTag>::IsDerived, "Given struct doesn't represent a valid tag type. Make sure to inherit from FLWComponentData or one of its child-types.");
+		checkfSlow(int(Presence) < int(EMassFragmentPresence::Optional), TEXT("Optional and MAX presence are not valid calues for AddTagRequirement"));
+		static_assert(TIsDerivedFrom<T, FMassTag>::IsDerived, "Given struct doesn't represent a valid tag type. Make sure to inherit from FMassFragment or one of its child-types.");
 		switch (Presence)
 		{
-			case ELWComponentPresence::All:
+			case EMassFragmentPresence::All:
 				RequiredAllTags.Add<T>();
 				break;
-			case ELWComponentPresence::Any:
+			case EMassFragmentPresence::Any:
 				RequiredAnyTags.Add<T>();
 				break;
-			case ELWComponentPresence::None:
+			case EMassFragmentPresence::None:
 				RequiredNoneTags.Add<T>();
 				break;
 		}
@@ -221,24 +221,24 @@ public:
 	}
 
 	// actual implementation in specializations
-	template<ELWComponentPresence Presence> 
-	FLWComponentQuery& AddTagRequirements(const FLWTagBitSet& TagBitSet)
+	template<EMassFragmentPresence Presence> 
+	FMassEntityQuery& AddTagRequirements(const FMassTagBitSet& TagBitSet)
 	{
-		static_assert(Presence == ELWComponentPresence::None || Presence == ELWComponentPresence::All || Presence == ELWComponentPresence::Any
+		static_assert(Presence == EMassFragmentPresence::None || Presence == EMassFragmentPresence::All || Presence == EMassFragmentPresence::Any
 			, "The only valid values for AddTagRequirements are All, Any and None");
 		return *this;
 	}
 
 	template<typename T>
-	FLWComponentQuery& AddChunkRequirement(const ELWComponentAccess AccessMode, const ELWComponentPresence Presence = ELWComponentPresence::All)
+	FMassEntityQuery& AddChunkRequirement(const EMassFragmentAccess AccessMode, const EMassFragmentPresence Presence = EMassFragmentPresence::All)
 	{
-		static_assert(TIsDerivedFrom<T, FLWChunkComponent>::IsDerived, "Given struct doesn't represent a valid chunk fragment type. Make sure to inherit from FLWChunkComponent or one of its child-types.");
-		checkf(ChunkRequirements.FindByPredicate([](const FLWComponentRequirement& Item) { return Item.StructType == T::StaticStruct(); }) == nullptr
+		static_assert(TIsDerivedFrom<T, FMassChunkFragment>::IsDerived, "Given struct doesn't represent a valid chunk fragment type. Make sure to inherit from FMassChunkFragment or one of its child-types.");
+		checkf(ChunkRequirements.FindByPredicate([](const FMassFragmentRequirement& Item) { return Item.StructType == T::StaticStruct(); }) == nullptr
 			, TEXT("Duplicated requirements are not supported. %s already present"), *T::StaticStruct()->GetName());
-		checkfSlow(Presence == ELWComponentPresence::None || Presence == ELWComponentPresence::All
+		checkfSlow(Presence == EMassFragmentPresence::None || Presence == EMassFragmentPresence::All
 			, TEXT("The only valid Presence values for AddChunkRequirement are All and None"));
 		
-		if (Presence == ELWComponentPresence::All)
+		if (Presence == EMassFragmentPresence::All)
 		{
 			RequiredAllChunkComponents.Add<T>();
 			ChunkRequirements.Emplace(T::StaticStruct(), AccessMode, Presence);
@@ -265,23 +265,23 @@ public:
 	
 	bool DoesArchetypeMatchRequirements(const FArchetypeHandle& ArchetypeHandle) const;
 
-	/** The function validates requirements we make for queries. See the FLWComponentQuery struct description for details.
+	/** The function validates requirements we make for queries. See the FMassEntityQuery struct description for details.
 	 *  Note that this function is non-trivial and end users are not expected to need to use it. 
 	 *  @return whether this query's requirements follow the rules. */
 	bool CheckValidity() const;
 
 	FString DebugGetDescription() const;
 
-	TConstArrayView<FLWComponentRequirement> GetRequirements() const { return Requirements; }
-	const FLWComponentBitSet& GetRequiredAllComponents() const { return RequiredAllComponents; }
-	const FLWComponentBitSet& GetRequiredAnyComponents() const { return RequiredAnyComponents; }
-	const FLWComponentBitSet& GetRequiredOptionalComponents() const { return RequiredOptionalComponents; }
-	const FLWComponentBitSet& GetRequiredNoneComponents() const { return RequiredNoneComponents; }
-	const FLWTagBitSet& GetRequiredAllTags() const { return RequiredAllTags; }
-	const FLWTagBitSet& GetRequiredAnyTags() const { return RequiredAnyTags; }
-	const FLWTagBitSet& GetRequiredNoneTags() const { return RequiredNoneTags; }
-	const FLWChunkComponentBitSet& GetRequiredAllChunkComponents() const { return RequiredAllChunkComponents; }
-	const FLWChunkComponentBitSet& GetRequiredNoneChunkComponents() const { return RequiredNoneChunkComponents; }
+	TConstArrayView<FMassFragmentRequirement> GetRequirements() const { return Requirements; }
+	const FMassFragmentBitSet& GetRequiredAllComponents() const { return RequiredAllComponents; }
+	const FMassFragmentBitSet& GetRequiredAnyComponents() const { return RequiredAnyComponents; }
+	const FMassFragmentBitSet& GetRequiredOptionalComponents() const { return RequiredOptionalComponents; }
+	const FMassFragmentBitSet& GetRequiredNoneComponents() const { return RequiredNoneComponents; }
+	const FMassTagBitSet& GetRequiredAllTags() const { return RequiredAllTags; }
+	const FMassTagBitSet& GetRequiredAnyTags() const { return RequiredAnyTags; }
+	const FMassTagBitSet& GetRequiredNoneTags() const { return RequiredNoneTags; }
+	const FMassChunkFragmentBitSet& GetRequiredAllChunkComponents() const { return RequiredAllChunkComponents; }
+	const FMassChunkFragmentBitSet& GetRequiredNoneChunkComponents() const { return RequiredNoneChunkComponents; }
 
 	const TArray<FArchetypeHandle>& GetArchetypes() const
 	{ 
@@ -307,7 +307,7 @@ public:
 	 * that this condition won't be applied when a specific collection of chunks is used (via FArchetypeChunkCollection)
 	 * The value returned by InFunction controls whether to allow execution (true) or block it (false).
 	 */
-	void SetChunkFilter(const FLWComponentSystemChunkConditionFunction& InFunction) { ChunkCondition = InFunction; }
+	void SetChunkFilter(const FMassChunkConditionFunction& InFunction) { ChunkCondition = InFunction; }
 
 	void ClearChunkFilter() { ChunkCondition.Reset(); }
 
@@ -318,37 +318,37 @@ protected:
 	void ReadCommandlineParams();
 
 protected:
-	TArray<FLWComponentRequirement> Requirements;
-	TArray<FLWComponentRequirement> ChunkRequirements;
-	FLWTagBitSet RequiredAllTags;
-	FLWTagBitSet RequiredAnyTags;
-	FLWTagBitSet RequiredNoneTags;
-	FLWComponentBitSet RequiredAllComponents;
-	FLWComponentBitSet RequiredAnyComponents;
-	FLWComponentBitSet RequiredOptionalComponents;
-	FLWComponentBitSet RequiredNoneComponents;
-	FLWChunkComponentBitSet RequiredAllChunkComponents;
-	FLWChunkComponentBitSet RequiredNoneChunkComponents;
+	TArray<FMassFragmentRequirement> Requirements;
+	TArray<FMassFragmentRequirement> ChunkRequirements;
+	FMassTagBitSet RequiredAllTags;
+	FMassTagBitSet RequiredAnyTags;
+	FMassTagBitSet RequiredNoneTags;
+	FMassFragmentBitSet RequiredAllComponents;
+	FMassFragmentBitSet RequiredAnyComponents;
+	FMassFragmentBitSet RequiredOptionalComponents;
+	FMassFragmentBitSet RequiredNoneComponents;
+	FMassChunkFragmentBitSet RequiredAllChunkComponents;
+	FMassChunkFragmentBitSet RequiredNoneChunkComponents;
 
 private:
 	/** 
 	 * This function represents a condition that will be called for every chunk to be processed before the actual 
-	 * execution function is called. The chunk component requirements are already bound and ready to be used by the time 
+	 * execution function is called. The chunk fragment requirements are already bound and ready to be used by the time 
 	 * ChunkCondition is executed.
 	 */
-	FLWComponentSystemChunkConditionFunction ChunkCondition;
+	FMassChunkConditionFunction ChunkCondition;
 
 	uint32 EntitySubsystemHash = 0;
 	uint32 ArchetypeDataVersion = 0;
 
 	TArray<FArchetypeHandle> ValidArchetypes;
-	TArray<FLWRequirementIndicesMapping> ArchetypeComponentMapping;
+	TArray<FMassQueryRequirementIndicesMapping> ArchetypeComponentMapping;
 
 	bool bAllowParallelExecution = false;
 };
 
 template<>
-FORCEINLINE FLWComponentQuery& FLWComponentQuery::AddTagRequirements<ELWComponentPresence::All>(const FLWTagBitSet& TagBitSet)
+FORCEINLINE FMassEntityQuery& FMassEntityQuery::AddTagRequirements<EMassFragmentPresence::All>(const FMassTagBitSet& TagBitSet)
 {
 	RequiredAllTags += TagBitSet;
 	// force recaching the next time this query is used or the following CacheArchetypes call.
@@ -357,7 +357,7 @@ FORCEINLINE FLWComponentQuery& FLWComponentQuery::AddTagRequirements<ELWComponen
 }
 
 template<>
-FORCEINLINE FLWComponentQuery& FLWComponentQuery::AddTagRequirements<ELWComponentPresence::Any>(const FLWTagBitSet& TagBitSet)
+FORCEINLINE FMassEntityQuery& FMassEntityQuery::AddTagRequirements<EMassFragmentPresence::Any>(const FMassTagBitSet& TagBitSet)
 {
 	RequiredAnyTags += TagBitSet;
 	// force recaching the next time this query is used or the following CacheArchetypes call.
@@ -366,7 +366,7 @@ FORCEINLINE FLWComponentQuery& FLWComponentQuery::AddTagRequirements<ELWComponen
 }
 
 template<>
-FORCEINLINE FLWComponentQuery& FLWComponentQuery::AddTagRequirements<ELWComponentPresence::None>(const FLWTagBitSet& TagBitSet)
+FORCEINLINE FMassEntityQuery& FMassEntityQuery::AddTagRequirements<EMassFragmentPresence::None>(const FMassTagBitSet& TagBitSet)
 {
 	RequiredNoneTags += TagBitSet;
 	// force recaching the next time this query is used or the following CacheArchetypes call.
diff --git a/Engine/Plugins/Runtime/MassEntity/Source/MassEntity/Public/EntityView.h b/Engine/Plugins/Runtime/MassEntity/Source/MassEntity/Public/EntityView.h
--- a/Engine/Plugins/Runtime/MassEntity/Source/MassEntity/Public/EntityView.h
+++ b/Engine/Plugins/Runtime/MassEntity/Source/MassEntity/Public/EntityView.h
@@ -8,20 +8,20 @@
 
 
 class UMassEntitySubsystem;
-struct FArchetypeData;
+struct FMassArchetypeData;
 struct FArchetypeHandle;
-struct FArchetypeComponentConfig;
+struct FMassArchetypeFragmentConfig;
 
 /** 
  * The type representing a single entity in a single archetype. It's of a very transient nature so we guarantee it's 
  * validity only within the scope it has been created in. Don't store it. 
  */
 USTRUCT()
-struct MASSENTITY_API FEntityView
+struct MASSENTITY_API FMassEntityView
 {
 	GENERATED_BODY()
 
-	FEntityView() = default;
+	FMassEntityView() = default;
 
 	/** 
 	 *  Resolves Entity against ArchetypeHandle. Note that this approach requires the caller to ensure that Entity
@@ -29,24 +29,24 @@ struct MASSENTITY_API FEntityView
 	 *  UMassEntitySubsystem-flavored constructor is recommended since it will first find the appropriate archetype for
 	 *  Entity. 
 	 */
-	FEntityView(const FArchetypeHandle& ArchetypeHandle, FLWEntity Entity);
+	FMassEntityView(const FArchetypeHandle& ArchetypeHandle, FMassEntityHandle Entity);
 
 	/** 
 	 *  Finds the archetype Entity belongs to and then resolves against it. The caller is responsible for ensuring
 	 *  that the given Entity is in fact a valid ID tied to any of the archetypes 
 	 */
-	FEntityView(const UMassEntitySubsystem& EntitySubsystem, FLWEntity Entity);
+	FMassEntityView(const UMassEntitySubsystem& EntitySubsystem, FMassEntityHandle Entity);
 
-	FLWEntity GetEntity() const	{ return Entity; }
+	FMassEntityHandle GetEntity() const	{ return Entity; }
 
 	/** will fail a check if the viewed entity doesn't have the given component */	
 	template<typename T>
 	T& GetComponentData() const
 	{
-		static_assert(!TIsDerivedFrom<T, FComponentTag>::IsDerived,
+		static_assert(!TIsDerivedFrom<T, FMassTag>::IsDerived,
 			"Given struct doesn't represent a valid fragment type but a tag. Use HasTag instead.");
-		static_assert(TIsDerivedFrom<T, FComponentTag>::IsDerived || TIsDerivedFrom<T, FLWComponentData>::IsDerived,
-			"Given struct doesn't represent a valid fragment type. Make sure to inherit from FLWComponentData or one of its child-types.");
+		static_assert(TIsDerivedFrom<T, FMassTag>::IsDerived || TIsDerivedFrom<T, FMassFragment>::IsDerived,
+			"Given struct doesn't represent a valid fragment type. Make sure to inherit from FMassFragment or one of its child-types.");
 
 		return *((T*)GetComponentPtrChecked(*T::StaticStruct()));
 	}
@@ -55,10 +55,10 @@ struct MASSENTITY_API FEntityView
 	template<typename T>
 	T* GetComponentDataPtr() const
 	{
-		static_assert(!TIsDerivedFrom<T, FComponentTag>::IsDerived,
+		static_assert(!TIsDerivedFrom<T, FMassTag>::IsDerived,
 			"Given struct doesn't represent a valid fragment type but a tag. Use HasTag instead.");
-		static_assert(TIsDerivedFrom<T, FComponentTag>::IsDerived || TIsDerivedFrom<T, FLWComponentData>::IsDerived,
-			"Given struct doesn't represent a valid fragment type. Make sure to inherit from FLWComponentData or one of its child-types.");
+		static_assert(TIsDerivedFrom<T, FMassTag>::IsDerived || TIsDerivedFrom<T, FMassFragment>::IsDerived,
+			"Given struct doesn't represent a valid fragment type. Make sure to inherit from FMassFragment or one of its child-types.");
 
 		return (T*)GetComponentPtr(*T::StaticStruct());
 	}
@@ -72,12 +72,12 @@ struct MASSENTITY_API FEntityView
 	template<typename T>
 	bool HasTag() const
 	{
-		static_assert(TIsDerivedFrom<T, FComponentTag>::IsDerived, "Given struct doesn't represent a valid tag type. Make sure to inherit from FComponentTag or one of its child-types.");
+		static_assert(TIsDerivedFrom<T, FMassTag>::IsDerived, "Given struct doesn't represent a valid tag type. Make sure to inherit from FMassTag or one of its child-types.");
 		return HasTag(*T::StaticStruct());
 	}
 
 	bool IsSet() const { return Archetype != nullptr && EntityHandle.IsValid(); }
-	bool operator==(const FEntityView& Other) const { return Archetype == Other.Archetype && EntityHandle == Other.EntityHandle; }
+	bool operator==(const FMassEntityView& Other) const { return Archetype == Other.Archetype && EntityHandle == Other.EntityHandle; }
 
 protected:
 	void* GetComponentPtr(const UScriptStruct& ComponentType) const;
@@ -85,7 +85,7 @@ protected:
 	bool HasTag(const UScriptStruct& TagType) const;
 
 private:
-	FLWEntity Entity;
+	FMassEntityHandle Entity;
 	FInternalEntityHandle EntityHandle;
-	FArchetypeData* Archetype = nullptr;
+	FMassArchetypeData* Archetype = nullptr;
 };
diff --git a/Engine/Plugins/Runtime/MassEntity/Source/MassEntity/Public/LWCCommandBuffer.h b/Engine/Plugins/Runtime/MassEntity/Source/MassEntity/Public/LWCCommandBuffer.h
--- a/Engine/Plugins/Runtime/MassEntity/Source/MassEntity/Public/LWCCommandBuffer.h
+++ b/Engine/Plugins/Runtime/MassEntity/Source/MassEntity/Public/LWCCommandBuffer.h
@@ -18,12 +18,12 @@ struct MASSENTITY_API FCommandBufferEntryBase
 {
 	GENERATED_BODY()
 
-	FLWEntity TargetEntity;
+	FMassEntityHandle TargetEntity;
 
 	FCommandBufferEntryBase() = default;
 	virtual ~FCommandBufferEntryBase() = default;
 
-	FCommandBufferEntryBase(FLWEntity InEntity)
+	FCommandBufferEntryBase(FMassEntityHandle InEntity)
 		: TargetEntity(InEntity)
 	{}
 
@@ -63,7 +63,7 @@ struct MASSENTITY_API FBuildEntityFromComponentInstance : public FCommandBufferE
 	GENERATED_BODY()
 
 	FBuildEntityFromComponentInstance() = default;
-	FBuildEntityFromComponentInstance(const FLWEntity Entity, FStructView InStruct)
+	FBuildEntityFromComponentInstance(const FMassEntityHandle Entity, FStructView InStruct)
 		: FCommandBufferEntryBase(Entity)
 		, Struct(InStruct)
 	{}
@@ -83,7 +83,7 @@ struct MASSENTITY_API FBuildEntityFromComponentInstances : public FCommandBuffer
 	GENERATED_BODY()
 
 	FBuildEntityFromComponentInstances() = default;
-	FBuildEntityFromComponentInstances(const FLWEntity Entity, TConstArrayView<FInstancedStruct> InInstances)
+	FBuildEntityFromComponentInstances(const FMassEntityHandle Entity, TConstArrayView<FInstancedStruct> InInstances)
 		: FCommandBufferEntryBase(Entity)
 		, Instances(InInstances)
 	{}
@@ -103,7 +103,7 @@ struct MASSENTITY_API FCommandAddComponent : public FCommandBufferEntryBase
 	GENERATED_BODY()
 
 	FCommandAddComponent() = default;
-	FCommandAddComponent(FLWEntity InEntity, UScriptStruct* InStruct)
+	FCommandAddComponent(FMassEntityHandle InEntity, UScriptStruct* InStruct)
 		: FCommandBufferEntryBase(InEntity)
 		, StructParam(InStruct)
 	{}
@@ -123,7 +123,7 @@ struct MASSENTITY_API FCommandAddComponentInstance : public FCommandBufferEntryB
 	GENERATED_BODY()
 
 	FCommandAddComponentInstance() = default;
-	FCommandAddComponentInstance(const FLWEntity InEntity, FStructView InStruct)
+	FCommandAddComponentInstance(const FMassEntityHandle InEntity, FStructView InStruct)
         : FCommandBufferEntryBase(InEntity)
         , Struct(InStruct)
 	{}
@@ -143,7 +143,7 @@ struct MASSENTITY_API FCommandRemoveComponent : public FCommandBufferEntryBase
 	GENERATED_BODY()
 
 	FCommandRemoveComponent() = default;
-	FCommandRemoveComponent(FLWEntity InEntity, UScriptStruct* InStruct)
+	FCommandRemoveComponent(FMassEntityHandle InEntity, UScriptStruct* InStruct)
 		: FCommandBufferEntryBase(InEntity)
 		, StructParam(InStruct)
 	{}
@@ -163,7 +163,7 @@ struct MASSENTITY_API FCommandAddComponentList : public FCommandBufferEntryBase
 	GENERATED_BODY()
 
 	FCommandAddComponentList() = default;
-	FCommandAddComponentList(FLWEntity InEntity, TConstArrayView<const UScriptStruct*> InComponentList)
+	FCommandAddComponentList(FMassEntityHandle InEntity, TConstArrayView<const UScriptStruct*> InComponentList)
 		: FCommandBufferEntryBase(InEntity)
 		, ComponentList(InComponentList)
 	{}
@@ -183,7 +183,7 @@ struct MASSENTITY_API FCommandRemoveComponentList : public FCommandBufferEntryBa
 	GENERATED_BODY()
 
 	FCommandRemoveComponentList() = default;
-	FCommandRemoveComponentList(FLWEntity InEntity, TConstArrayView<const UScriptStruct*> InComponentList)
+	FCommandRemoveComponentList(FMassEntityHandle InEntity, TConstArrayView<const UScriptStruct*> InComponentList)
 		: FCommandBufferEntryBase(InEntity)
 		, ComponentList(InComponentList)
 	{}
@@ -203,7 +203,7 @@ struct MASSENTITY_API FCommandAddTag: public FCommandBufferEntryBase
 	GENERATED_BODY()
 
 	FCommandAddTag() = default;
-	FCommandAddTag(FLWEntity InEntity, UScriptStruct* InStruct)
+	FCommandAddTag(FMassEntityHandle InEntity, UScriptStruct* InStruct)
 		: FCommandBufferEntryBase(InEntity)
 		, StructParam(InStruct)
 	{}
@@ -223,7 +223,7 @@ struct MASSENTITY_API FCommandRemoveTag : public FCommandBufferEntryBase
 	GENERATED_BODY()
 
 	FCommandRemoveTag() = default;
-	FCommandRemoveTag(FLWEntity InEntity, UScriptStruct* InStruct)
+	FCommandRemoveTag(FMassEntityHandle InEntity, UScriptStruct* InStruct)
 		: FCommandBufferEntryBase(InEntity)
 		, StructParam(InStruct)
 	{}
@@ -245,13 +245,13 @@ struct MASSENTITY_API FCommandSwapTags : public FCommandBufferEntryBase
 
 	FCommandSwapTags() = default;
 
-	FCommandSwapTags(const FLWEntity InEntity, const UScriptStruct* InOldTagType, const UScriptStruct* InNewTagType)
+	FCommandSwapTags(const FMassEntityHandle InEntity, const UScriptStruct* InOldTagType, const UScriptStruct* InNewTagType)
 		: FCommandBufferEntryBase(InEntity)
 		, OldTagType(InOldTagType)
 		, NewTagType(InNewTagType)
 	{
-		checkf((InOldTagType == nullptr) || InOldTagType->IsChildOf(FComponentTag::StaticStruct()), TEXT("FCommandSwapTags works only with tags while '%s' is not one."), *GetPathNameSafe(InOldTagType));
-		checkf((InNewTagType == nullptr) || InNewTagType->IsChildOf(FComponentTag::StaticStruct()), TEXT("FCommandSwapTags works only with tags while '%s' is not one."), *GetPathNameSafe(InNewTagType));
+		checkf((InOldTagType == nullptr) || InOldTagType->IsChildOf(FMassTag::StaticStruct()), TEXT("FCommandSwapTags works only with tags while '%s' is not one."), *GetPathNameSafe(InOldTagType));
+		checkf((InNewTagType == nullptr) || InNewTagType->IsChildOf(FMassTag::StaticStruct()), TEXT("FCommandSwapTags works only with tags while '%s' is not one."), *GetPathNameSafe(InNewTagType));
 	}
 
 protected:
@@ -270,7 +270,7 @@ struct MASSENTITY_API FCommandRemoveComposition : public FCommandBufferEntryBase
 	GENERATED_BODY()
 
 	FCommandRemoveComposition() = default;
-	FCommandRemoveComposition(FLWEntity InEntity, const FLWCompositionDescriptor& InDescriptor)
+	FCommandRemoveComposition(FMassEntityHandle InEntity, const FMassCompositionDescriptor& InDescriptor)
 		: FCommandBufferEntryBase(InEntity)
 		, Descriptor(InDescriptor)
 	{}
@@ -278,10 +278,10 @@ struct MASSENTITY_API FCommandRemoveComposition : public FCommandBufferEntryBase
 protected:
 	virtual void Execute(UMassEntitySubsystem& System) const override;
 
-	FLWCompositionDescriptor Descriptor;
+	FMassCompositionDescriptor Descriptor;
 };
 
-struct MASSENTITY_API FLWCCommandBuffer
+struct MASSENTITY_API FMassCommandBuffer
 {
 public:
 
@@ -309,39 +309,39 @@ public:
 	}
 
 	template<typename T>
-	void AddComponent(FLWEntity Entity)
+	void AddComponent(FMassEntityHandle Entity)
 	{
-		static_assert(TIsDerivedFrom<T, FLWComponentData>::IsDerived, "Given struct type is not a valid component type.");
+		static_assert(TIsDerivedFrom<T, FMassFragment>::IsDerived, "Given struct type is not a valid component type.");
 		EmplaceCommand<FCommandAddComponent>(Entity, T::StaticStruct());
 	}
 
 	template<typename T>
-	void RemoveComponent(FLWEntity Entity)
+	void RemoveComponent(FMassEntityHandle Entity)
 	{
-		static_assert(TIsDerivedFrom<T, FLWComponentData>::IsDerived, "Given struct type is not a valid component type.");
+		static_assert(TIsDerivedFrom<T, FMassFragment>::IsDerived, "Given struct type is not a valid component type.");
 		EmplaceCommand<FCommandRemoveComponent>(Entity, T::StaticStruct());
 	}
 
 	template<typename T>
-	void AddTag(FLWEntity Entity)
+	void AddTag(FMassEntityHandle Entity)
 	{
-		static_assert(TIsDerivedFrom<T, FComponentTag>::IsDerived, "Given struct type is not a valid tag type.");
+		static_assert(TIsDerivedFrom<T, FMassTag>::IsDerived, "Given struct type is not a valid tag type.");
 		EmplaceCommand<FCommandAddTag>(Entity, T::StaticStruct());
 	}
 
 	template<typename T>
-	void RemoveTag(FLWEntity Entity)
+	void RemoveTag(FMassEntityHandle Entity)
 	{
-		static_assert(TIsDerivedFrom<T, FComponentTag>::IsDerived, "Given struct type is not a valid tag type.");
+		static_assert(TIsDerivedFrom<T, FMassTag>::IsDerived, "Given struct type is not a valid tag type.");
 		EmplaceCommand<FCommandRemoveTag>(Entity, T::StaticStruct());
 	}
 
-	void DestroyEntity(FLWEntity Entity)
+	void DestroyEntity(FMassEntityHandle Entity)
 	{
 		EntitiesToDestroy.Add(Entity);
 	}
 
-	void BatchDestroyEntities(const TArray<FLWEntity>& InEntitiesToDestroy)
+	void BatchDestroyEntities(const TArray<FMassEntityHandle>& InEntitiesToDestroy)
 	{
 		EntitiesToDestroy.Append(InEntitiesToDestroy);
 	}
@@ -355,7 +355,7 @@ public:
 	 * @param InOutOther the source buffer to copy the commands from. Note that after the call the InOutOther will be 
 	 *	emptied due to the function using Move semantics
 	 */
-	void MoveAppend(FLWCCommandBuffer& InOutOther);
+	void MoveAppend(FMassCommandBuffer& InOutOther);
 
 	int32 GetPendingCommandsCount() const { return PendingCommands.Num(); }
 
@@ -364,5 +364,5 @@ private:
 	UE_MT_DECLARE_RW_ACCESS_DETECTOR(PendingCommandsDetector);
 	FCriticalSection AppendingCommandsCS;
 
-	TArray<FLWEntity> EntitiesToDestroy;
+	TArray<FMassEntityHandle> EntitiesToDestroy;
 };
diff --git a/Engine/Plugins/Runtime/MassEntity/Source/MassEntity/Public/LWComponentTypes.h b/Engine/Plugins/Runtime/MassEntity/Source/MassEntity/Public/LWComponentTypes.h
--- a/Engine/Plugins/Runtime/MassEntity/Source/MassEntity/Public/LWComponentTypes.h
+++ b/Engine/Plugins/Runtime/MassEntity/Source/MassEntity/Public/LWComponentTypes.h
@@ -7,39 +7,39 @@
 
 // This is the base class for all lightweight components
 USTRUCT()
-struct FLWComponentData
+struct FMassFragment
 {
 	GENERATED_BODY()
 
-	FLWComponentData() {}
+	FMassFragment() {}
 };
 
 // This is the base class for types that will only be tested for presence/absence, i.e. Tags.
 // Subclasses should never contain any member properties.
 USTRUCT()
-struct FComponentTag
+struct FMassTag
 {
 	GENERATED_BODY()
 
-	FComponentTag() {}
+	FMassTag() {}
 };
 
 USTRUCT()
-struct FLWChunkComponent
+struct FMassChunkFragment
 {
 	GENERATED_BODY()
 
-	FLWChunkComponent() {}
+	FMassChunkFragment() {}
 };
 
 // A handle to a lightweight entity.  An entity is used in conjunction with the UMassEntitySubsystem
 // for the current world and can contain lightweight components.
 USTRUCT()
-struct FLWEntity
+struct FMassEntityHandle
 {
 	GENERATED_BODY()
 
-	FLWEntity() {}
+	FMassEntityHandle() {}
 	
 	UPROPERTY(VisibleAnywhere, Category = "AggregateTicking|Debug", Transient)
 	int32 Index = 0;
@@ -47,12 +47,12 @@ struct FLWEntity
 	UPROPERTY(VisibleAnywhere, Category = "AggregateTicking|Debug", Transient)
 	int32 SerialNumber = 0;
 
-	bool operator==(const FLWEntity& Other) const
+	bool operator==(const FMassEntityHandle& Other) const
 	{
 		return Index == Other.Index && SerialNumber == Other.SerialNumber;
 	}
 
-	bool operator!=(const FLWEntity& Other) const
+	bool operator!=(const FMassEntityHandle& Other) const
 	{
 		return !operator==(Other);
 	}
@@ -69,7 +69,7 @@ struct FLWEntity
 		Index = SerialNumber = 0;
 	}
 
-	friend uint32 GetTypeHash(const FLWEntity& Entity)
+	friend uint32 GetTypeHash(const FMassEntityHandle& Entity)
 	{
 		return HashCombine(Entity.Index, Entity.SerialNumber);
 	}
@@ -80,31 +80,31 @@ struct FLWEntity
 	}
 };
 
-template struct MASSENTITY_API TScriptStructTypeBitSet<FLWComponentData>;
-using FLWComponentBitSet = TScriptStructTypeBitSet<FLWComponentData>;
-template struct MASSENTITY_API TScriptStructTypeBitSet<FComponentTag>;
-using FLWTagBitSet = TScriptStructTypeBitSet<FComponentTag>;
-template struct MASSENTITY_API TScriptStructTypeBitSet<FLWChunkComponent>;
-using FLWChunkComponentBitSet = TScriptStructTypeBitSet<FLWChunkComponent>;
+template struct MASSENTITY_API TScriptStructTypeBitSet<FMassFragment>;
+using FMassFragmentBitSet = TScriptStructTypeBitSet<FMassFragment>;
+template struct MASSENTITY_API TScriptStructTypeBitSet<FMassTag>;
+using FMassTagBitSet = TScriptStructTypeBitSet<FMassTag>;
+template struct MASSENTITY_API TScriptStructTypeBitSet<FMassChunkFragment>;
+using FMassChunkFragmentBitSet = TScriptStructTypeBitSet<FMassChunkFragment>;
 
 /** The type summarily describing a composition of an entity or an archetype. It contains information on both the
  *  components as well as tags */
-struct FLWCompositionDescriptor
+struct FMassCompositionDescriptor
 {
-	FLWCompositionDescriptor() = default;
-	FLWCompositionDescriptor(const FLWComponentBitSet& InComponents, const FLWTagBitSet& InTags, const FLWChunkComponentBitSet& InChunkComponents)
+	FMassCompositionDescriptor() = default;
+	FMassCompositionDescriptor(const FMassFragmentBitSet& InComponents, const FMassTagBitSet& InTags, const FMassChunkFragmentBitSet& InChunkComponents)
 		: Components(InComponents), Tags(InTags), ChunkComponents(InChunkComponents)
 	{}
 
-	FLWCompositionDescriptor(TConstArrayView<const UScriptStruct*> InComponents, const FLWTagBitSet& InTags, const FLWChunkComponentBitSet& InChunkComponents)
-		: FLWCompositionDescriptor(FLWComponentBitSet(InComponents), InTags, InChunkComponents)
+	FMassCompositionDescriptor(TConstArrayView<const UScriptStruct*> InComponents, const FMassTagBitSet& InTags, const FMassChunkFragmentBitSet& InChunkComponents)
+		: FMassCompositionDescriptor(FMassFragmentBitSet(InComponents), InTags, InChunkComponents)
 	{}
 
-	FLWCompositionDescriptor(TConstArrayView<FInstancedStruct> InComponentInstances, const FLWTagBitSet& InTags, const FLWChunkComponentBitSet& InChunkComponents)
-		: FLWCompositionDescriptor(FLWComponentBitSet(InComponentInstances), InTags, InChunkComponents)
+	FMassCompositionDescriptor(TConstArrayView<FInstancedStruct> InComponentInstances, const FMassTagBitSet& InTags, const FMassChunkFragmentBitSet& InChunkComponents)
+		: FMassCompositionDescriptor(FMassFragmentBitSet(InComponentInstances), InTags, InChunkComponents)
 	{}
 
-	FLWCompositionDescriptor(FLWComponentBitSet&& InComponents, FLWTagBitSet&& InTags, FLWChunkComponentBitSet&& InChunkComponents)
+	FMassCompositionDescriptor(FMassFragmentBitSet&& InComponents, FMassTagBitSet&& InTags, FMassChunkFragmentBitSet&& InChunkComponents)
 		: Components(MoveTemp(InComponents)), Tags(MoveTemp(InTags)), ChunkComponents(MoveTemp(InChunkComponents))
 	{}
 
@@ -114,19 +114,19 @@ struct FLWCompositionDescriptor
 		Tags.Reset();
 	}
 
-	bool IsEquivalent(const FLWComponentBitSet& InComponentBitSet, const FLWTagBitSet& InTagBitSet, const FLWChunkComponentBitSet& InChunkComponentsBitSet) const
+	bool IsEquivalent(const FMassFragmentBitSet& InComponentBitSet, const FMassTagBitSet& InTagBitSet, const FMassChunkFragmentBitSet& InChunkComponentsBitSet) const
 	{
 		return Components.IsEquivalent(InComponentBitSet) && Tags.IsEquivalent(InTagBitSet) && ChunkComponents.IsEquivalent(InChunkComponentsBitSet);
 	}
 
-	bool IsEquivalent(const FLWCompositionDescriptor& OtherDescriptor) const
+	bool IsEquivalent(const FMassCompositionDescriptor& OtherDescriptor) const
 	{
 		return IsEquivalent(OtherDescriptor.Components, OtherDescriptor.Tags, OtherDescriptor.ChunkComponents);
 	}
 
 	bool IsEmpty() const { return Components.IsEmpty() && Tags.IsEmpty() && ChunkComponents.IsEmpty(); }
 
-	static uint32 CalculateHash(const FLWComponentBitSet& InComponents, const FLWTagBitSet& InTags, const FLWChunkComponentBitSet& InChunkComponents)
+	static uint32 CalculateHash(const FMassFragmentBitSet& InComponents, const FMassTagBitSet& InTags, const FMassChunkFragmentBitSet& InChunkComponents)
 	{
 		const uint32 ComponentsHash = GetTypeHash(InComponents);
 		const uint32 TagsHash = GetTypeHash(InTags);
@@ -139,7 +139,7 @@ struct FLWCompositionDescriptor
 		return CalculateHash(Components, Tags, ChunkComponents);
 	}
 
-	FLWComponentBitSet Components;
-	FLWTagBitSet Tags;
-	FLWChunkComponentBitSet ChunkComponents;
+	FMassFragmentBitSet Components;
+	FMassTagBitSet Tags;
+	FMassChunkFragmentBitSet ChunkComponents;
 };
diff --git a/Engine/Plugins/Runtime/MassEntity/Source/MassEntity/Public/MassEntityDebug.h b/Engine/Plugins/Runtime/MassEntity/Source/MassEntity/Public/MassEntityDebug.h
--- a/Engine/Plugins/Runtime/MassEntity/Source/MassEntity/Public/MassEntityDebug.h
+++ b/Engine/Plugins/Runtime/MassEntity/Source/MassEntity/Public/MassEntityDebug.h
@@ -5,13 +5,13 @@
 #include "MassEntityTypes.h"
 
 class FOutputDevice;
-class UPipeProcessor;
+class UMassProcessor;
 
-MASSENTITY_API DECLARE_ENUM_TO_STRING(EPipeProcessingPhase);
+MASSENTITY_API DECLARE_ENUM_TO_STRING(EMassProcessingPhase);
 
 namespace UE { namespace Pipe { namespace Debug {
 
-MASSENTITY_API extern void DebugOutputDescription(TConstArrayView<UPipeProcessor*> Processors, FOutputDevice& Ar);
+MASSENTITY_API extern void DebugOutputDescription(TConstArrayView<UMassProcessor*> Processors, FOutputDevice& Ar);
 
 }}} // namespace UE::Pipe::Debug
 
diff --git a/Engine/Plugins/Runtime/MassEntity/Source/MassEntity/Public/MassEntitySettings.h b/Engine/Plugins/Runtime/MassEntity/Source/MassEntity/Public/MassEntitySettings.h
--- a/Engine/Plugins/Runtime/MassEntity/Source/MassEntity/Public/MassEntitySettings.h
+++ b/Engine/Plugins/Runtime/MassEntity/Source/MassEntity/Public/MassEntitySettings.h
@@ -8,15 +8,15 @@
 #include "InstancedStruct.h"
 #include "MassEntitySettings.generated.h"
 
-#define GET_PIPE_CONFIG_VALUE(a) (GetMutableDefault<UPipeSettings>()->a)
+#define GET_PIPE_CONFIG_VALUE(a) (GetMutableDefault<UMassSettings>()->a)
 
-class UPipeProcessingPhaseManager;
-struct FPipeProcessingPhaseConfig;
+class UMassProcessingPhaseManager;
+struct FMassProcessingPhaseConfig;
 struct FPropertyChangedEvent;
 
 
 USTRUCT()
-struct FPipeProcessingPhaseConfig
+struct FMassProcessingPhaseConfig
 {
 	GENERATED_BODY()
 
@@ -24,19 +24,19 @@ struct FPipeProcessingPhaseConfig
 	FName PhaseName;
 
 	UPROPERTY(EditAnywhere, Category = Pipe, config, NoClear)
-	TSubclassOf<UPipeCompositeProcessor> PhaseGroupClass = UPipeCompositeProcessor::StaticClass();
+	TSubclassOf<UMassCompositeProcessor> PhaseGroupClass = UMassCompositeProcessor::StaticClass();
 
 	UPROPERTY(EditAnywhere, Category = Pipe, config, NoClear)
 	TArray<FName> OffGameThreadGroupNames;
 
 	UPROPERTY(Transient)
-	TArray<UPipeProcessor*> ProcessorCDOs;
+	TArray<UMassProcessor*> ProcessorCDOs;
 
 #if WITH_EDITORONLY_DATA
 	// this processor is available only in editor since it's used to present the user the order in which processors
 	// will be executed when given processing phase gets triggered
 	UPROPERTY(Transient)
-	UPipeCompositeProcessor* PhaseProcessor = nullptr;
+	UMassCompositeProcessor* PhaseProcessor = nullptr;
 
 	UPROPERTY(VisibleAnywhere, Category = Pipe, Transient)
 	FString Description;
@@ -47,7 +47,7 @@ struct FPipeProcessingPhaseConfig
  * Implements the settings for MassSimulation
  */
 UCLASS(config = Game, defaultconfig, DisplayName = "Pipe")
-class MASSENTITY_API UPipeSettings : public UDeveloperSettings
+class MASSENTITY_API UMassSettings : public UDeveloperSettings
 {
 	GENERATED_BODY()
 public:
@@ -55,13 +55,13 @@ public:
 	DECLARE_MULTICAST_DELEGATE_OneParam(FOnSettingsChange, const FPropertyChangedEvent& /*PropertyChangedEvent*/);
 #endif // WITH_EDITORONLY_DATA
 
-	UPipeSettings(const FObjectInitializer& ObjectInitializer = FObjectInitializer::Get());
+	UMassSettings(const FObjectInitializer& ObjectInitializer = FObjectInitializer::Get());
 
 	void BuildProcessorListAndPhases();
-	void AddToActiveProcessorsList(TSubclassOf<UPipeProcessor> ProcessorClass);
+	void AddToActiveProcessorsList(TSubclassOf<UMassProcessor> ProcessorClass);
 
-	const FPipeProcessingPhaseConfig* GetProcessingPhasesConfig();
-	const FPipeProcessingPhaseConfig& GetProcessingPhaseConfig(const EPipeProcessingPhase ProcessingPhase) const { check(ProcessingPhase != EPipeProcessingPhase::MAX); return ProcessingPhasesConfig[int(ProcessingPhase)]; }
+	const FMassProcessingPhaseConfig* GetProcessingPhasesConfig();
+	const FMassProcessingPhaseConfig& GetProcessingPhaseConfig(const EMassProcessingPhase ProcessingPhase) const { check(ProcessingPhase != EMassProcessingPhase::MAX); return ProcessingPhasesConfig[int(ProcessingPhase)]; }
 
 #if WITH_EDITOR
 	FOnSettingsChange& GetOnSettingsChange() { return OnSettingsChange; }	
@@ -87,11 +87,11 @@ public:
 
 	/** Lets users configure processing phases including the composite processor class to be used as a container for the phases' processors. */
 	UPROPERTY(EditDefaultsOnly, Category = Pipe, config)
-	FPipeProcessingPhaseConfig ProcessingPhasesConfig[(uint8)EPipeProcessingPhase::MAX];
+	FMassProcessingPhaseConfig ProcessingPhasesConfig[(uint8)EMassProcessingPhase::MAX];
 
 	/** This list contains all the processors available in the given binary (including plugins). The contents are sorted by display name.*/
 	UPROPERTY(VisibleAnywhere, Category = Pipe, Transient, Instanced)
-	TArray<UPipeProcessor*> ProcessorCDOs;
+	TArray<UMassProcessor*> ProcessorCDOs;
 
 #if WITH_EDITORONLY_DATA
 protected:
diff --git a/Engine/Plugins/Runtime/MassEntity/Source/MassEntity/Public/MassEntitySubsystem.h b/Engine/Plugins/Runtime/MassEntity/Source/MassEntity/Public/MassEntitySubsystem.h
--- a/Engine/Plugins/Runtime/MassEntity/Source/MassEntity/Public/MassEntitySubsystem.h
+++ b/Engine/Plugins/Runtime/MassEntity/Source/MassEntity/Public/MassEntitySubsystem.h
@@ -12,17 +12,17 @@
 MASSENTITY_API DECLARE_LOG_CATEGORY_EXTERN(LogAggregateTicking, Warning, All);
 
 class UMassEntitySubsystem;
-struct FLWComponentQuery;
-struct FLWComponentSystemExecutionContext;
-struct FArchetypeData;
-struct FLWCCommandBuffer;
+struct FMassEntityQuery;
+struct FMassExecutionContext;
+struct FMassArchetypeData;
+struct FMassCommandBuffer;
 struct FArchetypeChunkCollection;
-struct FArchetypeChunk;
+struct FMassArchetypeChunk;
 class FOutputDevice;
-enum class ELWComponentAccess : uint8;
+enum class EMassFragmentAccess : uint8;
 
 namespace UE { namespace AggregateTicking {
-FString DebugGetComponentAccessString(ELWComponentAccess Access);
+FString DebugGetComponentAccessString(EMassFragmentAccess Access);
 }}
 
 //@TODO: Comment this guy
@@ -31,14 +31,14 @@ class MASSENTITY_API UMassEntitySubsystem : public UWorldSubsystem
 {
 	GENERATED_BODY()
 
-	friend struct FLWComponentQuery;
+	friend struct FMassEntityQuery;
 private:
 	// Index 0 is reserved so we can treat that index as an invalid entity handle
 	constexpr static int32 NumReservedEntities = 1;
 
 	struct FEntityData
 	{
-		TSharedPtr<FArchetypeData> CurrentArchetype;
+		TSharedPtr<FMassArchetypeData> CurrentArchetype;
 		int32 SerialNumber = 0;
 
 		void Reset()
@@ -68,7 +68,7 @@ public:
 		std::atomic<int32>& ScopedProcessingCount;
 	};
 
-	const static FLWEntity InvalidEntity;
+	const static FMassEntityHandle InvalidEntity;
 
 	UMassEntitySubsystem();
 
@@ -87,7 +87,7 @@ public:
 	 */
 	FArchetypeHandle CreateArchetype(TConstArrayView<const UScriptStruct*> ComponentsAngTagsList);
 
-	FArchetypeHandle CreateArchetype(TConstArrayView<const UScriptStruct*> ComponentList, const FLWTagBitSet& Tags, const FLWChunkComponentBitSet& ChunkComponents);
+	FArchetypeHandle CreateArchetype(TConstArrayView<const UScriptStruct*> ComponentList, const FMassTagBitSet& Tags, const FMassChunkFragmentBitSet& ChunkComponents);
 
 	/** 
 	 *  Creates an archetype like SourceArchetype + NewComponentList. 
@@ -99,12 +99,12 @@ public:
 	 *   types that SourceArchetype doesn't already have. If the caller cannot guarantee it use of AddComponent functions
 	 *   family is recommended.
 	 */
-	FArchetypeHandle CreateArchetype(const TSharedPtr<FArchetypeData>& SourceArchetype, const FLWComponentBitSet& NewComponentList);
+	FArchetypeHandle CreateArchetype(const TSharedPtr<FMassArchetypeData>& SourceArchetype, const FMassFragmentBitSet& NewComponentList);
 
-	FArchetypeHandle CreateArchetype(const FLWCompositionDescriptor& Descriptor);
+	FArchetypeHandle CreateArchetype(const FMassCompositionDescriptor& Descriptor);
 
 
-	FArchetypeHandle GetArchetypeForEntity(FLWEntity Entity) const;
+	FArchetypeHandle GetArchetypeForEntity(FMassEntityHandle Entity) const;
 	/** Method to iterate on all the component types of an archetype */
 	static void ForEachArchetypeComponentType(const FArchetypeHandle Archetype, TFunction< void(const UScriptStruct* /*ComponentType*/)> Function);
 
@@ -120,71 +120,71 @@ public:
 	 * Creates fully built entity ready to be used by the subsystem
 	 * @param Archetype you want this entity to be
 	 * @return FLWEntity id of the newly created entity */
-	FLWEntity CreateEntity(const FArchetypeHandle Archetype);
+	FMassEntityHandle CreateEntity(const FArchetypeHandle Archetype);
 
 	/**
 	 * Creates fully built entity ready to be used by the subsystem
 	 * @param ComponentInstanceList is the components to create the entity from and initialize values
 	 * @return FLWEntity id of the newly created entity */
-	FLWEntity CreateEntity(TConstArrayView<FInstancedStruct> ComponentInstanceList);
+	FMassEntityHandle CreateEntity(TConstArrayView<FInstancedStruct> ComponentInstanceList);
 
 	/** A version of CreateEntity that's creating a number of entities (Count) at one go
 	 *  @param Archetype you want this entity to be
 	 *  @param Count number of entities to create
 	 *  @param OutEntities the newly created entities are appended to given array, i.e. the pre-existing content of OutEntities won't be affected by the call
 	 *  @return number of entities created */
-	int32 BatchCreateEntities(const FArchetypeHandle Archetype, const int32 Count, TArray<FLWEntity>& OutEntities);
+	int32 BatchCreateEntities(const FArchetypeHandle Archetype, const int32 Count, TArray<FMassEntityHandle>& OutEntities);
 
 	/**
 	 * Destroys a fully built entity, use ReleaseReservedEntity if entity was not yet built.
 	 * @param Entity to destroy */
-	void DestroyEntity(FLWEntity Entity);
+	void DestroyEntity(FMassEntityHandle Entity);
 
 	/**
 	 * Reserves an entity in the subsystem, the entity is still not ready to be used by the subsystem, need to call BuildEntity()
 	 * @return FLWEntity id of the reserved entity */
-	FLWEntity ReserveEntity();
+	FMassEntityHandle ReserveEntity();
 
 	/**
 	 * Builds an entity for it to be ready to be used by the subsystem
 	 * @param Entity to build which was retrieved with ReserveEntity() method
 	 * @param Archetype you want this entity to be*/
-	void BuildEntity(FLWEntity Entity, const FArchetypeHandle Archetype);
+	void BuildEntity(FMassEntityHandle Entity, const FArchetypeHandle Archetype);
 
 	/**
 	 * Builds an entity for it to be ready to be used by the subsystem
 	 * @param Entity to build which was retrieved with ReserveEntity() method
 	 * @param ComponentInstanceList is the components to create the entity from and initialize values*/
-	void BuildEntity(FLWEntity Entity, TConstArrayView<FInstancedStruct> ComponentInstanceList);
+	void BuildEntity(FMassEntityHandle Entity, TConstArrayView<FInstancedStruct> ComponentInstanceList);
 
 	/*
 	 * Releases a previously reserved entity that was not yet built, otherwise call DestroyEntity
 	 * @param Entity to release */
-	void ReleaseReservedEntity(FLWEntity Entity);
+	void ReleaseReservedEntity(FMassEntityHandle Entity);
 
 	/**
 	 * Destroys all the entity in the provided array of entities
 	 * @param InEntities to destroy
 	 */
-	void BatchDestroyEntities(TConstArrayView<FLWEntity> InEntities);
+	void BatchDestroyEntities(TConstArrayView<FMassEntityHandle> InEntities);
 
 	void BatchDestroyEntityChunks(const FArchetypeChunkCollection& Chunks);
 
-	void AddComponentToEntity(FLWEntity Entity, const UScriptStruct* ComponentType);
+	void AddComponentToEntity(FMassEntityHandle Entity, const UScriptStruct* ComponentType);
 
 	/** 
 	 *  Ensures that only unique components are added. 
 	 *  @note It's caller's responsibility to ensure Entity's and ComponentList's validity. 
 	 */
-	void AddComponentListToEntity(FLWEntity Entity, TConstArrayView<const UScriptStruct*> ComponentList);
+	void AddComponentListToEntity(FMassEntityHandle Entity, TConstArrayView<const UScriptStruct*> ComponentList);
 
-	void AddComponentInstanceListToEntity(FLWEntity Entity, TConstArrayView<FInstancedStruct> ComponentInstanceList);
-	void RemoveComponentFromEntity(FLWEntity Entity, const UScriptStruct* ComponentType);
-	void RemoveComponentListFromEntity(FLWEntity Entity, TConstArrayView<const UScriptStruct*> ComponentList);
+	void AddComponentInstanceListToEntity(FMassEntityHandle Entity, TConstArrayView<FInstancedStruct> ComponentInstanceList);
+	void RemoveComponentFromEntity(FMassEntityHandle Entity, const UScriptStruct* ComponentType);
+	void RemoveComponentListFromEntity(FMassEntityHandle Entity, TConstArrayView<const UScriptStruct*> ComponentList);
 
-	void AddTagToEntity(FLWEntity Entity, const UScriptStruct* TagType);
-	void RemoveTagFromEntity(FLWEntity Entity, const UScriptStruct* TagType);
-	void SwapTagsForEntity(FLWEntity Entity, const UScriptStruct* FromComponentType, const UScriptStruct* ToComponentType);
+	void AddTagToEntity(FMassEntityHandle Entity, const UScriptStruct* TagType);
+	void RemoveTagFromEntity(FMassEntityHandle Entity, const UScriptStruct* TagType);
+	void SwapTagsForEntity(FMassEntityHandle Entity, const UScriptStruct* FromComponentType, const UScriptStruct* ToComponentType);
 
 
 	/**
@@ -192,19 +192,19 @@ public:
 	 * in InOutDescriptor are missing from the current composition of the given entity and then returns the resulting 
 	 * delta via InOutDescriptor.
 	 */
-	void AddCompositionToEntity_GetDelta(FLWEntity Entity, FLWCompositionDescriptor& InOutDescriptor);
-	void RemoveCompositionFromEntity(FLWEntity Entity, const FLWCompositionDescriptor& InDescriptor);
+	void AddCompositionToEntity_GetDelta(FMassEntityHandle Entity, FMassCompositionDescriptor& InOutDescriptor);
+	void RemoveCompositionFromEntity(FMassEntityHandle Entity, const FMassCompositionDescriptor& InDescriptor);
 
 	/** 
 	 * Moves an entity over to a new archetype by copying over components common to both archetypes
 	 * @param Entity is the entity to move 
 	 * @param NewArchetypeHandle the handle to the new archetype
 	 */
-	void MoveEntityToAnotherArchetype(FLWEntity Entity, FArchetypeHandle NewArchetypeHandle);
+	void MoveEntityToAnotherArchetype(FMassEntityHandle Entity, FArchetypeHandle NewArchetypeHandle);
 
 	/** Copies values from ComponentInstanceList over to Entity's component. Caller is responsible for ensuring that 
 	 *  the given entity does have given components. Failing this assumption will cause a check-fail.*/
-	void SetEntityComponentsValues(FLWEntity Entity, TArrayView<const FInstancedStruct> ComponentInstanceList);
+	void SetEntityComponentsValues(FMassEntityHandle Entity, TArrayView<const FInstancedStruct> ComponentInstanceList);
 
 	/** Copies values from ComponentInstanceList over to components of given entities collection. The caller is responsible 
 	 *  for ensuring that the given entity archetype (FArchetypeChunkCollection.Archetype) does have given components. 
@@ -212,36 +212,36 @@ public:
 	static void BatchSetEntityComponentsValues(const FArchetypeChunkCollection& SparseEntities, TArrayView<const FInstancedStruct> ComponentInstanceList);
 
 	// Return true if it is an valid built entity
-	bool IsEntityActive(FLWEntity Entity) const 
+	bool IsEntityActive(FMassEntityHandle Entity) const 
 	{
 		return IsEntityValid(Entity) && IsEntityBuilt(Entity);
 	}
 
 	// Returns true if Entity is valid
-	bool IsEntityValid(FLWEntity Entity) const;
+	bool IsEntityValid(FMassEntityHandle Entity) const;
 
 	// Returns true if Entity is has been fully built (expecting a valid Entity)
-	bool IsEntityBuilt(FLWEntity Entity) const;
+	bool IsEntityBuilt(FMassEntityHandle Entity) const;
 
 	// Asserts that IsEntityValid
-	void CheckIfEntityIsValid(FLWEntity Entity) const;
+	void CheckIfEntityIsValid(FMassEntityHandle Entity) const;
 
 	// Asserts that IsEntityBuilt
-	void CheckIfEntityIsActive(FLWEntity Entity) const;
+	void CheckIfEntityIsActive(FMassEntityHandle Entity) const;
 
 	template <typename ComponentType>
-	ComponentType& GetComponentDataChecked(FLWEntity Entity) const
+	ComponentType& GetComponentDataChecked(FMassEntityHandle Entity) const
 	{
 		return *((ComponentType*)InternalGetComponentDataChecked(Entity, ComponentType::StaticStruct()));
 	}
 
 	template <typename ComponentType>
-	ComponentType* GetComponentDataPtr(FLWEntity Entity) const
+	ComponentType* GetComponentDataPtr(FMassEntityHandle Entity) const
 	{
 		return (ComponentType*)InternalGetComponentDataPtr(Entity, ComponentType::StaticStruct());
 	}
 
-	FStructView GetComponentDataStruct(FLWEntity Entity, const UScriptStruct* ComponentType) const
+	FStructView GetComponentDataStruct(FMassEntityHandle Entity, const UScriptStruct* ComponentType) const
 	{
 		return FStructView(ComponentType, static_cast<uint8*>(InternalGetComponentDataPtr(Entity, ComponentType)));
 	}
@@ -249,17 +249,17 @@ public:
 	uint32 GetArchetypeDataVersion() const { return ArchetypeDataVersion; }
 
 	/**
-	 * Creates and initializes a FLWComponentSystemExecutionContext instance.
+	 * Creates and initializes a FMassExecutionContext instance.
 	 */
-	FLWComponentSystemExecutionContext CreateExecutionContext(const float DeltaSeconds) const;
+	FMassExecutionContext CreateExecutionContext(const float DeltaSeconds) const;
 
 	FScopedProcessing NewProcessingScope() { return FScopedProcessing(ProcessingScopeCount); }
 	bool IsProcessing() const { return ProcessingScopeCount > 0; }
-	FLWCCommandBuffer& Defer() { return *DeferredCommandBuffer.Get(); }
+	FMassCommandBuffer& Defer() { return *DeferredCommandBuffer.Get(); }
 
 #if WITH_MASSENTITY_DEBUG
 	void DebugPrintEntity(int32 Index, FOutputDevice& Ar, const TCHAR* InPrefix = TEXT("")) const;
-	void DebugPrintEntity(FLWEntity Entity, FOutputDevice& Ar, const TCHAR* InPrefix = TEXT("")) const;
+	void DebugPrintEntity(FMassEntityHandle Entity, FOutputDevice& Ar, const TCHAR* InPrefix = TEXT("")) const;
 	void DebugPrintArchetypes(FOutputDevice& Ar) const;
 	static void DebugGetStringDesc(const FArchetypeHandle& Archetype, FOutputDevice& Ar);
 	void DebugGetArchetypesStringDetails(FOutputDevice& Ar, const bool bIncludeEmpty = true);
@@ -273,53 +273,32 @@ public:
 #endif // WITH_MASSENTITY_DEBUG
 
 protected:
-	void GetValidArchetypes(const FLWComponentQuery& Query, TArray<FArchetypeHandle>& OutValidArchetypes);
+	void GetValidArchetypes(const FMassEntityQuery& Query, TArray<FArchetypeHandle>& OutValidArchetypes);
 	
-	FArchetypeHandle InternalCreateSiblingArchetype(const TSharedPtr<FArchetypeData>& SourceArchetype, const FLWTagBitSet& OverrideTags);
+	FArchetypeHandle InternalCreateSiblingArchetype(const TSharedPtr<FMassArchetypeData>& SourceArchetype, const FMassTagBitSet& OverrideTags);
 
 private:
-	void InternalBuildEntity(FLWEntity Entity, const FArchetypeHandle Archetype);
-	void InternalReleaseEntity(FLWEntity Entity);
+	void InternalBuildEntity(FMassEntityHandle Entity, const FArchetypeHandle Archetype);
+	void InternalReleaseEntity(FMassEntityHandle Entity);
 
 	/** 
 	 *  Adds components in ComponentList to Entity. Only the unique components will be added.
 	 */
-	void InternalAddComponentListToEntityChecked(FLWEntity Entity, const FLWComponentBitSet& InComponents);
+	void InternalAddComponentListToEntityChecked(FMassEntityHandle Entity, const FMassFragmentBitSet& InComponents);
 
 	/** 
 	 *  Similar to InternalAddComponentListToEntity but expects NewComponentList not overlapping with current entity's
 	 *  component list. It's callers responsibility to ensure that's true. Failing this will cause a `check` fail.
 	 */
-	void InternalAddComponentListToEntity(FLWEntity Entity, const FLWComponentBitSet& NewComponents);
-	void* InternalGetComponentDataChecked(FLWEntity Entity, const UScriptStruct* ComponentType) const;
-	void* InternalGetComponentDataPtr(FLWEntity Entity, const UScriptStruct* ComponentType) const;
-
-	//@TODO: Not optimized or a great API, super WIP
-	struct MASSENTITY_API FChunkIteratorByComponent
-	{
-		FChunkIteratorByComponent(UMassEntitySubsystem* InSystem, const UScriptStruct* InComponentType);
-
-		void GetCurrentChunkBase(void*& OutChunkBasePtr, int32& OutChunkEntityCount) const;
-
-		void Advance();
-		bool IsDone() const;
-		FArchetypeData* GetCurrentArchetype() const;
-
-		UMassEntitySubsystem* System;
-		const UScriptStruct* ComponentType;
-
-		TArray<TSharedPtr<FArchetypeData>>* pPendingArchetypes;
-		int32 ArchetypeIndex = 0;
-		int32 ChunkIndex = 0;
-	};
-
-	friend FChunkIteratorByComponent;
+	void InternalAddComponentListToEntity(FMassEntityHandle Entity, const FMassFragmentBitSet& NewComponents);
+	void* InternalGetComponentDataChecked(FMassEntityHandle Entity, const UScriptStruct* ComponentType) const;
+	void* InternalGetComponentDataPtr(FMassEntityHandle Entity, const UScriptStruct* ComponentType) const;
 
 private:
 	TChunkedArray<FEntityData> Entities;
 	TArray<int32> EntityFreeIndexList;
 
-	TSharedPtr<FLWCCommandBuffer> DeferredCommandBuffer;
+	TSharedPtr<FMassCommandBuffer> DeferredCommandBuffer;
 	std::atomic<int32> SerialNumberGenerator;
 	std::atomic<int32> ProcessingScopeCount;
 
@@ -327,46 +306,46 @@ private:
 	uint32 ArchetypeDataVersion = 0;
 
 	// Map of hash of sorted component list to archetypes with that hash
-	TMap<uint32, TArray<TSharedPtr<FArchetypeData>>> ComponentHashToArchetypeMap;
+	TMap<uint32, TArray<TSharedPtr<FMassArchetypeData>>> ComponentHashToArchetypeMap;
 
 	// Map to list of archetypes that contain the specified component type
-	TMap<const UScriptStruct*, TArray<TSharedPtr<FArchetypeData>>> ComponentTypeToArchetypeMap;
+	TMap<const UScriptStruct*, TArray<TSharedPtr<FMassArchetypeData>>> ComponentTypeToArchetypeMap;
 };
 
 
 //////////////////////////////////////////////////////////////////////
 //
 
-struct MASSENTITY_API FLWComponentSystemExecutionContext
+struct MASSENTITY_API FMassExecutionContext
 {
 private:
-	struct FComponentView 
+	struct FFragmentView 
 	{
-		FLWComponentRequirement Requirement;
-		TArrayView<FLWComponentData> ComponentView;
+		FMassFragmentRequirement Requirement;
+		TArrayView<FMassFragment> ComponentView;
 
-		FComponentView() {}
-		explicit FComponentView(const FLWComponentRequirement& InRequirement) : Requirement(InRequirement) {}
+		FFragmentView() {}
+		explicit FFragmentView(const FMassFragmentRequirement& InRequirement) : Requirement(InRequirement) {}
 
 		bool operator==(const UScriptStruct* ComponentType) const { return Requirement.StructType == ComponentType; }
 	};
-	TArray<FComponentView, TInlineAllocator<8>> ComponentViews;
+	TArray<FFragmentView, TInlineAllocator<8>> ComponentViews;
 
-	struct FChunkComponentView
+	struct FChunkFragmentView
 	{
-		FLWComponentRequirement Requirement;
+		FMassFragmentRequirement Requirement;
 		FStructView ChunkComponentView;
 
-		FChunkComponentView() {}
-		explicit FChunkComponentView(const FLWComponentRequirement& InRequirement) : Requirement(InRequirement)	{}
+		FChunkFragmentView() {}
+		explicit FChunkFragmentView(const FMassFragmentRequirement& InRequirement) : Requirement(InRequirement)	{}
 
 		bool operator==(const UScriptStruct* ComponentType) const { return Requirement.StructType == ComponentType; }
 	};
-	TArray<FChunkComponentView, TInlineAllocator<4>> ChunkComponents;
+	TArray<FChunkFragmentView, TInlineAllocator<4>> ChunkComponents;
 
 	// mz@todo make this shared ptr thread-safe and never auto-flush in MT environment. 
-	TSharedPtr<FLWCCommandBuffer> DeferredCommandBuffer;
-	TArrayView<FLWEntity> EntityListView;
+	TSharedPtr<FMassCommandBuffer> DeferredCommandBuffer;
+	TArrayView<FMassEntityHandle> EntityListView;
 	
 	/** If set this indicates the exact archetype and its chunks to be processed. 
 	 *  @todo this data should live somewhere else, preferably be just a parameter to Query.ForEachEntityChunk function */
@@ -375,7 +354,7 @@ private:
 	FInstancedStruct AuxData;
 	float DeltaTimeSeconds = 0.0f;
 	int32 ChunkSerialModificationNumber = -1;
-	FLWTagBitSet CurrentArchetypesTagBitSet;
+	FMassTagBitSet CurrentArchetypesTagBitSet;
 
 #if WITH_MASSENTITY_DEBUG
 	FString DebugExecutionDescription;
@@ -386,15 +365,15 @@ private:
 	 *  calling FLWCCommandBuffer.ReplayBufferAgainstSystem() */
 	bool bFlushDeferredCommands = true;
 
-	TArrayView<FComponentView> GetMutableRequirements() { return ComponentViews; }
-	TArrayView<FChunkComponentView> GetMutableChunkRequirements() { return ChunkComponents; }
+	TArrayView<FFragmentView> GetMutableRequirements() { return ComponentViews; }
+	TArrayView<FChunkFragmentView> GetMutableChunkRequirements() { return ChunkComponents; }
 	
-	friend FArchetypeData;
-	friend FLWComponentQuery;
+	friend FMassArchetypeData;
+	friend FMassEntityQuery;
 
 public:
-	FLWComponentSystemExecutionContext() = default;
-	explicit FLWComponentSystemExecutionContext(const float InDeltaTimeSeconds, const bool bInFlushDeferredCommands = true)
+	FMassExecutionContext() = default;
+	explicit FMassExecutionContext(const float InDeltaTimeSeconds, const bool bInFlushDeferredCommands = true)
 		: DeltaTimeSeconds(InDeltaTimeSeconds)
 		, bFlushDeferredCommands(bInFlushDeferredCommands)
 	{}
@@ -407,7 +386,7 @@ public:
 	/** Sets bFlushDeferredCommands. Note that setting to True while the system is being executed doesn't result in
 	 *  immediate commands flushing */
 	void SetFlushDeferredCommands(const bool bNewFlushDeferredCommands) { bFlushDeferredCommands = bNewFlushDeferredCommands; } 
-	void SetDeferredCommandBuffer(const TSharedPtr<FLWCCommandBuffer>& InDeferredCommandBuffer) { DeferredCommandBuffer = InDeferredCommandBuffer; }
+	void SetDeferredCommandBuffer(const TSharedPtr<FMassCommandBuffer>& InDeferredCommandBuffer) { DeferredCommandBuffer = InDeferredCommandBuffer; }
 	void SetChunkCollection(const FArchetypeChunkCollection& InChunkCollection);
 	void SetChunkCollection(FArchetypeChunkCollection&& InChunkCollection);
 	void ClearChunkCollection() { ChunkCollection.Reset(); }
@@ -418,13 +397,13 @@ public:
 		return DeltaTimeSeconds;
 	}
 
-	TSharedPtr<FLWCCommandBuffer> GetSharedDeferredCommandBuffer() const { return DeferredCommandBuffer; }
-	FLWCCommandBuffer& Defer() { checkSlow(DeferredCommandBuffer.IsValid()); return *DeferredCommandBuffer.Get(); }
+	TSharedPtr<FMassCommandBuffer> GetSharedDeferredCommandBuffer() const { return DeferredCommandBuffer; }
+	FMassCommandBuffer& Defer() { checkSlow(DeferredCommandBuffer.IsValid()); return *DeferredCommandBuffer.Get(); }
 
-	TConstArrayView<FLWEntity> GetEntities() const { return EntityListView; }
+	TConstArrayView<FMassEntityHandle> GetEntities() const { return EntityListView; }
 	int32 GetEntitiesNum() const { return EntityListView.Num(); }
 
-	FLWEntity GetEntity(const int32 Index) const
+	FMassEntityHandle GetEntity(const int32 Index) const
 	{
 		return EntityListView[Index];
 	}
@@ -432,7 +411,7 @@ public:
 	template<typename T>
 	bool DoesArchetypeHaveTag() const
 	{
-		static_assert(TIsDerivedFrom<T, FComponentTag>::IsDerived, "Given struct is not of a valid fragment type.");
+		static_assert(TIsDerivedFrom<T, FMassTag>::IsDerived, "Given struct is not of a valid fragment type.");
 		return CurrentArchetypesTagBitSet.Contains<T>();
 	}
 
@@ -443,10 +422,10 @@ public:
 	template<typename T>
 	T& GetMutableChunkComponent()
 	{
-		static_assert(TIsDerivedFrom<T, FLWChunkComponent>::IsDerived, "Given struct doesn't represent a valid chunk fragment type. Make sure to inherit from FLWChunkComponent or one of its child-types.");
+		static_assert(TIsDerivedFrom<T, FMassChunkFragment>::IsDerived, "Given struct doesn't represent a valid chunk fragment type. Make sure to inherit from FMassChunkFragment or one of its child-types.");
 
 		const UScriptStruct* Type = T::StaticStruct();
-		const FChunkComponentView* FoundChunkComponentData = ChunkComponents.FindByPredicate([Type](const FChunkComponentView& Element) { return Element.Requirement.StructType == Type; });
+		const FChunkFragmentView* FoundChunkComponentData = ChunkComponents.FindByPredicate([Type](const FChunkFragmentView& Element) { return Element.Requirement.StructType == Type; });
 		checkf(FoundChunkComponentData, TEXT("Chunk Component requirement not found: %s"), *T::StaticStruct()->GetName());
 		return FoundChunkComponentData->ChunkComponentView.template GetMutable<T>();
 	}
@@ -454,43 +433,43 @@ public:
 	template<typename T>
 	const T& GetChunkComponent() const
 	{
-		static_assert(TIsDerivedFrom<T, FLWChunkComponent>::IsDerived, "Given struct doesn't represent a valid chunk fragment type. Make sure to inherit from FLWChunkComponent or one of its child-types.");
+		static_assert(TIsDerivedFrom<T, FMassChunkFragment>::IsDerived, "Given struct doesn't represent a valid chunk fragment type. Make sure to inherit from FMassChunkFragment or one of its child-types.");
 
 		const UScriptStruct* Type = T::StaticStruct();
-		const FChunkComponentView* FoundChunkComponentData = ChunkComponents.FindByPredicate([Type](const FChunkComponentView& Element) { return Element.Requirement.StructType == Type; } );
+		const FChunkFragmentView* FoundChunkComponentData = ChunkComponents.FindByPredicate([Type](const FChunkFragmentView& Element) { return Element.Requirement.StructType == Type; } );
 		checkf(FoundChunkComponentData, TEXT("Chunk Component requirement not found: %s"), *T::StaticStruct()->GetName());
 		return FoundChunkComponentData->ChunkComponentView.template Get<T>();
 	}
 
-	template<typename TLWComponent>
-	TArrayView<TLWComponent> GetMutableComponentView()
+	template<typename TFragment>
+	TArrayView<TFragment> GetMutableComponentView()
 	{
-		const UScriptStruct* ComponentType = TLWComponent::StaticStruct();
-		const FComponentView* View = ComponentViews.FindByPredicate([ComponentType](const FComponentView& Element) { return Element.Requirement.StructType == ComponentType; });
+		const UScriptStruct* ComponentType = TFragment::StaticStruct();
+		const FFragmentView* View = ComponentViews.FindByPredicate([ComponentType](const FFragmentView& Element) { return Element.Requirement.StructType == ComponentType; });
 		//checkfSlow(View != nullptr, TEXT("Requested component type not bound"));
-		//checkfSlow(View->Requirement.AccessMode == ELWComponentAccess::ReadWrite, TEXT("Requested component has not been bound for writing"));
-		return MakeArrayView<TLWComponent>((TLWComponent*)View->ComponentView.GetData(), View->ComponentView.Num());
+		//checkfSlow(View->Requirement.AccessMode == EMassFragmentAccess::ReadWrite, TEXT("Requested component has not been bound for writing"));
+		return MakeArrayView<TFragment>((TFragment*)View->ComponentView.GetData(), View->ComponentView.Num());
 	}
 
-	template<typename TLWComponent>
-	TConstArrayView<TLWComponent> GetComponentView() const
+	template<typename TFragment>
+	TConstArrayView<TFragment> GetComponentView() const
 	{
-		const UScriptStruct* ComponentType = TLWComponent::StaticStruct();
-		const FComponentView* View = ComponentViews.FindByPredicate([ComponentType](const FComponentView& Element) { return Element.Requirement.StructType == ComponentType; });
+		const UScriptStruct* ComponentType = TFragment::StaticStruct();
+		const FFragmentView* View = ComponentViews.FindByPredicate([ComponentType](const FFragmentView& Element) { return Element.Requirement.StructType == ComponentType; });
 		//checkfSlow(View != nullptr, TEXT("Requested component type not bound"));
-		return TConstArrayView<TLWComponent>((const TLWComponent*)View->ComponentView.GetData(), View->ComponentView.Num());
+		return TConstArrayView<TFragment>((const TFragment*)View->ComponentView.GetData(), View->ComponentView.Num());
 	}
 
-	TConstArrayView<FLWComponentData> GetComponentComponentView(const UScriptStruct* ComponentType) const
+	TConstArrayView<FMassFragment> GetComponentComponentView(const UScriptStruct* ComponentType) const
 	{
-		const FComponentView* View = ComponentViews.FindByPredicate([ComponentType](const FComponentView& Element) { return Element.Requirement.StructType == ComponentType; });
+		const FFragmentView* View = ComponentViews.FindByPredicate([ComponentType](const FFragmentView& Element) { return Element.Requirement.StructType == ComponentType; });
 		checkSlow(View);
-		return TConstArrayView<FLWComponentData>((const FLWComponentData*)View->ComponentView.GetData(), View->ComponentView.Num());;
+		return TConstArrayView<FMassFragment>((const FMassFragment*)View->ComponentView.GetData(), View->ComponentView.Num());;
 	}
 
-	TArrayView<FLWComponentData> GetMutableComponentView(const UScriptStruct* ComponentType) 
+	TArrayView<FMassFragment> GetMutableComponentView(const UScriptStruct* ComponentType) 
 	{
-		const FComponentView* View = ComponentViews.FindByPredicate([ComponentType](const FComponentView& Element) { return Element.Requirement.StructType == ComponentType; });
+		const FFragmentView* View = ComponentViews.FindByPredicate([ComponentType](const FFragmentView& Element) { return Element.Requirement.StructType == ComponentType; });
 		checkSlow(View);
 		return View->ComponentView;
 	}
@@ -504,17 +483,17 @@ public:
 	void FlushDeferred(UMassEntitySubsystem& EntitySystem) const;
 
 	void ClearExecutionData();
-	void SetCurrentArchetypeData(FArchetypeData& ArchetypeData);
+	void SetCurrentArchetypeData(FMassArchetypeData& ArchetypeData);
 
 protected:
-	void SetRequirements(TConstArrayView<FLWComponentRequirement> InRequirements, TConstArrayView<FLWComponentRequirement> InChunkRequirements);
+	void SetRequirements(TConstArrayView<FMassFragmentRequirement> InRequirements, TConstArrayView<FMassFragmentRequirement> InChunkRequirements);
 	void ClearComponentViews()
 	{
-		for (FComponentView& View : ComponentViews)
+		for (FFragmentView& View : ComponentViews)
 		{
-			View.ComponentView = TArrayView<FLWComponentData>();
+			View.ComponentView = TArrayView<FMassFragment>();
 		}
-		for (FChunkComponentView& View : ChunkComponents)
+		for (FChunkFragmentView& View : ChunkComponents)
 		{
 			View.ChunkComponentView.Reset();
 		}
diff --git a/Engine/Plugins/Runtime/MassEntity/Source/MassEntity/Public/MassEntityTypes.h b/Engine/Plugins/Runtime/MassEntity/Source/MassEntity/Public/MassEntityTypes.h
--- a/Engine/Plugins/Runtime/MassEntity/Source/MassEntity/Public/MassEntityTypes.h
+++ b/Engine/Plugins/Runtime/MassEntity/Source/MassEntity/Public/MassEntityTypes.h
@@ -11,9 +11,9 @@
 MASSENTITY_API DECLARE_LOG_CATEGORY_EXTERN(LogPipe, Warning, All);
 
 class UMassEntitySubsystem;
-class UPipeProcessor;
-class UPipeSchematic;
-class UPipeCompositeProcessor;
+class UMassProcessor;
+class UMassSchematic;
+class UMassCompositeProcessor;
 
 UENUM(meta = (Bitflags, UseEnumValuesAsMaskValuesInEditor = "true"))
 enum class EProcessorExecutionFlags : uint8
@@ -33,7 +33,7 @@ struct FProcessorAuxDataBase
 };
 
 USTRUCT(BlueprintType)
-struct MASSENTITY_API FPipeContext
+struct MASSENTITY_API FMassProcessingContext
 {
 	GENERATED_BODY()
 
@@ -46,66 +46,66 @@ struct MASSENTITY_API FPipeContext
 	UPROPERTY()
 	FInstancedStruct AuxData;
 
-	FPipeContext() = default;
-	FPipeContext(UMassEntitySubsystem& InEntities, const float InDeltaSeconds);
+	FMassProcessingContext() = default;
+	FMassProcessingContext(UMassEntitySubsystem& InEntities, const float InDeltaSeconds);
 };
 
 /** 
- *  A helper type that converts a set of UPipeSchematics into a runtime-usable array of PipeProcessor copies
+ *  A helper type that converts a set of UMassSchematics into a runtime-usable array of MassProcessor copies
  */
 USTRUCT()
-struct MASSENTITY_API FRuntimePipeline
+struct MASSENTITY_API FMassRuntimePipeline
 {
 	GENERATED_BODY()
 
 	UPROPERTY()
-	TArray<UPipeProcessor*> Processors;
+	TArray<UMassProcessor*> Processors;
 
 	void Reset();
 	void Initialize(UObject& Owner);
 	
-	/** Creates runtime copies of the given UPipeProcessors collection. */
-	void SetProcessors(TArray<UPipeProcessor*>&& InProcessors);
+	/** Creates runtime copies of the given UMassProcessors collection. */
+	void SetProcessors(TArray<UMassProcessor*>&& InProcessors);
 
-	/** Creates runtime copies of UPipeProcessors declared in Schematics using InOwner as new UPipeProcessors' outer. */
-	void InitializeFromSchematics(TConstArrayView<TSoftObjectPtr<UPipeSchematic>> Schematics, UObject& InOwner);
+	/** Creates runtime copies of UMassProcessors declared in Schematics using InOwner as new UMassProcessors' outer. */
+	void InitializeFromSchematics(TConstArrayView<TSoftObjectPtr<UMassSchematic>> Schematics, UObject& InOwner);
 
-	/** Creates runtime copies of UPipeProcessors given in InProcessors input parameter, using InOwner as new UPipeProcessors' outer. */
-	void CreateFromArray(TConstArrayView<const UPipeProcessor*> InProcessors, UObject& InOwner);
+	/** Creates runtime copies of UMassProcessors given in InProcessors input parameter, using InOwner as new UMassProcessors' outer. */
+	void CreateFromArray(TConstArrayView<const UMassProcessor*> InProcessors, UObject& InOwner);
 
 	/** Calls CreateFromArray and calls Initialize on all processors afterwards. */
-	void InitializeFromArray(TConstArrayView<const UPipeProcessor*> InProcessors, UObject& InOwner);
+	void InitializeFromArray(TConstArrayView<const UMassProcessor*> InProcessors, UObject& InOwner);
 	
-	/** Creates runtime instances of UPipeProcessors for each processor class given via InProcessorClasses. 
+	/** Creates runtime instances of UMassProcessors for each processor class given via InProcessorClasses. 
 	 *  The instances will be created with InOwner as outer. */
-	void InitializeFromClassArray(TConstArrayView<TSubclassOf<UPipeProcessor>> InProcessorClasses, UObject& InOwner);
+	void InitializeFromClassArray(TConstArrayView<TSubclassOf<UMassProcessor>> InProcessorClasses, UObject& InOwner);
 
 	/** Creates a runtime instance of every processors in the given InProcessors array. If a processor of that class
 	 *  already exists in Processors array it gets overridden. Otherwise it gets added to the end of the collection.*/
-	void AppendOrOverrideRuntimeProcessorCopies(TConstArrayView<const UPipeProcessor*> InProcessors, UObject& InOwner);
+	void AppendOrOverrideRuntimeProcessorCopies(TConstArrayView<const UMassProcessor*> InProcessors, UObject& InOwner);
 
 	/** Creates a runtime instance of every processors in the given array if there's no processor of that class in Processors already.
-	 *  Call this function when adding processors to an already configured FRuntimePipeline instance. If you're creating 
+	 *  Call this function when adding processors to an already configured FMassRuntimePipeline instance. If you're creating 
 	 *  one from scratch calling any of the InitializeFrom* methods will be more efficient (and will produce same results)
 	 *  or call AppendOrOverrideRuntimeProcessorCopies.*/
-	void AppendUniqueRuntimeProcessorCopies(TConstArrayView<const UPipeProcessor*> InProcessors, UObject& InOwner);
+	void AppendUniqueRuntimeProcessorCopies(TConstArrayView<const UMassProcessor*> InProcessors, UObject& InOwner);
 
 	/** Adds InProcessor to Processors without any additional checks */
-	void AppendProcessor(UPipeProcessor& Processor);
+	void AppendProcessor(UMassProcessor& Processor);
 
-	/** goes through Processor looking for a UPipeCompositeProcessor instance which GroupName matches the one given as the parameter */
-	UPipeCompositeProcessor* FindTopLevelGroupByName(const FName GroupName);
+	/** goes through Processor looking for a UMassCompositeProcessor instance which GroupName matches the one given as the parameter */
+	UMassCompositeProcessor* FindTopLevelGroupByName(const FName GroupName);
 
 	void DebugOutputDescription(FOutputDevice& Ar) const;
 
-	bool HasProcessorOfExactClass(TSubclassOf<UPipeProcessor> InClass) const;
+	bool HasProcessorOfExactClass(TSubclassOf<UMassProcessor> InClass) const;
 	bool IsEmpty() const { return Processors.IsEmpty();}
 
-	MASSENTITY_API friend uint32 GetTypeHash(const FRuntimePipeline& Instance);
+	MASSENTITY_API friend uint32 GetTypeHash(const FMassRuntimePipeline& Instance);
 };
 
 UENUM()
-enum class EPipeProcessingPhase : uint8
+enum class EMassProcessingPhase : uint8
 {
 	PrePhysics,
 	StartPhysics,
diff --git a/Engine/Plugins/Runtime/MassEntity/Source/MassEntity/Public/MassExecutor.h b/Engine/Plugins/Runtime/MassEntity/Source/MassEntity/Public/MassExecutor.h
--- a/Engine/Plugins/Runtime/MassEntity/Source/MassEntity/Public/MassExecutor.h
+++ b/Engine/Plugins/Runtime/MassEntity/Source/MassEntity/Public/MassExecutor.h
@@ -7,36 +7,36 @@
 #include "Async/TaskGraphInterfaces.h"
 
 
-struct FRuntimePipeline;
-struct FPipeContext;
-struct FLWEntity;
+struct FMassRuntimePipeline;
+struct FMassProcessingContext;
+struct FMassEntityHandle;
 struct FArchetypeChunkCollection;
-class UPipeProcessor;
+class UMassProcessor;
 
 
 namespace UE::Pipe::Executor
 {
 	/** Executes processors in a given RuntimePipeline */
-	MASSENTITY_API void Run(FRuntimePipeline& RuntimePipeline, FPipeContext& PipeContext);
+	MASSENTITY_API void Run(FMassRuntimePipeline& RuntimePipeline, FMassProcessingContext& PipeContext);
 
-	/** Executes given Processor. Used mainly for triggering calculations via PipeCompositeProcessors, e.g processing phases */
-	MASSENTITY_API void Run(UPipeProcessor& Processor, FPipeContext& PipeContext);
+	/** Executes given Processor. Used mainly for triggering calculations via MassCompositeProcessors, e.g processing phases */
+	MASSENTITY_API void Run(UMassProcessor& Processor, FMassProcessingContext& PipeContext);
 
 	/** Similar to the Run function, but instead of using all the entities hosted by PipeContext.EntitySubsystem 
 	 *  it is processing only the entities given by EntityID via the Entities input parameter. 
 	 *  Note that all the entities need to be of Archetype archetype. 
 	 *  Under the hood the function converts Archetype-Entities pair to FArchetypeChunkCollection and calls the other flavor of RunSparse
 	 */
-	MASSENTITY_API void RunSparse(FRuntimePipeline& RuntimePipeline, FPipeContext& PipeContext, FArchetypeHandle Archetype, TConstArrayView<FLWEntity> Entities);
+	MASSENTITY_API void RunSparse(FMassRuntimePipeline& RuntimePipeline, FMassProcessingContext& PipeContext, FArchetypeHandle Archetype, TConstArrayView<FMassEntityHandle> Entities);
 
 	/** Similar to the Run function, but instead of using all the entities hosted by PipeContext.EntitySubsystem 
 	 *  it is processing only the entities given by SparseEntities input parameter.
 	 *  @todo rename
 	 */
-	MASSENTITY_API void RunSparse(FRuntimePipeline& RuntimePipeline, FPipeContext& PipeContext, const FArchetypeChunkCollection& ChunkCollection);
+	MASSENTITY_API void RunSparse(FMassRuntimePipeline& RuntimePipeline, FMassProcessingContext& PipeContext, const FArchetypeChunkCollection& ChunkCollection);
 
 	/** Executes given Processors array view. This function gets called under the hood by the rest of Run* functions */
-	MASSENTITY_API void RunProcessorsView(TArrayView<UPipeProcessor*> Processors, FPipeContext& PipeContext, const FArchetypeChunkCollection* ChunkCollection = nullptr);
+	MASSENTITY_API void RunProcessorsView(TArrayView<UMassProcessor*> Processors, FMassProcessingContext& PipeContext, const FArchetypeChunkCollection* ChunkCollection = nullptr);
 
 	/** 
 	 *  Triggers tasks executing Processor (and potentially it's children) and returns the task graph event representing 
@@ -44,5 +44,5 @@ namespace UE::Pipe::Executor
 	 *  @param OnDoneNotification will be called after all the processors are done, just after flushing the command buffer.
 	 *    Note that OnDoneNotification will be executed on GameThread.
 	 */
-	MASSENTITY_API FGraphEventRef TriggerParallelTasks(UPipeProcessor& Processor, FPipeContext& PipeContext, TFunction<void()> OnDoneNotification);
+	MASSENTITY_API FGraphEventRef TriggerParallelTasks(UMassProcessor& Processor, FMassProcessingContext& PipeContext, TFunction<void()> OnDoneNotification);
 };
diff --git a/Engine/Plugins/Runtime/MassEntity/Source/MassEntity/Public/MassProcessingPhase.h b/Engine/Plugins/Runtime/MassEntity/Source/MassEntity/Public/MassProcessingPhase.h
--- a/Engine/Plugins/Runtime/MassEntity/Source/MassEntity/Public/MassProcessingPhase.h
+++ b/Engine/Plugins/Runtime/MassEntity/Source/MassEntity/Public/MassProcessingPhase.h
@@ -10,20 +10,20 @@
 #include "MassProcessingPhase.generated.h"
 
 
-class UPipeProcessingPhaseManager;
-class UPipeProcessor;
-class UPipeCompositeProcessor;
+class UMassProcessingPhaseManager;
+class UMassProcessor;
+class UMassCompositeProcessor;
 class UMassEntitySubsystem;
-struct FLWCCommandBuffer;
+struct FMassCommandBuffer;
 
 USTRUCT()
-struct FPipeProcessingPhase : public FTickFunction
+struct FMassProcessingPhase : public FTickFunction
 {
 	GENERATED_BODY()
 
 	DECLARE_MULTICAST_DELEGATE_OneParam(FOnPhaseEvent, const float /*DeltaSeconds*/);
 
-	FPipeProcessingPhase();
+	FMassProcessingPhase();
 
 protected:
 	// FTickFunction interface
@@ -42,15 +42,15 @@ public:
 	bool IsDuringPipeProcessing() const { return bIsDuringPipeProcessing; }
 
 protected:
-	friend UPipeProcessingPhaseManager;
+	friend UMassProcessingPhaseManager;
 
 	UPROPERTY(EditAnywhere, Category=Pipe)
-	UPipeCompositeProcessor* PhaseProcessor = nullptr;
+	UMassCompositeProcessor* PhaseProcessor = nullptr;
 
 	UPROPERTY()
-	UPipeProcessingPhaseManager* Manager = nullptr;
+	UMassProcessingPhaseManager* Manager = nullptr;
 
-	EPipeProcessingPhase Phase = EPipeProcessingPhase::MAX;
+	EMassProcessingPhase Phase = EMassProcessingPhase::MAX;
 	FOnPhaseEvent OnPhaseStart;
 	FOnPhaseEvent OnPhaseEnd;
 
@@ -61,7 +61,7 @@ private:
 
 // It is unsafe to copy FTickFunctions and any subclasses of FTickFunction should specify the type trait WithCopy = false
 template<>
-struct TStructOpsTypeTraits<FPipeProcessingPhase> : public TStructOpsTypeTraitsBase2<FPipeProcessingPhase>
+struct TStructOpsTypeTraits<FMassProcessingPhase> : public TStructOpsTypeTraitsBase2<FMassProcessingPhase>
 {
 	enum
 	{
@@ -69,26 +69,26 @@ struct TStructOpsTypeTraits<FPipeProcessingPhase> : public TStructOpsTypeTraitsB
 	};
 };
 
-/** PipeProcessingPhaseManager owns separate FPipeProcessingPhase instances for every ETickingGroup. When activated
- *  via Start function it registers and enables the FPipeProcessingPhase instances which themselves are tick functions 
- *  that host FRuntimePipeline which they trigger as part of their Tick function. 
- *  PipeProcessingPhaseManager serves as an interface to said FPipeProcessingPhase instances and allows initialization
- *  with PipeSchematics (via InitializePhases function) as well as registering arbitrary functions to be called 
+/** PipeProcessingPhaseManager owns separate FMassProcessingPhase instances for every ETickingGroup. When activated
+ *  via Start function it registers and enables the FMassProcessingPhase instances which themselves are tick functions 
+ *  that host FMassRuntimePipeline which they trigger as part of their Tick function. 
+ *  PipeProcessingPhaseManager serves as an interface to said FMassProcessingPhase instances and allows initialization
+ *  with MassSchematics (via InitializePhases function) as well as registering arbitrary functions to be called 
  *  when a particular phase starts of ends (via GetOnPhaseStart and GetOnPhaseEnd functions). */
 UCLASS(Transient, HideCategories = (Tick))
-class MASSENTITY_API UPipeProcessingPhaseManager : public UObject
+class MASSENTITY_API UMassProcessingPhaseManager : public UObject
 {
 	GENERATED_BODY()
 public:
 	UMassEntitySubsystem& GetEntitySubsystemRef() { check(EntitySubsystem); return *EntitySubsystem; }
 
 	/** Retrieves OnPhaseStart multicast delegate's reference for a given Phase */
-	FPipeProcessingPhase::FOnPhaseEvent& GetOnPhaseStart(const EPipeProcessingPhase Phase) { return ProcessingPhases[uint8(Phase)].OnPhaseStart; }
+	FMassProcessingPhase::FOnPhaseEvent& GetOnPhaseStart(const EMassProcessingPhase Phase) { return ProcessingPhases[uint8(Phase)].OnPhaseStart; }
 	/** Retrieves OnPhaseEnd multicast delegate's reference for a given Phase */
-	FPipeProcessingPhase::FOnPhaseEvent& GetOnPhaseEnd(const EPipeProcessingPhase Phase) { return ProcessingPhases[uint8(Phase)].OnPhaseEnd; }
+	FMassProcessingPhase::FOnPhaseEvent& GetOnPhaseEnd(const EMassProcessingPhase Phase) { return ProcessingPhases[uint8(Phase)].OnPhaseEnd; }
 
 	/** 
-	 *  Populates hosted FPipeProcessingPhase instances with Processors read from PipeSettings configuration.
+	 *  Populates hosted FMassProcessingPhase instances with Processors read from PipeSettings configuration.
 	 *  Calling this function overrides previous configuration of Phases.
 	 */
 	void InitializePhases(UObject& InProcessorOwner);
@@ -108,7 +108,7 @@ public:
 	 *  Note that the function will return false while the OnPhaseStart or OnPhaseEnd are being broadcast,
 	 *  the value returned will be `true` only when the entity subsystem is actively engaged 
 	 */
-	bool IsDuringPipeProcessing() const { return CurrentPhase != EPipeProcessingPhase::MAX && ProcessingPhases[int(CurrentPhase)].IsDuringPipeProcessing(); }
+	bool IsDuringPipeProcessing() const { return CurrentPhase != EMassProcessingPhase::MAX && ProcessingPhases[int(CurrentPhase)].IsDuringPipeProcessing(); }
 
 protected:
 	virtual void PostInitProperties() override;
@@ -122,28 +122,28 @@ protected:
 	virtual void OnPipeSettingsChange(const FPropertyChangedEvent& PropertyChangedEvent);
 #endif // WITH_EDITOR
 
-	friend FPipeProcessingPhase;
+	friend FMassProcessingPhase;
 
 	/** 
-	 *  Called by the given Phase at the very start of its execution function (the FPipeProcessingPhase::ExecuteTick),
-	 *  even before the FPipeProcessingPhase.OnPhaseStart broadcast delegate
+	 *  Called by the given Phase at the very start of its execution function (the FMassProcessingPhase::ExecuteTick),
+	 *  even before the FMassProcessingPhase.OnPhaseStart broadcast delegate
 	 */
-	void OnPhaseStart(const FPipeProcessingPhase& Phase);
+	void OnPhaseStart(const FMassProcessingPhase& Phase);
 
 	/**
-	 *  Called by the given Phase at the very end of its execution function (the FPipeProcessingPhase::ExecuteTick),
-	 *  after the FPipeProcessingPhase.OnPhaseEnd broadcast delegate
+	 *  Called by the given Phase at the very end of its execution function (the FMassProcessingPhase::ExecuteTick),
+	 *  after the FMassProcessingPhase.OnPhaseEnd broadcast delegate
 	 */
-	void OnPhaseEnd(FPipeProcessingPhase& Phase);
+	void OnPhaseEnd(FMassProcessingPhase& Phase);
 
 protected:	
 	UPROPERTY(VisibleAnywhere, Category=Pipe)
-	FPipeProcessingPhase ProcessingPhases[(uint8)EPipeProcessingPhase::MAX];
+	FMassProcessingPhase ProcessingPhases[(uint8)EMassProcessingPhase::MAX];
 
 	UPROPERTY()
 	UMassEntitySubsystem* EntitySubsystem;
 
-	EPipeProcessingPhase CurrentPhase = EPipeProcessingPhase::MAX;
+	EMassProcessingPhase CurrentPhase = EMassProcessingPhase::MAX;
 
 #if WITH_EDITOR
 	FDelegateHandle PipeSettingsChangeHandle;
