AggregateTicking and Pipe type renames:

diff --git a/Engine/Plugins/Runtime/MassEntity/Source/MassEntity/Public/MassProcessor.h b/Engine/Plugins/Runtime/MassEntity/Source/MassEntity/Public/MassProcessor.h
--- a/Engine/Plugins/Runtime/MassEntity/Source/MassEntity/Public/MassProcessor.h
+++ b/Engine/Plugins/Runtime/MassEntity/Source/MassEntity/Public/MassProcessor.h
@@ -11,7 +11,7 @@
 #include "MassProcessor.generated.h"
 
 
-struct FPipeProcessingPhaseConfig;
+struct FMassProcessingPhaseConfig;
 
 enum class EProcessorCompletionStatus : uint8
 {
@@ -22,7 +22,7 @@ enum class EProcessorCompletionStatus : uint8
 };
 
 USTRUCT()
-struct FPipeProcessorExecutionOrder
+struct FMassProcessorExecutionOrder
 {
 	GENERATED_BODY()
 
@@ -39,20 +39,20 @@ struct FPipeProcessorExecutionOrder
 
 
 UCLASS(abstract, EditInlineNew, CollapseCategories, config = Game, defaultconfig)
-class MASSENTITY_API UPipeProcessor : public UObject
+class MASSENTITY_API UMassProcessor : public UObject
 {
 	GENERATED_BODY()
 public:
-	UPipeProcessor(const FObjectInitializer& ObjectInitializer = FObjectInitializer::Get());
+	UMassProcessor(const FObjectInitializer& ObjectInitializer = FObjectInitializer::Get());
 
 	virtual void Initialize(UObject& Owner) {}
-	virtual FGraphEventRef DispatchProcessorTasks(UMassEntitySubsystem& EntitySubsystem, FLWComponentSystemExecutionContext& ExecutionContext, const FGraphEventArray& Prerequisites = FGraphEventArray());
+	virtual FGraphEventRef DispatchProcessorTasks(UMassEntitySubsystem& EntitySubsystem, FMassExecutionContext& ExecutionContext, const FGraphEventArray& Prerequisites = FGraphEventArray());
 
 	EProcessorExecutionFlags GetExecutionFlags() const { return (EProcessorExecutionFlags)ExecutionFlags; }
 
 	/** Whether this processor should execute according the CurrentExecutionFlags parameters */
 	bool ShouldExecute(const EProcessorExecutionFlags CurrentExecutionFlags) const { return (GetExecutionFlags() & CurrentExecutionFlags) != EProcessorExecutionFlags::None; }
-	void CallExecute(UMassEntitySubsystem& EntitySubsystem, FLWComponentSystemExecutionContext& Context);
+	void CallExecute(UMassEntitySubsystem& EntitySubsystem, FMassExecutionContext& Context);
 	
 	bool AllowDuplicates() const { return bAllowDuplicates; }
 
@@ -62,11 +62,11 @@ public:
 	//----------------------------------------------------------------------//
 	// Ordering functions 
 	//----------------------------------------------------------------------//
-	virtual EPipeProcessingPhase GetProcessingPhase() const { return ProcessingPhase; }
-	virtual void SetProcessingPhase(EPipeProcessingPhase Phase) { ProcessingPhase = Phase; }
+	virtual EMassProcessingPhase GetProcessingPhase() const { return ProcessingPhase; }
+	virtual void SetProcessingPhase(EMassProcessingPhase Phase) { ProcessingPhase = Phase; }
 	bool DoesRequireGameThreadExecution() const { return bRequiresGameThreadExecution; }
 	
-	const FPipeProcessorExecutionOrder& GetExecutionOrder() const { return ExecutionOrder; }
+	const FMassProcessorExecutionOrder& GetExecutionOrder() const { return ExecutionOrder; }
 
 	TConstArrayView<const int32> GetPrerequisiteIndices() const { return DependencyIndices; }
 
@@ -86,9 +86,9 @@ public:
 #endif
 	
 protected:
-	virtual void ConfigureQueries() PURE_VIRTUAL(UPipeProcessor::ConfigureQueries);
+	virtual void ConfigureQueries() PURE_VIRTUAL(UMassProcessor::ConfigureQueries);
 	virtual void PostInitProperties() override;
-	virtual void Execute(UMassEntitySubsystem& EntitySubsystem, FLWComponentSystemExecutionContext& Context) PURE_VIRTUAL(UPipeProcessor::Execute);
+	virtual void Execute(UMassEntitySubsystem& EntitySubsystem, FMassExecutionContext& Context) PURE_VIRTUAL(UMassProcessor::Execute);
 #if WITH_EDITOR
 	virtual void PostEditChangeProperty(FPropertyChangedEvent& PropertyChangedEvent) override;
 #endif // WITH_EDITOR
@@ -100,18 +100,18 @@ protected:
 
 	/** Processing phase this processor will be automatically run as part of. */
 	UPROPERTY(EditDefaultsOnly, Category = Processor, config)
-	EPipeProcessingPhase ProcessingPhase = EPipeProcessingPhase::PrePhysics;
+	EMassProcessingPhase ProcessingPhase = EMassProcessingPhase::PrePhysics;
 
 	/** Configures when this given processor can be executed in relation to other processors and processing groups, within its processing phase. */
 	UPROPERTY(EditDefaultsOnly, Category = Processor, config)
-	FPipeProcessorExecutionOrder ExecutionOrder;
+	FMassProcessorExecutionOrder ExecutionOrder;
 
 	/** Configures whether this processor should be automatically included in the global list of processors executed every tick (see ProcessingPhase and ExecutionOrder). */
 	UPROPERTY(EditDefaultsOnly, Category = Processor, config)
 	bool bAutoRegisterWithProcessingPhases = true;
 
 	/** meant as a class property, make sure to set it in subclass' constructor. Controls whether there can be multiple 
-	 *  instances of a given class in a single FRuntimePipeline */
+	 *  instances of a given class in a single FMassRuntimePipeline */
 	bool bAllowDuplicates = false;
 
 	UPROPERTY(EditDefaultsOnly, Category = Processor, config)
@@ -124,44 +124,44 @@ protected:
 	bool bCanShowUpInSettings = true;
 #endif // WITH_EDITORONLY_DATA
 
-	friend class UPipeCompositeProcessor;
+	friend class UMassCompositeProcessor;
 	TArray<int32> DependencyIndices;
 	TArray<int32> TransientDependencyIndices;
 };
 
 
 UCLASS()
-class MASSENTITY_API UPipeCompositeProcessor : public UPipeProcessor
+class MASSENTITY_API UMassCompositeProcessor : public UMassProcessor
 {
 	GENERATED_BODY()
 
 	struct FDependencyNode
 	{
 		FName Name;
-		UPipeProcessor* Processor = nullptr;
+		UMassProcessor* Processor = nullptr;
 		TArray<int32> Dependencies;
 	};
 
 public:
-	UPipeCompositeProcessor();
+	UMassCompositeProcessor();
 
-	void SetChildProcessors(TArray<UPipeProcessor*>&& InProcessors);
+	void SetChildProcessors(TArray<UMassProcessor*>&& InProcessors);
 
 	virtual void Initialize(UObject& Owner) override;
 	virtual void DebugOutputDescription(FOutputDevice& Ar, int32 Indent = 0) const override;
-	virtual void SetProcessingPhase(EPipeProcessingPhase Phase) override;
+	virtual void SetProcessingPhase(EMassProcessingPhase Phase) override;
 
 	void SetGroupName(FName NewName);
 	FName GetGroupName() const { return GroupName; }
 
-	virtual void CopyAndSort(const FPipeProcessingPhaseConfig& PhaseConfig, const FString& DependencyGraphFileName = FString());
+	virtual void CopyAndSort(const FMassProcessingPhaseConfig& PhaseConfig, const FString& DependencyGraphFileName = FString());
 
 	/** adds SubProcessor to an appropriately named group. If RequestedGroupName == None then SubProcessor
 	 *  will be added directly to ChildPipeline. If not then the indicated group will be searched for in ChildPipeline 
 	 *  and if it's missing it will be created and AddGroupedProcessor will be called recursively */
-	void AddGroupedProcessor(FName RequestedGroupName, UPipeProcessor& SubProcessor);
+	void AddGroupedProcessor(FName RequestedGroupName, UMassProcessor& SubProcessor);
 
-	virtual FGraphEventRef DispatchProcessorTasks(UMassEntitySubsystem& EntitySubsystem, FLWComponentSystemExecutionContext& ExecutionContext, const FGraphEventArray& Prerequisites = FGraphEventArray()) override;
+	virtual FGraphEventRef DispatchProcessorTasks(UMassEntitySubsystem& EntitySubsystem, FMassExecutionContext& ExecutionContext, const FGraphEventArray& Prerequisites = FGraphEventArray()) override;
 
 	bool IsEmpty() const { return ChildPipeline.IsEmpty(); }
 
@@ -169,7 +169,7 @@ public:
 
 protected:
 	virtual void ConfigureQueries() override;
-	virtual void Execute(UMassEntitySubsystem& EntitySubsystem, FLWComponentSystemExecutionContext& Context) override;
+	virtual void Execute(UMassEntitySubsystem& EntitySubsystem, FMassExecutionContext& Context) override;
 
 	/**
 	 *  Called recursively to add processors and composite processors to ChildPipeline based on ProcessorsAndGroups
@@ -182,11 +182,11 @@ protected:
 	 *  @param RequestedGroupName name of the group for which we want to find or create the processor.
 	 *  @param OutRemainingGroupName contains the group name after cutting the high-level group. In the used example it
 	 *    will contain "B.C". This value is then used to recursively create subgroups */
-	UPipeCompositeProcessor* FindOrAddGroupProcessor(FName RequestedGroupName, FString* OutRemainingGroupName = nullptr);
+	UMassCompositeProcessor* FindOrAddGroupProcessor(FName RequestedGroupName, FString* OutRemainingGroupName = nullptr);
 
 protected:
 	UPROPERTY(VisibleAnywhere, Category=Pipe)
-	FRuntimePipeline ChildPipeline;
+	FMassRuntimePipeline ChildPipeline;
 
 	/** Group name that will be used when resolving processor dependencies and grouping */
 	UPROPERTY()
diff --git a/Engine/Plugins/Runtime/MassEntity/Source/MassEntity/Public/MassProcessorDependencySolver.h b/Engine/Plugins/Runtime/MassEntity/Source/MassEntity/Public/MassProcessorDependencySolver.h
--- a/Engine/Plugins/Runtime/MassEntity/Source/MassEntity/Public/MassProcessorDependencySolver.h
+++ b/Engine/Plugins/Runtime/MassEntity/Source/MassEntity/Public/MassProcessorDependencySolver.h
@@ -5,7 +5,7 @@
 #include "MassEntityTypes.h"
 
 
-class UPipeProcessor;
+class UMassProcessor;
 
 enum class EDependencyNodeType : uint8
 {
@@ -20,14 +20,14 @@ struct MASSENTITY_API FProcessorDependencySolver
 private:
 	struct FNode
 	{
-		FNode(const FName InName, UPipeProcessor* InProcessor) 
+		FNode(const FName InName, UMassProcessor* InProcessor) 
 			: Name(InName), Processor(InProcessor)
 		{}
 
 		FName Name = TEXT("");
 		TArray<FNode> SubNodes;
 		TMap<FName, int32> Indices;
-		UPipeProcessor* Processor = nullptr;
+		UMassProcessor* Processor = nullptr;
 		TArray<int32> OriginalDependencies;
 		TArray<int32> TransientDependencies;
 		TArray<FName> ExecuteBefore;
@@ -42,12 +42,12 @@ public:
 	struct FOrderInfo
 	{
 		FName Name = TEXT("");
-		UPipeProcessor* Processor = nullptr;
+		UMassProcessor* Processor = nullptr;
 		EDependencyNodeType NodeType = EDependencyNodeType::Invalid;
 		TArray<FName> Dependencies;
 	};
 
-	FProcessorDependencySolver(TArrayView<UPipeProcessor*> InProcessors, const FName Name, const FString& InDependencyGraphFileName = FString());
+	FProcessorDependencySolver(TArrayView<UMassProcessor*> InProcessors, const FName Name, const FString& InDependencyGraphFileName = FString());
 	void ResolveDependencies(TArray<FOrderInfo>& OutResult, TConstArrayView<const FName> PriorityNodes = TConstArrayView<const FName>());
 
 	static void CreateSubGroupNames(FName InGroupName, TArray<FString>& SubGroupNames);
@@ -74,13 +74,13 @@ protected:
 	
 	static FString NameViewToString(TConstArrayView<FName> View);
 
-	void AddNode(FName GroupName, UPipeProcessor& Processor);
+	void AddNode(FName GroupName, UMassProcessor& Processor);
 	void BuildDependencies(FNode& RootNode);
 	void Solve(FNode& RootNode, TConstArrayView<const FName> PriorityNodes, TArray<FProcessorDependencySolver::FOrderInfo>& OutResult, int LoggingIndent = 0);
 	void LogNode(const FNode& RootNode, const FNode* ParentNode = nullptr, int Indent = 0);
 	void DumpGraph(FArchive& LogFile) const;
 	
-	TArrayView<UPipeProcessor*> Processors;
+	TArrayView<UMassProcessor*> Processors;
 	FNode GroupRootNode;
 	bool bAnyCyclesDetected = false;
 	FString DependencyGraphFileName;
diff --git a/Engine/Plugins/Runtime/MassEntity/Source/MassEntity/Public/MassSchematic.h b/Engine/Plugins/Runtime/MassEntity/Source/MassEntity/Public/MassSchematic.h
--- a/Engine/Plugins/Runtime/MassEntity/Source/MassEntity/Public/MassSchematic.h
+++ b/Engine/Plugins/Runtime/MassEntity/Source/MassEntity/Public/MassSchematic.h
@@ -8,20 +8,20 @@
 #include "MassSchematic.generated.h"
 
 
-class UPipeProcessor;
+class UMassProcessor;
 
 /**
  * Pipe Schematic asset.
  */
 UCLASS(BlueprintType)
-class MASSENTITY_API UPipeSchematic : public UDataAsset
+class MASSENTITY_API UMassSchematic : public UDataAsset
 {
 	GENERATED_BODY()
 
 public:
-	TConstArrayView<const UPipeProcessor*> GetProcessors() const { return MakeArrayView(Processors);  }
+	TConstArrayView<const UMassProcessor*> GetProcessors() const { return MakeArrayView(Processors);  }
 
 protected:
 	UPROPERTY(EditAnywhere, BlueprintReadOnly, Category = Pipe, Instanced)
-	TArray<UPipeProcessor*> Processors;
+	TArray<UMassProcessor*> Processors;
 };
diff --git a/Engine/Plugins/Runtime/MassEntity/Source/MassEntityEditor/Private/AssetTypeActions_MassSchematic.cpp b/Engine/Plugins/Runtime/MassEntity/Source/MassEntityEditor/Private/AssetTypeActions_MassSchematic.cpp
--- a/Engine/Plugins/Runtime/MassEntity/Source/MassEntityEditor/Private/AssetTypeActions_MassSchematic.cpp
+++ b/Engine/Plugins/Runtime/MassEntity/Source/MassEntityEditor/Private/AssetTypeActions_MassSchematic.cpp
@@ -8,13 +8,13 @@
 
 #define LOCTEXT_NAMESPACE "AssetTypeActions"
 
-void FAssetTypeActions_PipeSchematic::OpenAssetEditor(const TArray<UObject*>& InObjects, TSharedPtr<IToolkitHost> EditWithinLevelEditor)
+void FAssetTypeActions_MassSchematic::OpenAssetEditor(const TArray<UObject*>& InObjects, TSharedPtr<IToolkitHost> EditWithinLevelEditor)
 {
 	EToolkitMode::Type Mode = EditWithinLevelEditor.IsValid() ? EToolkitMode::WorldCentric : EToolkitMode::Standalone;
 
 	for (auto ObjIt = InObjects.CreateConstIterator(); ObjIt; ++ObjIt)
 	{
-		if (UPipeSchematic* Pipe = Cast<UPipeSchematic>(*ObjIt))
+		if (UMassSchematic* Pipe = Cast<UMassSchematic>(*ObjIt))
 		{
 			FPipeEditorModule& EditorModule = FModuleManager::LoadModuleChecked<FPipeEditorModule>("PipeEditor");
 			TSharedRef<IPipeEditor> NewEditor = EditorModule.CreatePipeEditor(Mode, EditWithinLevelEditor, Pipe);
@@ -22,12 +22,12 @@ void FAssetTypeActions_PipeSchematic::OpenAssetEditor(const TArray<UObject*>& In
 	}
 }
 
-UClass* FAssetTypeActions_PipeSchematic::GetSupportedClass() const
+UClass* FAssetTypeActions_MassSchematic::GetSupportedClass() const
 {
-	return UPipeSchematic::StaticClass();
+	return UMassSchematic::StaticClass();
 }
 
-uint32 FAssetTypeActions_PipeSchematic::GetCategories()
+uint32 FAssetTypeActions_MassSchematic::GetCategories()
 { 
 	// @todo will probably need to use a different category, like "gameplay"
 	IAIModule& AIModule = FModuleManager::GetModuleChecked<IAIModule>("AIModule").Get();
diff --git a/Engine/Plugins/Runtime/MassEntity/Source/MassEntityEditor/Private/AssetTypeActions_MassSchematic.h b/Engine/Plugins/Runtime/MassEntity/Source/MassEntityEditor/Private/AssetTypeActions_MassSchematic.h
--- a/Engine/Plugins/Runtime/MassEntity/Source/MassEntityEditor/Private/AssetTypeActions_MassSchematic.h
+++ b/Engine/Plugins/Runtime/MassEntity/Source/MassEntityEditor/Private/AssetTypeActions_MassSchematic.h
@@ -6,14 +6,14 @@
 #include "AssetTypeActions_Base.h"
 
 /**
-  * Asset type actions for UPipeSchematic.
+  * Asset type actions for UMassSchematic.
  */
 
-class FAssetTypeActions_PipeSchematic : public FAssetTypeActions_Base
+class FAssetTypeActions_MassSchematic : public FAssetTypeActions_Base
 {
 public:
 	// IAssetTypeActions Implementation
-	virtual FText GetName() const override { return NSLOCTEXT("AssetTypeActions", "FAssetTypeActions_PipeSchematic", "PipeSchematic"); }
+	virtual FText GetName() const override { return NSLOCTEXT("AssetTypeActions", "FAssetTypeActions_MassSchematic", "MassSchematic"); }
 	virtual FColor GetTypeColor() const override { return FColor(201, 185, 29); }
 	virtual UClass* GetSupportedClass() const override;
 	virtual void OpenAssetEditor(const TArray<UObject*>& InObjects, TSharedPtr<class IToolkitHost> EditWithinLevelEditor = TSharedPtr<IToolkitHost>()) override;
diff --git a/Engine/Plugins/Runtime/MassEntity/Source/MassEntityEditor/Private/MassEntityEditor.cpp b/Engine/Plugins/Runtime/MassEntity/Source/MassEntityEditor/Private/MassEntityEditor.cpp
--- a/Engine/Plugins/Runtime/MassEntity/Source/MassEntityEditor/Private/MassEntityEditor.cpp
+++ b/Engine/Plugins/Runtime/MassEntity/Source/MassEntityEditor/Private/MassEntityEditor.cpp
@@ -22,9 +22,9 @@ const FName FPipeEditor::AssetDetailsTabId(TEXT("PipeEditor_AssetDetails"));
 
 void FPipeEditor::AddReferencedObjects(FReferenceCollector& Collector)
 {
-	if (PipeSchematic != nullptr)
+	if (MassSchematic != nullptr)
 	{
-		Collector.AddReferencedObject(PipeSchematic);
+		Collector.AddReferencedObject(MassSchematic);
 	}
 }
 
@@ -47,9 +47,9 @@ void FPipeEditor::UnregisterTabSpawners(const TSharedRef<class FTabManager>& InT
 	InTabManager->UnregisterTabSpawner(AssetDetailsTabId);
 }
 
-void FPipeEditor::InitEditor(const EToolkitMode::Type Mode, const TSharedPtr< class IToolkitHost >& InitToolkitHost, UPipeSchematic& InPipeSchematic)
+void FPipeEditor::InitEditor(const EToolkitMode::Type Mode, const TSharedPtr< class IToolkitHost >& InitToolkitHost, UMassSchematic& InMassSchematic)
 {
-	PipeSchematic = &InPipeSchematic;
+	MassSchematic = &InMassSchematic;
 
 	// use the commented block while iterating on the layout since the named layout gets saved
 	//FGuid Result;
@@ -70,7 +70,7 @@ void FPipeEditor::InitEditor(const EToolkitMode::Type Mode, const TSharedPtr< cl
 
 	const bool bCreateDefaultStandaloneMenu = true;
 	const bool bCreateDefaultToolbar = true;
-	FAssetEditorToolkit::InitAssetEditor(Mode, InitToolkitHost, PipeEditorAppName, StandaloneDefaultLayout, bCreateDefaultStandaloneMenu, bCreateDefaultToolbar, PipeSchematic);
+	FAssetEditorToolkit::InitAssetEditor(Mode, InitToolkitHost, PipeEditorAppName, StandaloneDefaultLayout, bCreateDefaultStandaloneMenu, bCreateDefaultToolbar, MassSchematic);
 }
 
 FName FPipeEditor::GetToolkitFName() const
@@ -102,7 +102,7 @@ TSharedRef<SDockTab> FPipeEditor::SpawnTab_AssetDetails(const FSpawnTabArgs& Arg
 	DetailsViewArgs.NameAreaSettings = FDetailsViewArgs::HideNameArea;
 
 	AssetDetailsView = PropertyEditorModule.CreateDetailView(DetailsViewArgs);
-	AssetDetailsView->SetObject(PipeSchematic);
+	AssetDetailsView->SetObject(MassSchematic);
 	AssetDetailsView->OnFinishedChangingProperties().AddSP(this, &FPipeEditor::OnAssetFinishedChangingProperties);
 
 	TSharedRef<SDockTab> SpawnedTab = SNew(SDockTab)
@@ -125,7 +125,7 @@ void FPipeEditor::SaveAsset_Execute()
 void FPipeEditor::OnAssetFinishedChangingProperties(const FPropertyChangedEvent& PropertyChangedEvent)
 {
 	FPipeEditorModule& PipeEditorModule = FModuleManager::GetModuleChecked<FPipeEditorModule>("PipeEditor");
-	PipeEditorModule.GetOnAssetPropertiesChanged().Broadcast(PipeSchematic, PropertyChangedEvent);
+	PipeEditorModule.GetOnAssetPropertiesChanged().Broadcast(MassSchematic, PropertyChangedEvent);
 }
 
 
diff --git a/Engine/Plugins/Runtime/MassEntity/Source/MassEntityEditor/Private/MassEntityEditor.h b/Engine/Plugins/Runtime/MassEntity/Source/MassEntityEditor/Private/MassEntityEditor.h
--- a/Engine/Plugins/Runtime/MassEntity/Source/MassEntityEditor/Private/MassEntityEditor.h
+++ b/Engine/Plugins/Runtime/MassEntity/Source/MassEntityEditor/Private/MassEntityEditor.h
@@ -9,7 +9,7 @@
 #include "IMassEntityEditor.h"
 
 class IDetailsView;
-class UPipeSchematic;
+class UMassSchematic;
 
 class FPipeEditor : public IPipeEditor, public FGCObject//, public FSelfRegisteringEditorUndoClient
 {
@@ -17,7 +17,7 @@ public:
 	virtual void RegisterTabSpawners(const TSharedRef<class FTabManager>& TabManager) override;
 	virtual void UnregisterTabSpawners(const TSharedRef<class FTabManager>& TabManager) override;
 
-	void InitEditor(const EToolkitMode::Type Mode, const TSharedPtr< class IToolkitHost >& InitToolkitHost, UPipeSchematic& InPipeSchematic);
+	void InitEditor(const EToolkitMode::Type Mode, const TSharedPtr< class IToolkitHost >& InitToolkitHost, UMassSchematic& InMassSchematic);
 
 	//~ Begin IToolkit Interface
 	virtual FName GetToolkitFName() const override;
@@ -47,7 +47,7 @@ private:
 	TSharedRef<SDockTab> SpawnTab_AssetDetails(const FSpawnTabArgs& Args);
 
 	/** Pipe Schematic being edited */
-	UPipeSchematic* PipeSchematic;
+	UMassSchematic* MassSchematic;
 
 	/** Asset Property View */
 	TSharedPtr<class IDetailsView> AssetDetailsView;
diff --git a/Engine/Plugins/Runtime/MassEntity/Source/MassEntityEditor/Private/MassEntityEditorModule.cpp b/Engine/Plugins/Runtime/MassEntity/Source/MassEntityEditor/Private/MassEntityEditorModule.cpp
--- a/Engine/Plugins/Runtime/MassEntity/Source/MassEntityEditor/Private/MassEntityEditorModule.cpp
+++ b/Engine/Plugins/Runtime/MassEntity/Source/MassEntityEditor/Private/MassEntityEditorModule.cpp
@@ -26,7 +26,7 @@ void FPipeEditorModule::StartupModule()
 	// Register asset types
 	IAssetTools& AssetTools = FModuleManager::LoadModuleChecked<FAssetToolsModule>("AssetTools").Get();
 
-	TSharedPtr<FAssetTypeActions_PipeSchematic> PipeAssetTypeAction = MakeShareable(new FAssetTypeActions_PipeSchematic);
+	TSharedPtr<FAssetTypeActions_MassSchematic> PipeAssetTypeAction = MakeShareable(new FAssetTypeActions_MassSchematic);
 	ItemDataAssetTypeActions.Add(PipeAssetTypeAction);
 	AssetTools.RegisterAssetTypeActions(PipeAssetTypeAction.ToSharedRef());
 
@@ -67,18 +67,18 @@ void FPipeEditorModule::ShutdownModule()
 	}
 }
 
-TSharedRef<IPipeEditor> FPipeEditorModule::CreatePipeEditor(const EToolkitMode::Type Mode, const TSharedPtr< class IToolkitHost >& InitToolkitHost, UPipeSchematic* PipeSchematic)
+TSharedRef<IPipeEditor> FPipeEditorModule::CreatePipeEditor(const EToolkitMode::Type Mode, const TSharedPtr< class IToolkitHost >& InitToolkitHost, UMassSchematic* MassSchematic)
 {
 	if (!ProcessorClassCache.IsValid())
 	{
-		ProcessorClassCache = MakeShareable(new FGraphNodeClassHelper(UPipeProcessor::StaticClass()));
+		ProcessorClassCache = MakeShareable(new FGraphNodeClassHelper(UMassProcessor::StaticClass()));
 		ProcessorClassCache->UpdateAvailableBlueprintClasses();
 	}
 
 	TSharedRef<FPipeEditor> NewEditor(new FPipeEditor());
-	if (ensure(PipeSchematic))
+	if (ensure(MassSchematic))
 	{
-		NewEditor->InitEditor(Mode, InitToolkitHost, *PipeSchematic);
+		NewEditor->InitEditor(Mode, InitToolkitHost, *MassSchematic);
 	}
 	return NewEditor;
 }
diff --git a/Engine/Plugins/Runtime/MassEntity/Source/MassEntityEditor/Private/MassSchematicFactory.cpp b/Engine/Plugins/Runtime/MassEntity/Source/MassEntityEditor/Private/MassSchematicFactory.cpp
--- a/Engine/Plugins/Runtime/MassEntity/Source/MassEntityEditor/Private/MassSchematicFactory.cpp
+++ b/Engine/Plugins/Runtime/MassEntity/Source/MassEntityEditor/Private/MassSchematicFactory.cpp
@@ -3,25 +3,25 @@
 #include "MassSchematicFactory.h"
 #include "MassSchematic.h"
 
-#define LOCTEXT_NAMESPACE "PipeSchematicFactory"
+#define LOCTEXT_NAMESPACE "MassSchematicFactory"
 
-UPipeSchematicFactory::UPipeSchematicFactory(const FObjectInitializer& ObjectInitializer)
+UMassSchematicFactory::UMassSchematicFactory(const FObjectInitializer& ObjectInitializer)
 	: Super(ObjectInitializer)
 {
-	SupportedClass = UPipeSchematic::StaticClass();
+	SupportedClass = UMassSchematic::StaticClass();
 	bCreateNew = true;
 	bEditAfterNew = true;
 }
 
-bool UPipeSchematicFactory::CanCreateNew() const
+bool UMassSchematicFactory::CanCreateNew() const
 {
 	return true;
 }
 
-UObject* UPipeSchematicFactory::FactoryCreateNew(UClass* Class,UObject* InParent,FName Name,EObjectFlags Flags,UObject* Context,FFeedbackContext* Warn)
+UObject* UMassSchematicFactory::FactoryCreateNew(UClass* Class,UObject* InParent,FName Name,EObjectFlags Flags,UObject* Context,FFeedbackContext* Warn)
 {
-	check(Class->IsChildOf(UPipeSchematic::StaticClass()));
-	return NewObject<UPipeSchematic>(InParent, Class, Name, Flags);
+	check(Class->IsChildOf(UMassSchematic::StaticClass()));
+	return NewObject<UMassSchematic>(InParent, Class, Name, Flags);
 }
 
 #undef LOCTEXT_NAMESPACE
diff --git a/Engine/Plugins/Runtime/MassEntity/Source/MassEntityEditor/Public/MassEntityEditorModule.h b/Engine/Plugins/Runtime/MassEntity/Source/MassEntityEditor/Public/MassEntityEditorModule.h
--- a/Engine/Plugins/Runtime/MassEntity/Source/MassEntityEditor/Public/MassEntityEditorModule.h
+++ b/Engine/Plugins/Runtime/MassEntity/Source/MassEntityEditor/Public/MassEntityEditorModule.h
@@ -12,7 +12,7 @@
 class IPipeEditor;
 class FAssetTypeActions_Base;
 struct FGraphPanelNodeFactory;
-class UPipeSchematic;
+class UMassSchematic;
 
 /**
 * The public interface to this module
@@ -20,13 +20,13 @@ class UPipeSchematic;
 class MASSENTITYEDITOR_API FPipeEditorModule : public IModuleInterface, public IHasMenuExtensibility, public IHasToolBarExtensibility
 {
 public:
-	DECLARE_MULTICAST_DELEGATE_TwoParams(FOnAssetPropertiesChanged, class UPipeSchematic* /*PipeSchematic*/, const FPropertyChangedEvent& /*PropertyChangedEvent*/);
+	DECLARE_MULTICAST_DELEGATE_TwoParams(FOnAssetPropertiesChanged, class UMassSchematic* /*MassSchematic*/, const FPropertyChangedEvent& /*PropertyChangedEvent*/);
 
 	virtual void StartupModule() override;
 	virtual void ShutdownModule() override;
 
 	/** Creates an instance of Pipe editor. Only virtual so that it can be called across the DLL boundary. */
-	virtual TSharedRef<IPipeEditor> CreatePipeEditor(const EToolkitMode::Type Mode, const TSharedPtr< class IToolkitHost >& InitToolkitHost, UPipeSchematic* PipeSchematic);
+	virtual TSharedRef<IPipeEditor> CreatePipeEditor(const EToolkitMode::Type Mode, const TSharedPtr< class IToolkitHost >& InitToolkitHost, UMassSchematic* MassSchematic);
 
 	virtual TSharedPtr<FExtensibilityManager> GetMenuExtensibilityManager() override { return MenuExtensibilityManager; }
 	virtual TSharedPtr<FExtensibilityManager> GetToolBarExtensibilityManager() override { return ToolBarExtensibilityManager; }
diff --git a/Engine/Plugins/Runtime/MassEntity/Source/MassEntityEditor/Public/MassSchematicFactory.h b/Engine/Plugins/Runtime/MassEntity/Source/MassEntityEditor/Public/MassSchematicFactory.h
--- a/Engine/Plugins/Runtime/MassEntity/Source/MassEntityEditor/Public/MassSchematicFactory.h
+++ b/Engine/Plugins/Runtime/MassEntity/Source/MassEntityEditor/Public/MassSchematicFactory.h
@@ -8,7 +8,7 @@
 #include "MassSchematicFactory.generated.h"
 
 UCLASS()
-class UPipeSchematicFactory : public UFactory
+class UMassSchematicFactory : public UFactory
 {
 	GENERATED_UCLASS_BODY()
 
diff --git a/Engine/Plugins/Runtime/MassEntity/Source/MassEntityTestSuite/Private/LWComponentTestFarmPlot.cpp b/Engine/Plugins/Runtime/MassEntity/Source/MassEntityTestSuite/Private/LWComponentTestFarmPlot.cpp
--- a/Engine/Plugins/Runtime/MassEntity/Source/MassEntityTestSuite/Private/LWComponentTestFarmPlot.cpp
+++ b/Engine/Plugins/Runtime/MassEntity/Source/MassEntityTestSuite/Private/LWComponentTestFarmPlot.cpp
@@ -20,11 +20,11 @@ void UComponentUpdateSystem::PostInitProperties()
 //////////////////////////////////////////////////////////////////////////
 // UFarmHarvestTimerSetIcon
 
-void UFarmHarvestTimerSetIcon::Execute(UMassEntitySubsystem& EntitySubsystem, FLWComponentSystemExecutionContext& Context)
+void UFarmHarvestTimerSetIcon::Execute(UMassEntitySubsystem& EntitySubsystem, FMassExecutionContext& Context)
 {
 	QUICK_SCOPE_CYCLE_COUNTER(SET_ICON_SET_ICON_SET_ICON);
 
-	EntityQuery.ForEachEntityChunk(EntitySubsystem, Context, [this](FLWComponentSystemExecutionContext& Context) {
+	EntityQuery.ForEachEntityChunk(EntitySubsystem, Context, [this](FMassExecutionContext& Context) {
 
 		const int32 NumEntities = Context.GetEntitiesNum();
 		TConstArrayView<FFarmGridCellData> GridCoordList = Context.GetComponentView<FFarmGridCellData>();
@@ -39,7 +39,7 @@ void UFarmHarvestTimerSetIcon::Execute(UMassEntitySubsystem& EntitySubsystem, FL
 
 			VisualList[i].HarvestIconIndex = HarvestIconISMC->AddInstance(IconTransform);
 
-			FLWEntity ThisEntity = Context.GetEntity(i);
+			FMassEntityHandle ThisEntity = Context.GetEntity(i);
 			Context.Defer().RemoveTag<FFarmJustBecameReadyToHarvest>(ThisEntity);
 			Context.Defer().AddTag<FFarmReadyToHarvest>(ThisEntity);
 		}
@@ -66,7 +66,7 @@ ALWComponentTestFarmPlot::ALWComponentTestFarmPlot()
 
 void ALWComponentTestFarmPlot::AddItemToGrid(UMassEntitySubsystem* EntitySystem, int32 X, int32 Y, FArchetypeHandle Archetype, int32 VisualIndex)
 {
-	FLWEntity NewItem = EntitySystem->CreateEntity(Archetype);
+	FMassEntityHandle NewItem = EntitySystem->CreateEntity(Archetype);
 	PlantedSquares[X + Y * GridWidth] = NewItem;
 
 	EntitySystem->GetComponentDataChecked<FFarmWaterComponent>(NewItem).DeltaWaterPerSecond = FMath::FRandRange(-0.01f, -0.001f);
@@ -154,8 +154,8 @@ void ALWComponentTestFarmPlot::TickActor(float DeltaTime, enum ELevelTick TickTy
 
 	// Run every frame systems
 	{
-		FLWComponentSystemExecutionContext ExecContext(DeltaTime);
-		TSharedPtr<FLWCCommandBuffer> DeferredCommandBuffer = MakeShareable(new FLWCCommandBuffer());
+		FMassExecutionContext ExecContext(DeltaTime);
+		TSharedPtr<FMassCommandBuffer> DeferredCommandBuffer = MakeShareable(new FMassCommandBuffer());
 		ExecContext.SetDeferredCommandBuffer(DeferredCommandBuffer);
 
 		for (UComponentUpdateSystem* System : PerFrameSystems)
@@ -169,8 +169,8 @@ void ALWComponentTestFarmPlot::TickActor(float DeltaTime, enum ELevelTick TickTy
 	NextSecondTimer -= DeltaTime;
 	while (NextSecondTimer < 0.0f)
 	{
-		FLWComponentSystemExecutionContext ExecContext(1.f);
-		TSharedPtr<FLWCCommandBuffer> DeferredCommandBuffer = MakeShareable(new FLWCCommandBuffer());
+		FMassExecutionContext ExecContext(1.f);
+		TSharedPtr<FMassCommandBuffer> DeferredCommandBuffer = MakeShareable(new FMassCommandBuffer());
 		ExecContext.SetDeferredCommandBuffer(DeferredCommandBuffer);
 
 		NextSecondTimer += 1.0f;
@@ -190,7 +190,7 @@ void ALWComponentTestFarmPlot::TickActor(float DeltaTime, enum ELevelTick TickTy
 		{
 			for (int32 X = 0; X < GridWidth; ++X)
 			{
-				FLWEntity GridEntity = PlantedSquares[X + Y * GridWidth];
+				FMassEntityHandle GridEntity = PlantedSquares[X + Y * GridWidth];
 				if (EntitySystem->IsValidEntity(GridEntity))
 				{
 					FFarmVisualComponent& VisualComp = EntitySystem->GetComponentDataChecked<FFarmVisualComponent>(GridEntity);
@@ -242,7 +242,7 @@ void ALWComponentTestFarmPlot::TickActor(float DeltaTime, enum ELevelTick TickTy
 		{
 			for (int32 X = 0; X < GridWidth; ++X)
 			{
-				FLWEntity GridEntity = PlantedSquares[X + Y * GridWidth];
+				FMassEntityHandle GridEntity = PlantedSquares[X + Y * GridWidth];
 				if (EntitySystem->IsValidEntity(GridEntity))
 				{
 					FHarvestTimerComponent& HarvestTimerData = EntitySystem->GetComponentDataChecked<FHarvestTimerComponent>(GridEntity);
diff --git a/Engine/Plugins/Runtime/MassEntity/Source/MassEntityTestSuite/Private/MassCompositeProcessorTest.cpp b/Engine/Plugins/Runtime/MassEntity/Source/MassEntityTestSuite/Private/MassCompositeProcessorTest.cpp
--- a/Engine/Plugins/Runtime/MassEntity/Source/MassEntityTestSuite/Private/MassCompositeProcessorTest.cpp
+++ b/Engine/Plugins/Runtime/MassEntity/Source/MassEntityTestSuite/Private/MassCompositeProcessorTest.cpp
@@ -22,15 +22,15 @@ struct FCompositeProcessorTest_Empty : FEntityTestBase
 	{
 		CA_ASSUME(EntitySubsystem);
 
-		UPipeCompositeProcessor* CompositeProcessor = NewObject<UPipeCompositeProcessor>(EntitySubsystem);
+		UMassCompositeProcessor* CompositeProcessor = NewObject<UMassCompositeProcessor>(EntitySubsystem);
 		check(CompositeProcessor);
-		FPipeContext PipeContext(*EntitySubsystem, /*DeltaSeconds=*/0.f);
+		FMassProcessingContext ProcessingContext(*EntitySubsystem, /*DeltaSeconds=*/0.f);
 		// it should just run, no warnings
-		UE::Pipe::Executor::Run(*CompositeProcessor, PipeContext);
+		UE::Pipe::Executor::Run(*CompositeProcessor, ProcessingContext);
 		return true;
 	}
 };
-IMPLEMENT_AI_INSTANT_TEST(FCompositeProcessorTest_Empty, "System.Pipe.Processor.Composite.Empty");
+IMPLEMENT_AI_INSTANT_TEST(FCompositeProcessorTest_Empty, "System.Mass.Processor.Composite.Empty");
 
 
 struct FCompositeProcessorTest_MultipleSubProcessors : FEntityTestBase
@@ -39,16 +39,16 @@ struct FCompositeProcessorTest_MultipleSubProcessors : FEntityTestBase
 	{
 		CA_ASSUME(EntitySubsystem);
 
-		UPipeCompositeProcessor* CompositeProcessor = NewObject<UPipeCompositeProcessor>(EntitySubsystem);
+		UMassCompositeProcessor* CompositeProcessor = NewObject<UMassCompositeProcessor>(EntitySubsystem);
 		check(CompositeProcessor);
 
 		int Result = 0;
 		{
-			TArray<UPipeProcessor*> Processors;
+			TArray<UMassProcessor*> Processors;
 			for (int i = 0; i < 3; ++i)
 			{
 				UPipeTestProcessorBase* Processor = NewObject<UPipeTestProcessorBase>(EntitySubsystem);
-				Processor->ExecutionFunction = [Processor, &Result, i](UMassEntitySubsystem& InEntitySubsystem, FLWComponentSystemExecutionContext& Context) {
+				Processor->ExecutionFunction = [Processor, &Result, i](UMassEntitySubsystem& InEntitySubsystem, FMassExecutionContext& Context) {
 						check(Processor);
 						Result += FMath::Pow(10.f, float(i));
 					};
@@ -58,13 +58,13 @@ struct FCompositeProcessorTest_MultipleSubProcessors : FEntityTestBase
 			CompositeProcessor->SetChildProcessors(MoveTemp(Processors));
 		}
 
-		FPipeContext PipeContext(*EntitySubsystem, /*DeltaSeconds=*/0.f);
-		UE::Pipe::Executor::Run(*CompositeProcessor, PipeContext);
+		FMassProcessingContext ProcessingContext(*EntitySubsystem, /*DeltaSeconds=*/0.f);
+		UE::Pipe::Executor::Run(*CompositeProcessor, ProcessingContext);
 		AITEST_EQUAL("All of the child processors should get run", Result, 111);
 		return true;
 	}
 };
-IMPLEMENT_AI_INSTANT_TEST(FCompositeProcessorTest_MultipleSubProcessors, "System.Pipe.Processor.Composite.MultipleSubProcessors");
+IMPLEMENT_AI_INSTANT_TEST(FCompositeProcessorTest_MultipleSubProcessors, "System.Mass.Processor.Composite.MultipleSubProcessors");
 
 } // FPipeCompositeProcessorTest
 
diff --git a/Engine/Plugins/Runtime/MassEntity/Source/MassEntityTestSuite/Private/MassDependencyTest.cpp b/Engine/Plugins/Runtime/MassEntity/Source/MassEntityTestSuite/Private/MassDependencyTest.cpp
--- a/Engine/Plugins/Runtime/MassEntity/Source/MassEntityTestSuite/Private/MassDependencyTest.cpp
+++ b/Engine/Plugins/Runtime/MassEntity/Source/MassEntityTestSuite/Private/MassDependencyTest.cpp
@@ -33,7 +33,7 @@ struct FDependencySolverBase : FAITestBase
 	void Solve()
 	{
 		Result.Reset();
-		FProcessorDependencySolver Solver(MakeArrayView((UPipeProcessor**)Processors.GetData(), Processors.Num()), TEXT("Test") );
+		FProcessorDependencySolver Solver(MakeArrayView((UMassProcessor**)Processors.GetData(), Processors.Num()), TEXT("Test") );
 		Solver.ResolveDependencies(Result);
 	}
 };	
diff --git a/Engine/Plugins/Runtime/MassEntity/Source/MassEntityTestSuite/Private/MassEntityTest.cpp b/Engine/Plugins/Runtime/MassEntity/Source/MassEntityTestSuite/Private/MassEntityTest.cpp
--- a/Engine/Plugins/Runtime/MassEntity/Source/MassEntityTestSuite/Private/MassEntityTest.cpp
+++ b/Engine/Plugins/Runtime/MassEntity/Source/MassEntityTestSuite/Private/MassEntityTest.cpp
@@ -96,7 +96,7 @@ struct FEntityTest_EntityBatchCreation : FEntityTestBase
 	{
 		CA_ASSUME(EntitySubsystem);
 		const int32 Count = 123;
-		TArray<FLWEntity> Entities;
+		TArray<FMassEntityHandle> Entities;
 		EntitySubsystem->BatchCreateEntities(FloatsIntsArchetype, Count, Entities);
 		AITEST_EQUAL("Batch creation should create the expected number of entities", Entities.Num(), Count);
 		AITEST_EQUAL("The total number of entities present must match the number requested", EntitySubsystem->DebugGetEntityCount(), Count);
@@ -111,7 +111,7 @@ struct FEntityTest_BatchCreatingSingleEntity : FEntityTestBase
 	virtual bool InstantTest() override
 	{
 		CA_ASSUME(EntitySubsystem);
-		TArray<FLWEntity> Entities;
+		TArray<FMassEntityHandle> Entities;
 		EntitySubsystem->BatchCreateEntities(FloatsIntsArchetype, /*Count=*/1, Entities);
 		AITEST_EQUAL("Batch creation should have created a single entity", Entities.Num(), 1);
 		AITEST_EQUAL("The total number of entities present must match the number created by batch creation", EntitySubsystem->DebugGetEntityCount(), 1);
@@ -127,7 +127,7 @@ struct FEntityTest_EntityCreation : FEntityTestBase
 	virtual bool InstantTest() override
 	{
 		CA_ASSUME(EntitySubsystem);		
-		const FLWEntity Entity = EntitySubsystem->CreateEntity(FloatsArchetype);
+		const FMassEntityHandle Entity = EntitySubsystem->CreateEntity(FloatsArchetype);
 		AITEST_EQUAL("There should be one entity across the whole system", EntitySubsystem->DebugGetEntityCount(), 1);
 		AITEST_EQUAL("Entity\'s archetype should be the Float one", EntitySubsystem->GetArchetypeForEntity(Entity), FloatsArchetype);
 		AITEST_EQUAL("The created entity should have been added to the Floats archetype", EntitySubsystem->DebugGetArchetypeEntitiesCount(FloatsArchetype), 1);
@@ -144,7 +144,7 @@ struct FEntityTest_EntityCreationFromInstances : FEntityTestBase
 	virtual bool InstantTest() override
 	{
 		CA_ASSUME(EntitySubsystem);
-		const FLWEntity Entity = EntitySubsystem->CreateEntity(MakeArrayView(&InstanceInt, 1));
+		const FMassEntityHandle Entity = EntitySubsystem->CreateEntity(MakeArrayView(&InstanceInt, 1));
 		AITEST_EQUAL("There should be one entity across the whole system", EntitySubsystem->DebugGetEntityCount(), 1);
 		AITEST_EQUAL("Entity\'s archetype should be the Ints one", EntitySubsystem->GetArchetypeForEntity(Entity), IntsArchetype);
 		AITEST_EQUAL("The created entity should have been added to the Ints archetype", EntitySubsystem->DebugGetArchetypeEntitiesCount(IntsArchetype), 1);
@@ -177,7 +177,7 @@ struct FEntityTest_AddingFragmentType : FEntityTestBase
 	virtual bool InstantTest() override
 	{
 		CA_ASSUME(EntitySubsystem);
-		const FLWEntity Entity = EntitySubsystem->CreateEntity(FloatsArchetype);
+		const FMassEntityHandle Entity = EntitySubsystem->CreateEntity(FloatsArchetype);
 		EntitySubsystem->AddComponentToEntity(Entity, FTestFragment_Int::StaticStruct());
 		AITEST_EQUAL("There should be one entity across the whole system", EntitySubsystem->DebugGetEntityCount(), 1);
 		AITEST_EQUAL("The original archetype should now have no entities", EntitySubsystem->DebugGetArchetypeEntitiesCount(FloatsArchetype), 0);
@@ -196,7 +196,7 @@ struct FEntityTest_AddingFragmentInstance : FEntityTestBase
 	virtual bool InstantTest() override
 	{
 		CA_ASSUME(EntitySubsystem);
-		const FLWEntity Entity = EntitySubsystem->CreateEntity(FloatsArchetype);
+		const FMassEntityHandle Entity = EntitySubsystem->CreateEntity(FloatsArchetype);
 		EntitySubsystem->AddComponentInstanceListToEntity(Entity, MakeArrayView(&InstanceInt, 1));
 		AITEST_EQUAL("There should be one entity across the whole system", EntitySubsystem->DebugGetEntityCount(), 1);
 		AITEST_EQUAL("The original archetype should now have no entities", EntitySubsystem->DebugGetArchetypeEntitiesCount(FloatsArchetype), 0);
@@ -217,7 +217,7 @@ struct FEntityTest_RemovingFragment : FEntityTestBase
 	virtual bool InstantTest() override
 	{
 		CA_ASSUME(EntitySubsystem);
-		const FLWEntity Entity = EntitySubsystem->CreateEntity(FloatsIntsArchetype);
+		const FMassEntityHandle Entity = EntitySubsystem->CreateEntity(FloatsIntsArchetype);
 		EntitySubsystem->RemoveComponentFromEntity(Entity, FTestFragment_Float::StaticStruct());
 		AITEST_EQUAL("There should be just one entity across the whole system", EntitySubsystem->DebugGetEntityCount(), 1);
 		AITEST_EQUAL("The original archetype should now have no entities", EntitySubsystem->DebugGetArchetypeEntitiesCount(FloatsIntsArchetype), 0);
@@ -236,7 +236,7 @@ struct FEntityTest_RemovingLastFragment : FEntityTestBase
 	virtual bool InstantTest() override
 	{
 		CA_ASSUME(EntitySubsystem);
-		const FLWEntity Entity = EntitySubsystem->CreateEntity(FloatsArchetype);
+		const FMassEntityHandle Entity = EntitySubsystem->CreateEntity(FloatsArchetype);
 		EntitySubsystem->RemoveComponentFromEntity(Entity, FTestFragment_Float::StaticStruct());
 		AITEST_EQUAL("There should be one entity across the whole system", EntitySubsystem->DebugGetEntityCount(), 1);
 		AITEST_EQUAL("The original archetype should now have no entities", EntitySubsystem->DebugGetArchetypeEntitiesCount(FloatsArchetype), 0);
@@ -253,7 +253,7 @@ struct FEntityTest_DestroyEntity : FEntityTestBase
 	virtual bool InstantTest() override
 	{
 		CA_ASSUME(EntitySubsystem);
-		const FLWEntity Entity = EntitySubsystem->CreateEntity(FloatsArchetype);
+		const FMassEntityHandle Entity = EntitySubsystem->CreateEntity(FloatsArchetype);
 		AITEST_EQUAL("The entity should get associated to the right archetype", EntitySubsystem->GetArchetypeForEntity(Entity), FloatsArchetype);
 		EntitySubsystem->DestroyEntity(Entity);
 		AITEST_EQUAL("There should not be any entity across the whole system", EntitySubsystem->DebugGetEntityCount(), 0);
@@ -268,7 +268,7 @@ struct FEntityTest_EntityReservationAndBuilding : FEntityTestBase
 	virtual bool InstantTest() override
 	{
 		CA_ASSUME(EntitySubsystem);
-		const FLWEntity ReservedEntity = EntitySubsystem->ReserveEntity();
+		const FMassEntityHandle ReservedEntity = EntitySubsystem->ReserveEntity();
 		AITEST_TRUE("The reserved entity should be a valid entity", EntitySubsystem->IsEntityValid(ReservedEntity));
 		AITEST_FALSE("The reserved entity should not be a valid entity", EntitySubsystem->IsEntityBuilt(ReservedEntity));
 		EntitySubsystem->BuildEntity(ReservedEntity, FloatsArchetype);
@@ -291,7 +291,7 @@ struct FEntityTest_EntityReservationAndBuildingFromInstances : FEntityTestBase
 	virtual bool InstantTest() override
 	{
 		CA_ASSUME(EntitySubsystem);
-		const FLWEntity ReservedEntity = EntitySubsystem->ReserveEntity();
+		const FMassEntityHandle ReservedEntity = EntitySubsystem->ReserveEntity();
 		AITEST_TRUE("The reserved entity should be a valid entity", EntitySubsystem->IsEntityValid(ReservedEntity));
 		AITEST_FALSE("The reserved entity should not be a valid entity", EntitySubsystem->IsEntityBuilt(ReservedEntity));
 		EntitySubsystem->BuildEntity(ReservedEntity, MakeArrayView(&InstanceInt, 1));
@@ -314,7 +314,7 @@ struct FEntityTest_ReleaseEntity : FEntityTestBase
 	virtual bool InstantTest() override
 	{
 		CA_ASSUME(EntitySubsystem);
-		const FLWEntity ReservedEntity = EntitySubsystem->ReserveEntity();
+		const FMassEntityHandle ReservedEntity = EntitySubsystem->ReserveEntity();
 		AITEST_TRUE("The reserved entity should be a valid entity", EntitySubsystem->IsEntityValid(ReservedEntity));
 		AITEST_FALSE("The reserved entity should not be a valid entity", EntitySubsystem->IsEntityBuilt(ReservedEntity));
 		AITEST_EQUAL("There should only be one entity across the whole system", EntitySubsystem->DebugGetEntityCount(), 1);
@@ -332,14 +332,14 @@ struct FEntityTest_ReserveAPreviouslyBuiltEntity : FEntityTestBase
 	{
 		CA_ASSUME(EntitySubsystem);
 		{
-			const FLWEntity Entity = EntitySubsystem->CreateEntity(IntsArchetype);
+			const FMassEntityHandle Entity = EntitySubsystem->CreateEntity(IntsArchetype);
 			AITEST_EQUAL("The entity should get associated to the right archetype", EntitySubsystem->GetArchetypeForEntity(Entity), IntsArchetype);
 			EntitySubsystem->DestroyEntity(Entity);
 			AITEST_EQUAL("There should not be any entity across the whole system", EntitySubsystem->DebugGetEntityCount(), 0);
 			AITEST_EQUAL("The original archetype should now have no entities", EntitySubsystem->DebugGetArchetypeEntitiesCount(IntsArchetype), 0);
 		}
 		
-		const FLWEntity ReservedEntity = EntitySubsystem->ReserveEntity();
+		const FMassEntityHandle ReservedEntity = EntitySubsystem->ReserveEntity();
 		AITEST_TRUE("The reserved entity should be a valid entity", EntitySubsystem->IsEntityValid(ReservedEntity));
 		AITEST_FALSE("The reserved entity should not be a valid entity", EntitySubsystem->IsEntityBuilt(ReservedEntity));
 		AITEST_EQUAL("There should only be one entity across the whole system", EntitySubsystem->DebugGetEntityCount(), 1);
diff --git a/Engine/Plugins/Runtime/MassEntity/Source/MassEntityTestSuite/Private/MassExecutionTest.cpp b/Engine/Plugins/Runtime/MassEntity/Source/MassEntityTestSuite/Private/MassExecutionTest.cpp
--- a/Engine/Plugins/Runtime/MassEntity/Source/MassEntityTestSuite/Private/MassExecutionTest.cpp
+++ b/Engine/Plugins/Runtime/MassEntity/Source/MassEntityTestSuite/Private/MassExecutionTest.cpp
@@ -36,9 +36,9 @@ struct FExecution_EmptyArray : FExecutionTestBase
 	{
 		CA_ASSUME(EntitySubsystem); // if EntitySubsystem null InstantTest won't be called at all
 		const float DeltaSeconds = 0.f;
-		FPipeContext PipeContext(*EntitySubsystem, DeltaSeconds);
+		FMassProcessingContext PipeContext(*EntitySubsystem, DeltaSeconds);
 		// no test performed, let's just see if it results in errors/warnings
-		UE::Pipe::Executor::RunProcessorsView(TArrayView<UPipeProcessor*>(), PipeContext);
+		UE::Pipe::Executor::RunProcessorsView(TArrayView<UMassProcessor*>(), PipeContext);
 		return true;
 	}
 };
@@ -51,8 +51,8 @@ struct FExecution_EmptyPipeline : FExecutionTestBase
 	{
 		CA_ASSUME(EntitySubsystem); // if EntitySubsystem null InstantTest won't be called at all
 		const float DeltaSeconds = 0.f;
-		FPipeContext PipeContext(*EntitySubsystem, DeltaSeconds);
-		FRuntimePipeline Pipeline;
+		FMassProcessingContext PipeContext(*EntitySubsystem, DeltaSeconds);
+		FMassRuntimePipeline Pipeline;
 		// no test performed, let's just see if it results in errors/warnings
 		UE::Pipe::Executor::Run(Pipeline, PipeContext);
 		return true;
@@ -67,14 +67,14 @@ struct FExecution_InvalidPipeContext : FExecutionTestBase
 	{
 		CA_ASSUME(EntitySubsystem);
 		const float DeltaSeconds = 0.f;
-		FPipeContext PipeContext;
+		FMassProcessingContext PipeContext;
 		// test assumption
-		AITEST_NULL("FPipeContext\'s default constructor is expected to set FPipeContext.EntitySubsystem to null", PipeContext.EntitySubsystem);
+		AITEST_NULL("FMassProcessingContext\'s default constructor is expected to set FMassProcessingContext.EntitySubsystem to null", PipeContext.EntitySubsystem);
 		
 		GetTestRunner().AddExpectedError(TEXT("PipeContext.EntitySubsystem is null"), EAutomationExpectedErrorFlags::Contains, 1);		
 		// note that using RunProcessorsView is to bypass reasonable tests UE::Pipe::Executor::Run(Pipeline,...) does that are 
 		// reported via ensures which are not handled by the automation framework
-		UE::Pipe::Executor::RunProcessorsView(TArrayView<UPipeProcessor*>(), PipeContext);
+		UE::Pipe::Executor::RunProcessorsView(TArrayView<UMassProcessor*>(), PipeContext);
 		return true;
 	}
 };
@@ -88,8 +88,8 @@ struct FExecution_SingleNullProcessor : FExecutionTestBase
 	{
 		CA_ASSUME(EntitySubsystem);
 		const float DeltaSeconds = 0.f;
-		FPipeContext PipeContext(*EntitySubsystem, DeltaSeconds);
-		TArray<UPipeProcessor*> Processors;
+		FMassProcessingContext PipeContext(*EntitySubsystem, DeltaSeconds);
+		TArray<UMassProcessor*> Processors;
 		Processors.Add(nullptr);
 		
 
@@ -109,11 +109,11 @@ struct FExecution_SingleValidProcessor : FExecutionTestBase
 	{
 		CA_ASSUME(EntitySubsystem);
 		const float DeltaSeconds = 0.f;
-		FPipeContext PipeContext(*EntitySubsystem, DeltaSeconds);
+		FMassProcessingContext PipeContext(*EntitySubsystem, DeltaSeconds);
 		UPipeTestProcessorBase* Processor = NewObject<UPipeTestProcessorBase>(EntitySubsystem);
 		check(Processor);
 
-		// nothing should break. The actual result of processing is getting tested in PipeProcessorTests.cpp
+		// nothing should break. The actual result of processing is getting tested in MassProcessorTests.cpp
 		UE::Pipe::Executor::Run(*Processor, PipeContext);
 		return true;
 	}
@@ -127,8 +127,8 @@ struct FExecution_MultipleNullProcessors : FExecutionTestBase
 	{
 		CA_ASSUME(EntitySubsystem);
 		const float DeltaSeconds = 0.f;
-		FPipeContext PipeContext(*EntitySubsystem, DeltaSeconds);
-		TArray<UPipeProcessor*> Processors;
+		FMassProcessingContext PipeContext(*EntitySubsystem, DeltaSeconds);
+		TArray<UMassProcessor*> Processors;
 		Processors.Add(nullptr);
 		Processors.Add(nullptr);
 		Processors.Add(nullptr);
@@ -150,19 +150,19 @@ struct FExecution_Sparse : FEntityTestBase
 	{
 		CA_ASSUME(EntitySubsystem);
 		const float DeltaSeconds = 0.f;
-		FPipeContext PipeContext(*EntitySubsystem, DeltaSeconds);
+		FMassProcessingContext PipeContext(*EntitySubsystem, DeltaSeconds);
 		UPipeTestProcessorBase* Processor = NewObject<UPipeTestProcessorBase>(EntitySubsystem);
 		check(Processor);
 
-		FRuntimePipeline Pipeline;
+		FMassRuntimePipeline Pipeline;
 		{
-			TArray<UPipeProcessor*> Processors;
+			TArray<UMassProcessor*> Processors;
 			Processors.Add(Processor);
 			Pipeline.SetProcessors(MoveTemp(Processors));
 		}
 
 		FArchetypeChunkCollection ChunkCollection(FloatsArchetype);
-		// nothing should break. The actual result of processing is getting tested in PipeProcessorTests.cpp
+		// nothing should break. The actual result of processing is getting tested in MassProcessorTests.cpp
 		
 		UE::Pipe::Executor::RunSparse(Pipeline, PipeContext, ChunkCollection);
 
diff --git a/Engine/Plugins/Runtime/MassEntity/Source/MassEntityTestSuite/Private/MassProcessorTest.cpp b/Engine/Plugins/Runtime/MassEntity/Source/MassEntityTestSuite/Private/MassProcessorTest.cpp
--- a/Engine/Plugins/Runtime/MassEntity/Source/MassEntityTestSuite/Private/MassProcessorTest.cpp
+++ b/Engine/Plugins/Runtime/MassEntity/Source/MassEntityTestSuite/Private/MassProcessorTest.cpp
@@ -13,7 +13,7 @@
 
 PRAGMA_DISABLE_OPTIMIZATION
 
-namespace FPipeProcessorTest
+namespace FMassProcessorTest
 {
 
 template<typename TProcessor>
@@ -21,15 +21,15 @@ int32 SimpleProcessorRun(UMassEntitySubsystem& EntitySubsystem)
 {
 	int32 EntityProcessedCount = 0;
 	TProcessor* Processor = NewObject<TProcessor>(&EntitySubsystem);
-	Processor->ExecutionFunction = [Processor, &EntityProcessedCount](UMassEntitySubsystem& InEntitySubsystem, FLWComponentSystemExecutionContext& Context) {
+	Processor->ExecutionFunction = [Processor, &EntityProcessedCount](UMassEntitySubsystem& InEntitySubsystem, FMassExecutionContext& Context) {
 		check(Processor);
-		Processor->TestGetQuery().ForEachEntityChunk(InEntitySubsystem, Context, [Processor, &EntityProcessedCount](FLWComponentSystemExecutionContext& Context)
+		Processor->TestGetQuery().ForEachEntityChunk(InEntitySubsystem, Context, [Processor, &EntityProcessedCount](FMassExecutionContext& Context)
 			{
 				EntityProcessedCount += Context.GetEntitiesNum();
 			});
 	};
 
-	FPipeContext PipeContext(EntitySubsystem, /*DeltaSeconds=*/0.f);
+	FMassProcessingContext PipeContext(EntitySubsystem, /*DeltaSeconds=*/0.f);
 	UE::Pipe::Executor::Run(*Processor, PipeContext);
 
 	return EntityProcessedCount;
@@ -56,7 +56,7 @@ struct FProcessorTest_NoMatchingEntities : FEntityTestBase
 	{
 		CA_ASSUME(EntitySubsystem);
 
-		TArray<FLWEntity> EntitiesCreated;
+		TArray<FMassEntityHandle> EntitiesCreated;
 		EntitySubsystem->BatchCreateEntities(IntsArchetype, 100, EntitiesCreated);
 		int32 EntityProcessedCount = SimpleProcessorRun<UPipeTestProcessor_Floats>(*EntitySubsystem);
 		AITEST_EQUAL("No matching entities have been created yet so the processor shouldn't do any work.", EntityProcessedCount, 0);
@@ -73,7 +73,7 @@ struct FProcessorTest_OneMatchingArchetype : FEntityTestBase
 		CA_ASSUME(EntitySubsystem);
 
 		const int32 EntitiesToCreate = 100;
-		TArray<FLWEntity> EntitiesCreated;
+		TArray<FMassEntityHandle> EntitiesCreated;
 		EntitySubsystem->BatchCreateEntities(FloatsArchetype, EntitiesToCreate, EntitiesCreated);
 		int32 EntityProcessedCount = SimpleProcessorRun<UPipeTestProcessor_Floats>(*EntitySubsystem);
 		AITEST_EQUAL("No matching entities have been created yet so the processor shouldn't do any work.", EntityProcessedCount, EntitiesToCreate);
@@ -91,7 +91,7 @@ struct FProcessorTest_MultipleMatchingArchetype : FEntityTestBase
 		CA_ASSUME(EntitySubsystem);
 
 		const int32 EntitiesToCreate = 100;
-		TArray<FLWEntity> EntitiesCreated;
+		TArray<FMassEntityHandle> EntitiesCreated;
 		EntitySubsystem->BatchCreateEntities(FloatsArchetype, EntitiesToCreate, EntitiesCreated);
 		EntitySubsystem->BatchCreateEntities(FloatsIntsArchetype, EntitiesToCreate, EntitiesCreated);
 		EntitySubsystem->BatchCreateEntities(IntsArchetype, EntitiesToCreate, EntitiesCreated);
@@ -117,7 +117,7 @@ struct FProcessorTest_Requirements : FEntityTestBase
 IMPLEMENT_AI_INSTANT_TEST(FProcessorTest_Requirements, "System.Pipe.Processor.Requirements");
 
 
-} // FPipeProcessorTestTest
+} // FMassProcessorTestTest
 
 PRAGMA_ENABLE_OPTIMIZATION
 
diff --git a/Engine/Plugins/Runtime/MassEntity/Source/MassEntityTestSuite/Private/MassQueryTest.cpp b/Engine/Plugins/Runtime/MassEntity/Source/MassEntityTestSuite/Private/MassQueryTest.cpp
--- a/Engine/Plugins/Runtime/MassEntity/Source/MassEntityTestSuite/Private/MassQueryTest.cpp
+++ b/Engine/Plugins/Runtime/MassEntity/Source/MassEntityTestSuite/Private/MassQueryTest.cpp
@@ -19,8 +19,8 @@ mz@todo:
 	- test Present/Absent modes for tags
 	- test Present/Absent modes for non-tag fragments
 	- read/write tag fragments 
-	- constructs like Query.AddRequirement<FTestFragment_Tag>(ELWComponentAccess::ReadWrite);
-	- constructs like Query.AddRequirement<FTestFragment_Int>(ELWComponentAccess::ReadWrite, Absent);
+	- constructs like Query.AddRequirement<FTestFragment_Tag>(EMassFragmentAccess::ReadWrite);
+	- constructs like Query.AddRequirement<FTestFragment_Int>(EMassFragmentAccess::ReadWrite, Absent);
 */
 
 
@@ -34,7 +34,7 @@ struct FQueryTest_ProcessorRequirements : FEntityTestBase
 		CA_ASSUME(EntitySubsystem);
 
 		UPipeTestProcessor_Floats* Processor = NewObject<UPipeTestProcessor_Floats>(EntitySubsystem);
-		TConstArrayView<FLWComponentRequirement> Requirements = Processor->TestGetQuery().GetRequirements();
+		TConstArrayView<FMassFragmentRequirement> Requirements = Processor->TestGetQuery().GetRequirements();
 		
 		AITEST_TRUE("Query should have extracted some requirements from the given Processor", Requirements.Num() > 0);
 		AITEST_TRUE("There should be exactly one requirement", Requirements.Num() == 1);
@@ -52,8 +52,8 @@ struct FQueryTest_ExplicitRequirements : FEntityTestBase
 	{
 		CA_ASSUME(EntitySubsystem);
 				
-		FLWComponentQuery Query({ FTestFragment_Float::StaticStruct()});
-		TConstArrayView<FLWComponentRequirement> Requirements = Query.GetRequirements();
+		FMassEntityQuery Query({ FTestFragment_Float::StaticStruct()});
+		TConstArrayView<FMassFragmentRequirement> Requirements = Query.GetRequirements();
 
 		AITEST_TRUE("Query should have extracted some requirements from the given Processor", Requirements.Num() > 0);
 		AITEST_TRUE("There should be exactly one requirement", Requirements.Num() == 1);
@@ -71,16 +71,16 @@ struct FQueryTest_FragmentViewBinding : FEntityTestBase
 	{
 		CA_ASSUME(EntitySubsystem);
 
-		FLWEntity Entity = EntitySubsystem->CreateEntity(FloatsArchetype);
+		FMassEntityHandle Entity = EntitySubsystem->CreateEntity(FloatsArchetype);
 		FTestFragment_Float& TestedFragment = EntitySubsystem->GetComponentDataChecked<FTestFragment_Float>(Entity);
 		AITEST_TRUE("Initial value of the component should match expectations", TestedFragment.Value == 0.f);
 
 		UPipeTestProcessor_Floats* Processor = NewObject<UPipeTestProcessor_Floats>(EntitySubsystem);
-		Processor->ExecutionFunction = [Processor](UMassEntitySubsystem& InEntitySubsystem, FLWComponentSystemExecutionContext& Context) {
+		Processor->ExecutionFunction = [Processor](UMassEntitySubsystem& InEntitySubsystem, FMassExecutionContext& Context) {
 			check(Processor);
-			//FLWComponentQuery Query(*Processor);
+			//FMassEntityQuery Query(*Processor);
 
-			Processor->TestGetQuery().ForEachEntityChunk(InEntitySubsystem, Context, [](FLWComponentSystemExecutionContext& Context)
+			Processor->TestGetQuery().ForEachEntityChunk(InEntitySubsystem, Context, [](FMassExecutionContext& Context)
 				{
 					TArrayView<FTestFragment_Float> Floats = Context.GetMutableComponentView<FTestFragment_Float>();
 
@@ -91,7 +91,7 @@ struct FQueryTest_FragmentViewBinding : FEntityTestBase
 				});
 		};
 
-		FPipeContext PipeContext(*EntitySubsystem, /*DeltaSeconds=*/0.f);
+		FMassProcessingContext PipeContext(*EntitySubsystem, /*DeltaSeconds=*/0.f);
 		UE::Pipe::Executor::Run(*Processor, PipeContext);
 
 		AITEST_EQUAL("Fragment value should have changed to the expected value", TestedFragment.Value, 13.f);
@@ -108,14 +108,14 @@ struct FQueryTest_ExecuteSingleArchetype : FEntityTestBase
 		CA_ASSUME(EntitySubsystem);
 
 		const int32 NumToCreate = 10;
-		TArray<FLWEntity> EntitiesCreated;
+		TArray<FMassEntityHandle> EntitiesCreated;
 		EntitySubsystem->BatchCreateEntities(FloatsArchetype, NumToCreate, EntitiesCreated);
 		
 		int TotalProcessed = 0;
 
-		FLWComponentSystemExecutionContext ExecContext;
-		FLWComponentQuery Query({ FTestFragment_Float::StaticStruct() });
-		Query.ForEachEntityChunk(*EntitySubsystem, ExecContext, [&TotalProcessed](FLWComponentSystemExecutionContext& Context)
+		FMassExecutionContext ExecContext;
+		FMassEntityQuery Query({ FTestFragment_Float::StaticStruct() });
+		Query.ForEachEntityChunk(*EntitySubsystem, ExecContext, [&TotalProcessed](FMassExecutionContext& Context)
 			{
 				TotalProcessed += Context.GetEntitiesNum();
 				TArrayView<FTestFragment_Float> Floats = Context.GetMutableComponentView<FTestFragment_Float>();
@@ -128,7 +128,7 @@ struct FQueryTest_ExecuteSingleArchetype : FEntityTestBase
 
 		AITEST_TRUE("The number of entities processed needs to match expectations", TotalProcessed == NumToCreate);
 
-		for (FLWEntity& Entity : EntitiesCreated)
+		for (FMassEntityHandle& Entity : EntitiesCreated)
 		{
 			const FTestFragment_Float& TestedFragment = EntitySubsystem->GetComponentDataChecked<FTestFragment_Float>(Entity);
 			AITEST_EQUAL("Every fragment value should have changed to the expected value", TestedFragment.Value, 13.f);
@@ -149,7 +149,7 @@ struct FQueryTest_ExecuteMultipleArchetypes : FEntityTestBase
 		const int32 FloatsArchetypeCreated = 7;
 		const int32 IntsArchetypeCreated = 11;
 		const int32 FloatsIntsArchetypeCreated = 13;
-		TArray<FLWEntity> EntitiesCreated;
+		TArray<FMassEntityHandle> EntitiesCreated;
 		EntitySubsystem->BatchCreateEntities(IntsArchetype, IntsArchetypeCreated, EntitiesCreated);
 		// clear to store only the float-related entities
 		EntitiesCreated.Reset();
@@ -157,9 +157,9 @@ struct FQueryTest_ExecuteMultipleArchetypes : FEntityTestBase
 		EntitySubsystem->BatchCreateEntities(FloatsIntsArchetype, FloatsIntsArchetypeCreated, EntitiesCreated);
 
 		int TotalProcessed = 0;
-		FLWComponentSystemExecutionContext ExecContext;
-		FLWComponentQuery Query({ FTestFragment_Float::StaticStruct() });
-		Query.ForEachEntityChunk(*EntitySubsystem, ExecContext, [&TotalProcessed](FLWComponentSystemExecutionContext& Context)
+		FMassExecutionContext ExecContext;
+		FMassEntityQuery Query({ FTestFragment_Float::StaticStruct() });
+		Query.ForEachEntityChunk(*EntitySubsystem, ExecContext, [&TotalProcessed](FMassExecutionContext& Context)
 			{
 				TotalProcessed += Context.GetEntitiesNum();
 				TArrayView<FTestFragment_Float> Floats = Context.GetMutableComponentView<FTestFragment_Float>();
@@ -172,7 +172,7 @@ struct FQueryTest_ExecuteMultipleArchetypes : FEntityTestBase
 
 		AITEST_TRUE("The number of entities processed needs to match expectations", TotalProcessed == FloatsIntsArchetypeCreated + FloatsArchetypeCreated);
 
-		for (FLWEntity& Entity : EntitiesCreated)
+		for (FMassEntityHandle& Entity : EntitiesCreated)
 		{
 			const FTestFragment_Float& TestedFragment = EntitySubsystem->GetComponentDataChecked<FTestFragment_Float>(Entity);
 			AITEST_EQUAL("Every fragment value should have changed to the expected value", TestedFragment.Value, 13.f);
@@ -192,12 +192,12 @@ struct FQueryTest_ExecuteSparse : FEntityTestBase
 		CA_ASSUME(EntitySubsystem);
 
 		const int32 NumToCreate = 10;
-		TArray<FLWEntity> AllEntitiesCreated;
+		TArray<FMassEntityHandle> AllEntitiesCreated;
 		EntitySubsystem->BatchCreateEntities(FloatsArchetype, NumToCreate, AllEntitiesCreated);
 		
 		TArray<int32> IndicesToProcess = { 1, 2, 3, 6, 7};
-		TArray<FLWEntity> EntitiesToProcess;
-		TArray<FLWEntity> EntitiesToIgnore;
+		TArray<FMassEntityHandle> EntitiesToProcess;
+		TArray<FMassEntityHandle> EntitiesToIgnore;
 		for (int32 i = 0; i < AllEntitiesCreated.Num(); ++i)
 		{
 			if (IndicesToProcess.Find(i) != INDEX_NONE)
@@ -213,9 +213,9 @@ struct FQueryTest_ExecuteSparse : FEntityTestBase
 
 		int TotalProcessed = 0;
 
-		FLWComponentSystemExecutionContext ExecContext;
-		FLWComponentQuery().AddRequirement<FTestFragment_Float>(ELWComponentAccess::ReadWrite)
-			.ForEachEntityChunk(FArchetypeChunkCollection(FloatsArchetype, EntitiesToProcess), *EntitySubsystem, ExecContext, [&TotalProcessed](FLWComponentSystemExecutionContext& Context)
+		FMassExecutionContext ExecContext;
+		FMassEntityQuery().AddRequirement<FTestFragment_Float>(EMassFragmentAccess::ReadWrite)
+			.ForEachEntityChunk(FArchetypeChunkCollection(FloatsArchetype, EntitiesToProcess), *EntitySubsystem, ExecContext, [&TotalProcessed](FMassExecutionContext& Context)
 			{
 				TotalProcessed += Context.GetEntitiesNum();
 				TArrayView<FTestFragment_Float> Floats = Context.GetMutableComponentView<FTestFragment_Float>();
@@ -228,13 +228,13 @@ struct FQueryTest_ExecuteSparse : FEntityTestBase
 
 		AITEST_TRUE("The number of entities processed needs to match expectations", TotalProcessed == IndicesToProcess.Num());
 
-		for (FLWEntity& Entity : EntitiesToProcess)
+		for (FMassEntityHandle& Entity : EntitiesToProcess)
 		{
 			const FTestFragment_Float& TestedFragment = EntitySubsystem->GetComponentDataChecked<FTestFragment_Float>(Entity);
 			AITEST_EQUAL("Every fragment value should have changed to the expected value", TestedFragment.Value, 13.f);
 		}
 		
-		for (FLWEntity& Entity : EntitiesToIgnore)
+		for (FMassEntityHandle& Entity : EntitiesToIgnore)
 		{
 			const FTestFragment_Float& TestedFragment = EntitySubsystem->GetComponentDataChecked<FTestFragment_Float>(Entity);
 			AITEST_EQUAL("Untouched entites should retain default fragment value ", TestedFragment.Value, 0.f);
@@ -255,9 +255,9 @@ struct FQueryTest_TagPresent : FEntityTestBase
 		TArray<const UScriptStruct*> Fragments = {FTestFragment_Float::StaticStruct(), FTestFragment_Tag::StaticStruct()};
 		const FArchetypeHandle FloatsTagArchetype = EntitySubsystem->CreateArchetype(Fragments);
 
-		FLWComponentQuery Query;
-		Query.AddRequirement<FTestFragment_Float>(ELWComponentAccess::ReadWrite);
-		Query.AddTagRequirement<FTestFragment_Tag>(ELWComponentPresence::All);
+		FMassEntityQuery Query;
+		Query.AddRequirement<FTestFragment_Float>(EMassFragmentAccess::ReadWrite);
+		Query.AddTagRequirement<FTestFragment_Tag>(EMassFragmentPresence::All);
 		Query.CacheArchetypes(*EntitySubsystem);
 
 		AITEST_EQUAL("There's a single archetype matching the requirements", Query.GetArchetypes().Num(), 1);
@@ -278,9 +278,9 @@ struct FQueryTest_TagAbsent : FEntityTestBase
 		TArray<const UScriptStruct*> Fragments = { FTestFragment_Float::StaticStruct(), FTestFragment_Tag::StaticStruct() };
 		const FArchetypeHandle FloatsTagArchetype = EntitySubsystem->CreateArchetype(Fragments);
 
-		FLWComponentQuery Query;
-		Query.AddRequirement<FTestFragment_Float>(ELWComponentAccess::ReadWrite);
-		Query.AddTagRequirement<FTestFragment_Tag>(ELWComponentPresence::None);
+		FMassEntityQuery Query;
+		Query.AddRequirement<FTestFragment_Float>(EMassFragmentAccess::ReadWrite);
+		Query.AddTagRequirement<FTestFragment_Tag>(EMassFragmentPresence::None);
 		Query.CacheArchetypes(*EntitySubsystem);
 
 		AITEST_EQUAL("There are exactly two archetypes matching the requirements", Query.GetArchetypes().Num(), 2);
@@ -300,8 +300,8 @@ struct FQueryTest_FragmentPresent : FEntityTestBase
 	{
 		CA_ASSUME(EntitySubsystem);
 
-		FLWComponentQuery Query;
-		Query.AddRequirement<FTestFragment_Int>(ELWComponentAccess::ReadOnly);
+		FMassEntityQuery Query;
+		Query.AddRequirement<FTestFragment_Int>(EMassFragmentAccess::ReadOnly);
 		Query.CacheArchetypes(*EntitySubsystem);
 
 		AITEST_EQUAL("There are exactly two archetypes matching the requirements", Query.GetArchetypes().Num(), 2);
@@ -320,8 +320,8 @@ struct FQueryTest_OnlySingleAbsentFragment : FEntityTestBase
 	{
 		CA_ASSUME(EntitySubsystem);
 
-		FLWComponentQuery Query;
-		Query.AddRequirement<FTestFragment_Int>(ELWComponentAccess::ReadOnly, ELWComponentPresence::None);
+		FMassEntityQuery Query;
+		Query.AddRequirement<FTestFragment_Int>(EMassFragmentAccess::ReadOnly, EMassFragmentPresence::None);
 		
 		AITEST_FALSE("The query is not valid", Query.CheckValidity());
 
@@ -343,9 +343,9 @@ struct FQueryTest_OnlyMultipleAbsentFragments : FEntityTestBase
 	{
 		CA_ASSUME(EntitySubsystem);
 
-		FLWComponentQuery Query;
-		Query.AddRequirement<FTestFragment_Int>(ELWComponentAccess::ReadOnly, ELWComponentPresence::None);
-		Query.AddRequirement<FTestFragment_Float>(ELWComponentAccess::ReadOnly, ELWComponentPresence::None);
+		FMassEntityQuery Query;
+		Query.AddRequirement<FTestFragment_Int>(EMassFragmentAccess::ReadOnly, EMassFragmentPresence::None);
+		Query.AddRequirement<FTestFragment_Float>(EMassFragmentAccess::ReadOnly, EMassFragmentPresence::None);
 
 		AITEST_FALSE("The query is not valid", Query.CheckValidity());
 
@@ -367,9 +367,9 @@ struct FQueryTest_AbsentAndPresentFragments : FEntityTestBase
 	{
 		CA_ASSUME(EntitySubsystem);
 
-		FLWComponentQuery Query;
-		Query.AddRequirement<FTestFragment_Int>(ELWComponentAccess::None, ELWComponentPresence::None);
-		Query.AddRequirement<FTestFragment_Float>(ELWComponentAccess::ReadOnly, ELWComponentPresence::All);
+		FMassEntityQuery Query;
+		Query.AddRequirement<FTestFragment_Int>(EMassFragmentAccess::None, EMassFragmentPresence::None);
+		Query.AddRequirement<FTestFragment_Float>(EMassFragmentAccess::ReadOnly, EMassFragmentPresence::All);
 
 		AITEST_TRUE("The query is valid", Query.CheckValidity());
 		Query.CacheArchetypes(*EntitySubsystem);
@@ -388,8 +388,8 @@ struct FQueryTest_SingleOptionalFragment : FEntityTestBase
 	{
 		CA_ASSUME(EntitySubsystem);
 
-		FLWComponentQuery Query;
-		Query.AddRequirement<FTestFragment_Int>(ELWComponentAccess::ReadWrite, ELWComponentPresence::Optional);
+		FMassEntityQuery Query;
+		Query.AddRequirement<FTestFragment_Int>(EMassFragmentAccess::ReadWrite, EMassFragmentPresence::Optional);
 		Query.CacheArchetypes(*EntitySubsystem);
 
 		AITEST_EQUAL("There are exactly two archetypes matching the requirements", Query.GetArchetypes().Num(), 2);
@@ -408,9 +408,9 @@ struct FQueryTest_MultipleOptionalFragment : FEntityTestBase
 	{
 		CA_ASSUME(EntitySubsystem);
 
-		FLWComponentQuery Query;
-		Query.AddRequirement<FTestFragment_Int>(ELWComponentAccess::ReadWrite, ELWComponentPresence::Optional);
-		Query.AddRequirement<FTestFragment_Float>(ELWComponentAccess::ReadWrite, ELWComponentPresence::Optional);
+		FMassEntityQuery Query;
+		Query.AddRequirement<FTestFragment_Int>(EMassFragmentAccess::ReadWrite, EMassFragmentPresence::Optional);
+		Query.AddRequirement<FTestFragment_Float>(EMassFragmentAccess::ReadWrite, EMassFragmentPresence::Optional);
 		Query.CacheArchetypes(*EntitySubsystem);
 
 		AITEST_EQUAL("All three archetype meet requirements", Query.GetArchetypes().Num(), 3);
@@ -429,18 +429,18 @@ struct FQueryTest_UsingOptionalFragment : FEntityTestBase
 		CA_ASSUME(EntitySubsystem);
 
 		EntitySubsystem->CreateEntity(FloatsArchetype);
-		const FLWEntity EntityWithFloatsInts = EntitySubsystem->CreateEntity(FloatsIntsArchetype);
+		const FMassEntityHandle EntityWithFloatsInts = EntitySubsystem->CreateEntity(FloatsIntsArchetype);
 		EntitySubsystem->CreateEntity(IntsArchetype);
 
 		const int32 IntValueSet = 123;
 		int TotalProcessed = 0;
 		int EmptyIntsViewCount = 0;
 
-		FLWComponentQuery Query;
-		Query.AddRequirement<FTestFragment_Int>(ELWComponentAccess::ReadWrite, ELWComponentPresence::Optional);
-		Query.AddRequirement<FTestFragment_Float>(ELWComponentAccess::ReadWrite, ELWComponentPresence::All);
-		FLWComponentSystemExecutionContext ExecContext;
-		Query.ForEachEntityChunk(*EntitySubsystem, ExecContext, [&TotalProcessed, &EmptyIntsViewCount, IntValueSet](FLWComponentSystemExecutionContext& Context) {
+		FMassEntityQuery Query;
+		Query.AddRequirement<FTestFragment_Int>(EMassFragmentAccess::ReadWrite, EMassFragmentPresence::Optional);
+		Query.AddRequirement<FTestFragment_Float>(EMassFragmentAccess::ReadWrite, EMassFragmentPresence::All);
+		FMassExecutionContext ExecContext;
+		Query.ForEachEntityChunk(*EntitySubsystem, ExecContext, [&TotalProcessed, &EmptyIntsViewCount, IntValueSet](FMassExecutionContext& Context) {
 			++TotalProcessed;
 			TArrayView<FTestFragment_Int> Ints = Context.GetMutableComponentView<FTestFragment_Int>();
 			if (Ints.Num() == 0)
@@ -481,9 +481,9 @@ struct FQueryTest_AnyFragment : FEntityTestBase
 		const FArchetypeHandle BoolArchetype = EntitySubsystem->CreateArchetype({ FTestFragment_Bool::StaticStruct() });
 		const FArchetypeHandle BoolFloatArchetype = EntitySubsystem->CreateArchetype({ FTestFragment_Bool::StaticStruct(), FTestFragment_Float::StaticStruct() });
 
-		FLWComponentQuery Query;
-		Query.AddRequirement<FTestFragment_Int>(ELWComponentAccess::ReadWrite, ELWComponentPresence::Any);
-		Query.AddRequirement<FTestFragment_Bool>(ELWComponentAccess::ReadWrite, ELWComponentPresence::Any);
+		FMassEntityQuery Query;
+		Query.AddRequirement<FTestFragment_Int>(EMassFragmentAccess::ReadWrite, EMassFragmentPresence::Any);
+		Query.AddRequirement<FTestFragment_Bool>(EMassFragmentAccess::ReadWrite, EMassFragmentPresence::Any);
 		// this query should match: 
 		// IntsArchetype, FloatsIntsArchetype, BoolArchetype, BoolFloatArchetype
 		Query.CacheArchetypes(*EntitySubsystem);
@@ -496,8 +496,8 @@ struct FQueryTest_AnyFragment : FEntityTestBase
 			EntitySubsystem->CreateEntity(ArchetypeHandle);
 		}
 
-		FLWComponentSystemExecutionContext TestContext;
-		Query.ForEachEntityChunk(*EntitySubsystem, TestContext, [this](FLWComponentSystemExecutionContext& Context)
+		FMassExecutionContext TestContext;
+		Query.ForEachEntityChunk(*EntitySubsystem, TestContext, [this](FMassExecutionContext& Context)
 			{
 				TArrayView<FTestFragment_Bool> BoolView = Context.GetMutableComponentView<FTestFragment_Bool>();
 				TArrayView<FTestFragment_Int> IntView = Context.GetMutableComponentView<FTestFragment_Int>();
@@ -523,11 +523,11 @@ struct FQueryTest_AnyTag : FEntityTestBase
 		const FArchetypeHandle BDArchetype = EntitySubsystem->CreateArchetype({ FTestFragment_Int::StaticStruct(), FTestTag_B::StaticStruct(), FTestTag_D::StaticStruct() });
 		const FArchetypeHandle FloatACArchetype = EntitySubsystem->CreateArchetype({ FTestFragment_Float::StaticStruct(), FTestTag_A::StaticStruct(), FTestTag_C::StaticStruct() });
 
-		FLWComponentQuery Query;
+		FMassEntityQuery Query;
 		// at least one fragment requirement needs to be present for the query to be valid
-		Query.AddRequirement<FTestFragment_Int>(ELWComponentAccess::ReadOnly); 
-		Query.AddTagRequirement<FTestTag_A>(ELWComponentPresence::Any);
-		Query.AddTagRequirement<FTestTag_C>(ELWComponentPresence::Any);
+		Query.AddRequirement<FTestFragment_Int>(EMassFragmentAccess::ReadOnly); 
+		Query.AddTagRequirement<FTestTag_A>(EMassFragmentPresence::Any);
+		Query.AddTagRequirement<FTestTag_C>(EMassFragmentPresence::Any);
 		// this query should match: 
 		// ABArchetype, ACArchetype and ABCrchetype but not BDArchetype nor FEntityTestBase.IntsArchetype
 		Query.CacheArchetypes(*EntitySubsystem);
diff --git a/Engine/Plugins/Runtime/MassEntity/Source/MassEntityTestSuite/Private/MassTestTypes.cpp b/Engine/Plugins/Runtime/MassEntity/Source/MassEntityTestSuite/Private/MassTestTypes.cpp
--- a/Engine/Plugins/Runtime/MassEntity/Source/MassEntityTestSuite/Private/MassTestTypes.cpp
+++ b/Engine/Plugins/Runtime/MassEntity/Source/MassEntityTestSuite/Private/MassTestTypes.cpp
@@ -54,31 +54,31 @@ UPipeTestProcessorBase::UPipeTestProcessorBase()
 	bAutoRegisterWithProcessingPhases = false;
 	ExecutionFlags = int32(EProcessorExecutionFlags::All);
 
-	ExecutionFunction = [](UMassEntitySubsystem& InEntitySubsystem, FLWComponentSystemExecutionContext& Context) {};
-	RequirementsFunction = [](FLWComponentQuery& Query){};
+	ExecutionFunction = [](UMassEntitySubsystem& InEntitySubsystem, FMassExecutionContext& Context) {};
+	RequirementsFunction = [](FMassEntityQuery& Query){};
 }
 
 UPipeTestProcessor_Floats::UPipeTestProcessor_Floats()
 {
-	RequirementsFunction = [this](FLWComponentQuery& Query)
+	RequirementsFunction = [this](FMassEntityQuery& Query)
 	{
-		EntityQuery.AddRequirement<FTestFragment_Float>(ELWComponentAccess::ReadWrite);
+		EntityQuery.AddRequirement<FTestFragment_Float>(EMassFragmentAccess::ReadWrite);
 	};
 }
 
 UPipeTestProcessor_Ints::UPipeTestProcessor_Ints()
 {
-	RequirementsFunction = [this](FLWComponentQuery& Query)
+	RequirementsFunction = [this](FMassEntityQuery& Query)
 	{
-		EntityQuery.AddRequirement<FTestFragment_Int>(ELWComponentAccess::ReadWrite);
+		EntityQuery.AddRequirement<FTestFragment_Int>(EMassFragmentAccess::ReadWrite);
 	};
 }
 
 UPipeTestProcessor_FloatsInts::UPipeTestProcessor_FloatsInts()
 {
-	RequirementsFunction = [this](FLWComponentQuery& Query)
+	RequirementsFunction = [this](FMassEntityQuery& Query)
 	{
-		EntityQuery.AddRequirement<FTestFragment_Float>(ELWComponentAccess::ReadWrite);
-		EntityQuery.AddRequirement<FTestFragment_Int>(ELWComponentAccess::ReadWrite);
+		EntityQuery.AddRequirement<FTestFragment_Float>(EMassFragmentAccess::ReadWrite);
+		EntityQuery.AddRequirement<FTestFragment_Int>(EMassFragmentAccess::ReadWrite);
 	};
 }
\ No newline at end of file
diff --git a/Engine/Plugins/Runtime/MassEntity/Source/MassEntityTestSuite/Public/LWComponentTestFarmPlot.h b/Engine/Plugins/Runtime/MassEntity/Source/MassEntityTestSuite/Public/LWComponentTestFarmPlot.h
--- a/Engine/Plugins/Runtime/MassEntity/Source/MassEntityTestSuite/Public/LWComponentTestFarmPlot.h
+++ b/Engine/Plugins/Runtime/MassEntity/Source/MassEntityTestSuite/Public/LWComponentTestFarmPlot.h
@@ -29,19 +29,19 @@ struct FFarmVisualDataRow : public FTableRowBase
 };
 
 USTRUCT()
-struct FFarmJustBecameReadyToHarvest : public FComponentTag
+struct FFarmJustBecameReadyToHarvest : public FMassTag
 {
 	GENERATED_BODY()
 };
 
 USTRUCT()
-struct FFarmReadyToHarvest : public FComponentTag
+struct FFarmReadyToHarvest : public FMassTag
 {
 	GENERATED_BODY()
 };
 
 USTRUCT()
-struct FFarmGridCellData : public FLWComponentData
+struct FFarmGridCellData : public FMassFragment
 {
 	GENERATED_BODY()
 
@@ -50,7 +50,7 @@ struct FFarmGridCellData : public FLWComponentData
 };
 
 USTRUCT()
-struct FFarmWaterComponent : public FLWComponentData
+struct FFarmWaterComponent : public FMassFragment
 {
 	GENERATED_BODY()
 
@@ -59,7 +59,7 @@ struct FFarmWaterComponent : public FLWComponentData
 };
 
 USTRUCT()
-struct FFarmFlowerComponent : public FLWComponentData
+struct FFarmFlowerComponent : public FMassFragment
 {
 	GENERATED_BODY()
 		
@@ -68,7 +68,7 @@ struct FFarmFlowerComponent : public FLWComponentData
 };
 
 USTRUCT()
-struct FFarmCropComponent : public FLWComponentData
+struct FFarmCropComponent : public FMassFragment
 {
 	GENERATED_BODY()
 
@@ -77,7 +77,7 @@ struct FFarmCropComponent : public FLWComponentData
 
 
 USTRUCT()
-struct FFarmVisualComponent : public FLWComponentData
+struct FFarmVisualComponent : public FMassFragment
 {
 	GENERATED_BODY()
 
@@ -87,7 +87,7 @@ struct FFarmVisualComponent : public FLWComponentData
 };
 
 USTRUCT()
-struct FHarvestTimerComponent : public FLWComponentData
+struct FHarvestTimerComponent : public FMassFragment
 {
 	GENERATED_BODY()
 
@@ -101,11 +101,11 @@ class UComponentUpdateSystem : public UObject
 	GENERATED_BODY()
 public:
 	virtual void PostInitProperties() override;
-	virtual void Execute(UMassEntitySubsystem& EntitySubsystem, FLWComponentSystemExecutionContext& Context) {}
+	virtual void Execute(UMassEntitySubsystem& EntitySubsystem, FMassExecutionContext& Context) {}
 protected:
 	virtual void ConfigureQueries() {}
 protected:
-	FLWComponentQuery EntityQuery;
+	FMassEntityQuery EntityQuery;
 };
 
 
@@ -117,13 +117,13 @@ class UFarmWaterUpdateSystem : public UComponentUpdateSystem
 public:
 	virtual void ConfigureQueries() override
 	{
-		EntityQuery.AddRequirement<FFarmWaterComponent>(ELWComponentAccess::ReadWrite);
+		EntityQuery.AddRequirement<FFarmWaterComponent>(EMassFragmentAccess::ReadWrite);
 	}
 
-	virtual void Execute(UMassEntitySubsystem& EntitySubsystem, FLWComponentSystemExecutionContext& Context) override
+	virtual void Execute(UMassEntitySubsystem& EntitySubsystem, FMassExecutionContext& Context) override
 	{
 		QUICK_SCOPE_CYCLE_COUNTER(UFarmWaterUpdateSystem_Run);
-		EntityQuery.ForEachEntityChunk(EntitySubsystem, Context, [this](FLWComponentSystemExecutionContext& Context) {
+		EntityQuery.ForEachEntityChunk(EntitySubsystem, Context, [this](FMassExecutionContext& Context) {
 
 			TArrayView<FFarmWaterComponent> WaterList = Context.GetMutableComponentView<FFarmWaterComponent>();
 
@@ -144,16 +144,16 @@ class UFarmHarvestTimerSystem_Flowers : public UComponentUpdateSystem
 public:
 	virtual void ConfigureQueries() override
 	{
-		EntityQuery.AddRequirement<FHarvestTimerComponent>(ELWComponentAccess::ReadWrite);
-		EntityQuery.AddRequirement<FFarmWaterComponent>(ELWComponentAccess::ReadOnly);
-		EntityQuery.AddRequirement<FFarmFlowerComponent>(ELWComponentAccess::ReadWrite);
+		EntityQuery.AddRequirement<FHarvestTimerComponent>(EMassFragmentAccess::ReadWrite);
+		EntityQuery.AddRequirement<FFarmWaterComponent>(EMassFragmentAccess::ReadOnly);
+		EntityQuery.AddRequirement<FFarmFlowerComponent>(EMassFragmentAccess::ReadWrite);
 	}
 
-	virtual void Execute(UMassEntitySubsystem& EntitySubsystem, FLWComponentSystemExecutionContext& Context) override
+	virtual void Execute(UMassEntitySubsystem& EntitySubsystem, FMassExecutionContext& Context) override
 	{
 		QUICK_SCOPE_CYCLE_COUNTER(UFarmHarvestTimerSystem_Flowers_Run);
 		
-		EntityQuery.ForEachEntityChunk(EntitySubsystem, Context, [this](FLWComponentSystemExecutionContext& Context) {
+		EntityQuery.ForEachEntityChunk(EntitySubsystem, Context, [this](FMassExecutionContext& Context) {
 			const int32 NumEntities = Context.GetEntitiesNum();
 			const float WellWateredThreshold = 0.25f;
 			TArrayView<FHarvestTimerComponent> TimerList = Context.GetMutableComponentView<FHarvestTimerComponent>();
@@ -187,16 +187,16 @@ class UFarmHarvestTimerSystem_Crops : public UComponentUpdateSystem
 public:
 	virtual void ConfigureQueries() override
 	{
-		EntityQuery.AddRequirement<FHarvestTimerComponent>(ELWComponentAccess::ReadWrite);
-		EntityQuery.AddRequirement<FFarmWaterComponent>(ELWComponentAccess::ReadOnly);
-		EntityQuery.AddRequirement<FFarmCropComponent>(ELWComponentAccess::ReadOnly);
+		EntityQuery.AddRequirement<FHarvestTimerComponent>(EMassFragmentAccess::ReadWrite);
+		EntityQuery.AddRequirement<FFarmWaterComponent>(EMassFragmentAccess::ReadOnly);
+		EntityQuery.AddRequirement<FFarmCropComponent>(EMassFragmentAccess::ReadOnly);
 	}
 
-	virtual void Execute(UMassEntitySubsystem& EntitySubsystem, FLWComponentSystemExecutionContext& Context) override
+	virtual void Execute(UMassEntitySubsystem& EntitySubsystem, FMassExecutionContext& Context) override
 	{
 		QUICK_SCOPE_CYCLE_COUNTER(UFarmHarvestTimerSystem_Crops_Run);
 		
-		EntityQuery.ForEachEntityChunk(EntitySubsystem, Context, [this](FLWComponentSystemExecutionContext& Context) {
+		EntityQuery.ForEachEntityChunk(EntitySubsystem, Context, [this](FMassExecutionContext& Context) {
 			
 			const int32 NumEntities = Context.GetEntitiesNum();
 			const float WellWateredThreshold = 0.25f;
@@ -222,16 +222,16 @@ class UFarmHarvestTimerExpired : public UComponentUpdateSystem
 public:
 	virtual void ConfigureQueries() override
 	{
-		EntityQuery.AddRequirement<FHarvestTimerComponent>(ELWComponentAccess::ReadOnly);
-		EntityQuery.AddTagRequirement<FFarmJustBecameReadyToHarvest>(ELWComponentPresence::None);
-		EntityQuery.AddTagRequirement<FFarmReadyToHarvest>(ELWComponentPresence::None);
+		EntityQuery.AddRequirement<FHarvestTimerComponent>(EMassFragmentAccess::ReadOnly);
+		EntityQuery.AddTagRequirement<FFarmJustBecameReadyToHarvest>(EMassFragmentPresence::None);
+		EntityQuery.AddTagRequirement<FFarmReadyToHarvest>(EMassFragmentPresence::None);
 	}
 
-	virtual void Execute(UMassEntitySubsystem& EntitySubsystem, FLWComponentSystemExecutionContext& Context) override
+	virtual void Execute(UMassEntitySubsystem& EntitySubsystem, FMassExecutionContext& Context) override
 	{
 		QUICK_SCOPE_CYCLE_COUNTER(UFarmHarvestTimerExpired_Run);
 		
-		EntityQuery.ForEachEntityChunk(EntitySubsystem, Context, [this](FLWComponentSystemExecutionContext& Context) {
+		EntityQuery.ForEachEntityChunk(EntitySubsystem, Context, [this](FMassExecutionContext& Context) {
 			const int32 NumEntities = Context.GetEntitiesNum();
 			TConstArrayView<FHarvestTimerComponent> TimerList = Context.GetComponentView<FHarvestTimerComponent>();
 
@@ -264,12 +264,12 @@ public:
 public:
 	virtual void ConfigureQueries() override
 	{
-		EntityQuery.AddRequirement<FFarmGridCellData>(ELWComponentAccess::ReadOnly);
-		EntityQuery.AddRequirement<FFarmVisualComponent>(ELWComponentAccess::ReadWrite);
-		EntityQuery.AddTagRequirement<FFarmJustBecameReadyToHarvest>(ELWComponentPresence::All);
+		EntityQuery.AddRequirement<FFarmGridCellData>(EMassFragmentAccess::ReadOnly);
+		EntityQuery.AddRequirement<FFarmVisualComponent>(EMassFragmentAccess::ReadWrite);
+		EntityQuery.AddTagRequirement<FFarmJustBecameReadyToHarvest>(EMassFragmentPresence::All);
 	}
 
-	virtual void Execute(UMassEntitySubsystem& EntitySubsystem, FLWComponentSystemExecutionContext& Context) override;
+	virtual void Execute(UMassEntitySubsystem& EntitySubsystem, FMassExecutionContext& Context) override;
 };
 
 //////////////////////////////////////////////////////////////////////
@@ -293,7 +293,7 @@ public:
 	UPROPERTY(EditAnywhere, Category=Farm)
 	float HarvestIconScale = 0.3f;
 	
-	TArray<FLWEntity> PlantedSquares;
+	TArray<FMassEntityHandle> PlantedSquares;
 
 	
 	UPROPERTY(EditDefaultsOnly, Category=Farm)
diff --git a/Engine/Plugins/Runtime/MassEntity/Source/MassEntityTestSuite/Public/MassEntityTestTypes.h b/Engine/Plugins/Runtime/MassEntity/Source/MassEntityTestSuite/Public/MassEntityTestTypes.h
--- a/Engine/Plugins/Runtime/MassEntity/Source/MassEntityTestSuite/Public/MassEntityTestTypes.h
+++ b/Engine/Plugins/Runtime/MassEntity/Source/MassEntityTestSuite/Public/MassEntityTestTypes.h
@@ -13,21 +13,21 @@ class UWorld;
 class UMassEntitySubsystem;
 
 USTRUCT()
-struct FTestFragment_Float : public FLWComponentData
+struct FTestFragment_Float : public FMassFragment
 {
 	GENERATED_BODY()
 	float Value = 0;
 };
 
 USTRUCT()
-struct FTestFragment_Int : public FLWComponentData
+struct FTestFragment_Int : public FMassFragment
 {
 	GENERATED_BODY()
 	int32 Value = 0;
 };
 
 USTRUCT()
-struct FTestFragment_Bool : public FLWComponentData
+struct FTestFragment_Bool : public FMassFragment
 {
 	GENERATED_BODY()
 	bool bValue = false;
@@ -35,44 +35,44 @@ struct FTestFragment_Bool : public FLWComponentData
 
 /** @todo rename to FTestTag */
 USTRUCT()
-struct FTestFragment_Tag : public FComponentTag
+struct FTestFragment_Tag : public FMassTag
 {
 	GENERATED_BODY()
 };
 
 USTRUCT()
-struct FTestTag_A : public FComponentTag
+struct FTestTag_A : public FMassTag
 {
 	GENERATED_BODY()
 };
 
 USTRUCT()
-struct FTestTag_B : public FComponentTag
+struct FTestTag_B : public FMassTag
 {
 	GENERATED_BODY()
 };
 
 USTRUCT()
-struct FTestTag_C : public FComponentTag
+struct FTestTag_C : public FMassTag
 {
 	GENERATED_BODY()
 };
 
 USTRUCT()
-struct FTestTag_D : public FComponentTag
+struct FTestTag_D : public FMassTag
 {
 	GENERATED_BODY()
 };
 
 
 UCLASS()
-class UPipeTestProcessorBase : public UPipeProcessor
+class UPipeTestProcessorBase : public UMassProcessor
 {
 	GENERATED_BODY()
 public:
 	UPipeTestProcessorBase();
-	FPipeProcessorExecutionOrder& GetMutableExecutionOrder() { return ExecutionOrder; }
-	virtual void Execute(UMassEntitySubsystem& EntitySubsystem, FLWComponentSystemExecutionContext& Context) override 
+	FMassProcessorExecutionOrder& GetMutableExecutionOrder() { return ExecutionOrder; }
+	virtual void Execute(UMassEntitySubsystem& EntitySubsystem, FMassExecutionContext& Context) override 
 	{
 		ExecutionFunction(EntitySubsystem, Context);
 	}
@@ -81,13 +81,13 @@ public:
 		RequirementsFunction(EntityQuery);
 	}
 
-	TFunction<void(UMassEntitySubsystem& EntitySubsystem, FLWComponentSystemExecutionContext& Context)> ExecutionFunction;
-	TFunction<void(FLWComponentQuery& Query)> RequirementsFunction;
+	TFunction<void(UMassEntitySubsystem& EntitySubsystem, FMassExecutionContext& Context)> ExecutionFunction;
+	TFunction<void(FMassEntityQuery& Query)> RequirementsFunction;
 
-	FLWComponentQuery& TestGetQuery() { return EntityQuery; }
+	FMassEntityQuery& TestGetQuery() { return EntityQuery; }
 
 protected:
-	FLWComponentQuery EntityQuery;
+	FMassEntityQuery EntityQuery;
 };
 
 UCLASS()
diff --git a/Engine/Plugins/Runtime/MassGameplay/Source/MassActors/Private/MassActorSpawnerSubsystem.cpp b/Engine/Plugins/Runtime/MassGameplay/Source/MassActors/Private/MassActorSpawnerSubsystem.cpp
--- a/Engine/Plugins/Runtime/MassGameplay/Source/MassActors/Private/MassActorSpawnerSubsystem.cpp
+++ b/Engine/Plugins/Runtime/MassGameplay/Source/MassActors/Private/MassActorSpawnerSubsystem.cpp
@@ -342,8 +342,8 @@ void UMassActorSpawnerSubsystem::Initialize(FSubsystemCollectionBase& Collection
 	{
 		UMassSimulationSubsystem* SimSystem = UWorld::GetSubsystem<UMassSimulationSubsystem>(World);
 		check(SimSystem);
-		SimSystem->GetOnProcessingPhaseStarted(EPipeProcessingPhase::PrePhysics).AddUObject(this, &UMassActorSpawnerSubsystem::OnPrePhysicsPhaseStarted);
-		SimSystem->GetOnProcessingPhaseFinished(EPipeProcessingPhase::PrePhysics).AddUObject(this, &UMassActorSpawnerSubsystem::OnPrePhysicsPhaseFinished);
+		SimSystem->GetOnProcessingPhaseStarted(EMassProcessingPhase::PrePhysics).AddUObject(this, &UMassActorSpawnerSubsystem::OnPrePhysicsPhaseStarted);
+		SimSystem->GetOnProcessingPhaseFinished(EMassProcessingPhase::PrePhysics).AddUObject(this, &UMassActorSpawnerSubsystem::OnPrePhysicsPhaseFinished);
 	}
 }
 
@@ -355,8 +355,8 @@ void UMassActorSpawnerSubsystem::Deinitialize()
 	{
 		if (const UMassSimulationSubsystem* SimSystem = UWorld::GetSubsystem<UMassSimulationSubsystem>(World))
 		{
-			SimSystem->GetOnProcessingPhaseStarted(EPipeProcessingPhase::PrePhysics).RemoveAll(this);
-			SimSystem->GetOnProcessingPhaseFinished(EPipeProcessingPhase::PrePhysics).RemoveAll(this);
+			SimSystem->GetOnProcessingPhaseStarted(EMassProcessingPhase::PrePhysics).RemoveAll(this);
+			SimSystem->GetOnProcessingPhaseFinished(EMassProcessingPhase::PrePhysics).RemoveAll(this);
 		}
 	}
 }
diff --git a/Engine/Plugins/Runtime/MassGameplay/Source/MassActors/Private/MassActorSubsystem.cpp b/Engine/Plugins/Runtime/MassGameplay/Source/MassActors/Private/MassActorSubsystem.cpp
--- a/Engine/Plugins/Runtime/MassGameplay/Source/MassActors/Private/MassActorSubsystem.cpp
+++ b/Engine/Plugins/Runtime/MassGameplay/Source/MassActors/Private/MassActorSubsystem.cpp
@@ -68,7 +68,7 @@ void UMassActorSubsystem::Initialize(FSubsystemCollectionBase& Collection)
 FMassHandle UMassActorSubsystem::GetHandleFromActor(const TObjectKey<const AActor> Actor)
 {
 	UE_MT_SCOPED_READ_ACCESS(ActorHandleMapDetector);
-	FLWEntity* LWEntity = ActorHandleMap.Find(Actor);
+	FMassEntityHandle* LWEntity = ActorHandleMap.Find(Actor);
 	if (!LWEntity)
 	{
 		return FMassHandle::InvalidHandle;
@@ -104,7 +104,7 @@ void UMassActorSubsystem::DisconnectActor(const TObjectKey<const AActor> Actor,
 		return;
 	}
 
-	FLWEntity LWEntity;
+	FMassEntityHandle LWEntity;
 	{
 		UE_MT_SCOPED_WRITE_ACCESS(ActorHandleMapDetector);
 		// We're assuming the Handle does match Actor, so we're RemoveAndCopyValue. If if doesn't we'll add it back.
diff --git a/Engine/Plugins/Runtime/MassGameplay/Source/MassActors/Private/MassAgentComponent.cpp b/Engine/Plugins/Runtime/MassGameplay/Source/MassActors/Private/MassAgentComponent.cpp
--- a/Engine/Plugins/Runtime/MassGameplay/Source/MassActors/Private/MassAgentComponent.cpp
+++ b/Engine/Plugins/Runtime/MassGameplay/Source/MassActors/Private/MassAgentComponent.cpp
@@ -192,7 +192,7 @@ void UMassAgentComponent::SetEntityHandleInternal(FMassHandle NewHandle)
 	{
 		if (IsNetSimulating())
 		{
-			const FEntityView EntityView(*EntitySubsystem, AgentHandle.GetLWEntity());
+			const FMassEntityView EntityView(*EntitySubsystem, AgentHandle.GetLWEntity());
 
 			// @todo Find a way to add these initialization into either translator initializer or adding new fragments
 			// Make sure to fetch the fragment after any release, as that action can move the entity around into new archetype and 
@@ -281,7 +281,7 @@ void UMassAgentComponent::ClearEntityHandleInternal()
 	{
 		if (IsNetSimulating())
 		{
-			const FEntityView EntityView(*EntitySubsystem, AgentHandle.GetLWEntity());
+			const FMassEntityView EntityView(*EntitySubsystem, AgentHandle.GetLWEntity());
 			if (FDataFragment_Actor* ActorInfo = EntityView.GetComponentDataPtr<FDataFragment_Actor>())
 			{
 				checkf(!ActorInfo->IsValid() || ActorInfo->Get() == GetOwner(), TEXT("Expecting actor pointer to be the component owner"));
diff --git a/Engine/Plugins/Runtime/MassGameplay/Source/MassActors/Private/MassAgentSubsystem.cpp b/Engine/Plugins/Runtime/MassGameplay/Source/MassActors/Private/MassAgentSubsystem.cpp
--- a/Engine/Plugins/Runtime/MassGameplay/Source/MassActors/Private/MassAgentSubsystem.cpp
+++ b/Engine/Plugins/Runtime/MassGameplay/Source/MassActors/Private/MassAgentSubsystem.cpp
@@ -19,7 +19,7 @@
 namespace FMassAgentSubsystemHelper
 {
 
-inline void InitializeAgentComponentFragments(const UMassAgentComponent& AgentComp, FEntityView& EntityView, const EMassTranslationDirection Direction, TConstArrayView<FMassEntityTemplate::FObjectFragmentInitializerFunction> ObjectFragmentInitializers)
+inline void InitializeAgentComponentFragments(const UMassAgentComponent& AgentComp, FMassEntityView& EntityView, const EMassTranslationDirection Direction, TConstArrayView<FMassEntityTemplate::FObjectFragmentInitializerFunction> ObjectFragmentInitializers)
 {
 	AActor* Owner = AgentComp.GetOwner();
 	check(Owner);
@@ -50,7 +50,7 @@ void UMassAgentSubsystem::Initialize(FSubsystemCollectionBase& Collection)
 	SpawnerSystem = UWorld::GetSubsystem<UMassSpawnerSubsystem>(World);
 
 	SimulationSystem = UWorld::GetSubsystem<UMassSimulationSubsystem>(World);
-	SimulationSystem->GetOnProcessingPhaseStarted(EPipeProcessingPhase::PrePhysics).AddUObject(this, &UMassAgentSubsystem::OnProcessingPhaseStarted, EPipeProcessingPhase::PrePhysics);
+	SimulationSystem->GetOnProcessingPhaseStarted(EMassProcessingPhase::PrePhysics).AddUObject(this, &UMassAgentSubsystem::OnProcessingPhaseStarted, EMassProcessingPhase::PrePhysics);
 
 #if UE_REPLICATION_COMPILE_CLIENT_CODE
 	ReplicationManager = UWorld::GetSubsystem<UMassReplicationManager>(World);
@@ -136,7 +136,7 @@ void UMassAgentSubsystem::UpdateAgentComponent(const UMassAgentComponent& AgentC
 		return;
 	}
 
-	const FLWEntity Entity = AgentComp.GetEntityHandle().GetLWEntity();
+	const FMassEntityHandle Entity = AgentComp.GetEntityHandle().GetLWEntity();
 	const FArchetypeHandle CurrentArchetypeHandle = EntitySystem->GetArchetypeForEntity(Entity);
 	if (CurrentArchetypeHandle == EntityTemplate->GetArchetype())
 	{
@@ -178,7 +178,7 @@ void UMassAgentSubsystem::UnregisterAgentComponent(UMassAgentComponent& AgentCom
 		{
 			if (AgentComp.GetPuppetSpecificAddition().IsEmpty() == false)
 			{
-				const FLWEntity Entity = AgentComp.GetEntityHandle().GetLWEntity();
+				const FMassEntityHandle Entity = AgentComp.GetEntityHandle().GetLWEntity();
 
 			    // remove fragments that have been added for the puppet agent
 			    if (EntitySystem->IsProcessing())
@@ -207,14 +207,14 @@ void UMassAgentSubsystem::UnregisterAgentComponent(UMassAgentComponent& AgentCom
 				EntityTemplate = &EntityConfig.GetEntityTemplateChecked(*AgentActor, AgentComp);
 			}
 
-			FLWEntity Entity = AgentComp.GetEntityHandle().GetLWEntity();
+			FMassEntityHandle Entity = AgentComp.GetEntityHandle().GetLWEntity();
 			// Clearing the entity before it become invalid as the clear contains notifications
 			AgentComp.ClearEntityHandle();
 
 			// Destroy the entity
 			if (ensure(EntityTemplate))
 			{
-				SpawnerSystem->DestroyEntities(EntityTemplate->GetTemplateID(), TArrayView<FLWEntity>(&Entity, 1));
+				SpawnerSystem->DestroyEntities(EntityTemplate->GetTemplateID(), TArrayView<FMassEntityHandle>(&Entity, 1));
 			}
 		}
 		else if (AgentComp.IsEntityPendingCreation())
@@ -264,11 +264,11 @@ void UMassAgentSubsystem::MakePuppet(UMassAgentComponent& AgentComp)
 	AgentComp.PuppetInitializationPending();
 }
 
-void UMassAgentSubsystem::OnProcessingPhaseStarted(const float DeltaSeconds, const EPipeProcessingPhase Phase)
+void UMassAgentSubsystem::OnProcessingPhaseStarted(const float DeltaSeconds, const EMassProcessingPhase Phase)
 {
 	switch (Phase)
 	{
-	case EPipeProcessingPhase::PrePhysics:
+	case EMassProcessingPhase::PrePhysics:
 		// initialize what needs initialization
 		if (PendingAgentEntities.Num() > 0 || PendingPuppets.Num() > 0)
 		{
@@ -303,7 +303,7 @@ void UMassAgentSubsystem::HandlePendingInitialization()
 			continue;
 		}
 		
-		TArray<FLWEntity> Entities;
+		TArray<FMassEntityHandle> Entities;
 		SpawnerSystem->SpawnEntities(*EntityTemplate, NewEntityCount, Entities);
 		check(Entities.Num() == NewEntityCount);
 
@@ -313,7 +313,7 @@ void UMassAgentSubsystem::HandlePendingInitialization()
 
 			for (int AgentIndex = 0; AgentIndex < Entities.Num(); ++AgentIndex)
 			{
-				FEntityView EntityView(EntityTemplate->GetArchetype(), Entities[AgentIndex]);
+				FMassEntityView EntityView(EntityTemplate->GetArchetype(), Entities[AgentIndex]);
 				FMassAgentSubsystemHelper::InitializeAgentComponentFragments(*AgentComponents[AgentIndex], EntityView, EMassTranslationDirection::ActorToMass, ObjectFragmentInitializers);
 			}
 		}
@@ -336,26 +336,26 @@ void UMassAgentSubsystem::HandlePendingInitialization()
 			continue;
 		}
 
-		const FLWCompositionDescriptor TemplateDescriptor = EntityTemplate->GetCompositionDescriptor();
+		const FMassCompositionDescriptor TemplateDescriptor = EntityTemplate->GetCompositionDescriptor();
 
 		TArray<UMassAgentComponent*>& AgentComponents = Data.Get<1>().AgentComponents;
 
 		for (UMassAgentComponent* AgentComp : AgentComponents)
 		{
-			const FLWEntity PuppetEntity = AgentComp->GetEntityHandle().GetLWEntity();
+			const FMassEntityHandle PuppetEntity = AgentComp->GetEntityHandle().GetLWEntity();
 			if (!ensureMsgf(PuppetEntity.IsSet(), TEXT("Trying to initialize puppet's fragments while the pupped doesn't have a corresponding Entity identifier set. This should not happen.")))
 			{
 				continue;
 			}
 
-			FLWCompositionDescriptor& PuppetDescriptor = AgentComp->GetMutablePuppetSpecificAddition();
+			FMassCompositionDescriptor& PuppetDescriptor = AgentComp->GetMutablePuppetSpecificAddition();
 			PuppetDescriptor = TemplateDescriptor;
 			EntitySystem->AddCompositionToEntity_GetDelta(PuppetEntity, PuppetDescriptor);
 			
 			if (EntityTemplate->GetObjectFragmentInitializers().Num())
 			{
 				const FArchetypeHandle ArchetypeHandle = EntitySystem->GetArchetypeForEntity(PuppetEntity);
-				FEntityView EntityView(ArchetypeHandle, PuppetEntity);
+				FMassEntityView EntityView(ArchetypeHandle, PuppetEntity);
 				FMassAgentSubsystemHelper::InitializeAgentComponentFragments(*AgentComp, EntityView, EMassTranslationDirection::MassToActor, EntityTemplate->GetObjectFragmentInitializers());
 			}
 
@@ -376,7 +376,7 @@ void UMassAgentSubsystem::NotifyMassAgentComponentReplicated(UMassAgentComponent
 		check(AgentComp.GetNetID().IsValid());
 		check(!ReplicatedAgentComponents.Find(AgentComp.GetNetID()));
 		ReplicatedAgentComponents.Add(AgentComp.GetNetID(), &AgentComp);
-		const FLWEntity Entity = ReplicationManager->FindEntity(AgentComp.GetNetID());
+		const FMassEntityHandle Entity = ReplicationManager->FindEntity(AgentComp.GetNetID());
 
 		// If not found, the NotifyMassAgentAddedToReplication will link it later once replicated.
 		if (Entity.IsSet())
@@ -402,7 +402,7 @@ void UMassAgentSubsystem::NotifyMassAgentComponentEntityDetaching(const UMassAge
 	OnMassAgentComponentEntityDetaching.Broadcast(AgentComp);
 }
 
-void UMassAgentSubsystem::OnMassAgentAddedToReplication(FMassNetworkID NetID, FLWEntity Entity)
+void UMassAgentSubsystem::OnMassAgentAddedToReplication(FMassNetworkID NetID, FMassEntityHandle Entity)
 {
 #if UE_REPLICATION_COMPILE_CLIENT_CODE
 	UWorld* World = GetWorld();
@@ -417,7 +417,7 @@ void UMassAgentSubsystem::OnMassAgentAddedToReplication(FMassNetworkID NetID, FL
 #endif // UE_REPLICATION_COMPILE_CLIENT_CODE
 }
 
-void UMassAgentSubsystem::OnMassAgentRemovedFromReplication(FMassNetworkID NetID, FLWEntity Entity)
+void UMassAgentSubsystem::OnMassAgentRemovedFromReplication(FMassNetworkID NetID, FMassEntityHandle Entity)
 {
 #if UE_REPLICATION_COMPILE_CLIENT_CODE
 	UWorld* World = GetWorld();
diff --git a/Engine/Plugins/Runtime/MassGameplay/Source/MassActors/Private/MassAgentTraits.cpp b/Engine/Plugins/Runtime/MassGameplay/Source/MassActors/Private/MassAgentTraits.cpp
--- a/Engine/Plugins/Runtime/MassGameplay/Source/MassActors/Private/MassAgentTraits.cpp
+++ b/Engine/Plugins/Runtime/MassGameplay/Source/MassActors/Private/MassAgentTraits.cpp
@@ -49,7 +49,7 @@ void UMassAgentCapsuleCollisionSyncTrait::BuildTemplate(FMassEntityTemplateBuild
 		BuildContext.AddFragment<FDataFragment_Transform>();
 	}
 	
-	BuildContext.GetMutableObjectFragmentInitializers().Add([=](UObject& Owner, FEntityView& EntityView, const EMassTranslationDirection CurrentDirection)
+	BuildContext.GetMutableObjectFragmentInitializers().Add([=](UObject& Owner, FMassEntityView& EntityView, const EMassTranslationDirection CurrentDirection)
 		{
 			if (UCapsuleComponent* CapsuleComponent = FMassAgentTraitsHelper::AsComponent<UCapsuleComponent>(Owner))
 			{
@@ -89,7 +89,7 @@ void UMassAgentMovementSyncTrait::BuildTemplate(FMassEntityTemplateBuildContext&
 	BuildContext.AddFragment<FDataFragment_CharacterMovementComponentWrapper>();
 	BuildContext.AddFragment<FMassVelocityFragment>();
 	
-	BuildContext.GetMutableObjectFragmentInitializers().Add([=](UObject& Owner, FEntityView& EntityView, const EMassTranslationDirection CurrentDirection)
+	BuildContext.GetMutableObjectFragmentInitializers().Add([=](UObject& Owner, FMassEntityView& EntityView, const EMassTranslationDirection CurrentDirection)
 		{
 			if (UCharacterMovementComponent* MovementComp = FMassAgentTraitsHelper::AsComponent<UCharacterMovementComponent>(Owner))
 			{
@@ -151,7 +151,7 @@ void UMassAgentFeetLocationSyncTrait::BuildTemplate(FMassEntityTemplateBuildCont
 	BuildContext.AddFragment<FMassSceneComponentWrapperFragment>();
 	BuildContext.AddFragment<FDataFragment_Transform>();
 
-	BuildContext.GetMutableObjectFragmentInitializers().Add([=](UObject& Owner, FEntityView& EntityView, const EMassTranslationDirection CurrentDirection)
+	BuildContext.GetMutableObjectFragmentInitializers().Add([=](UObject& Owner, FMassEntityView& EntityView, const EMassTranslationDirection CurrentDirection)
 		{
 			AActor* AsActor = Cast<AActor>(&Owner);
 			if (AsActor && AsActor->GetRootComponent())
