Renamed Mass' files to match the stored types.
Additional type renames:
FArchetypeChunkIterator -> FMassArchetypeChunkIterator

diff --git a/Engine/Plugins/AI/MassAI/Source/MassAIBehavior/Private/MassLookAtProcessors.cpp b/Engine/Plugins/AI/MassAI/Source/MassAIBehavior/Private/MassLookAtProcessors.cpp
--- a/Engine/Plugins/AI/MassAI/Source/MassAIBehavior/Private/MassLookAtProcessors.cpp
+++ b/Engine/Plugins/AI/MassAI/Source/MassAIBehavior/Private/MassLookAtProcessors.cpp
@@ -1,7 +1,7 @@
 // Copyright Epic Games, Inc. All Rights Reserved.
 
 #include "MassLookAtProcessors.h"
-#include "EntityView.h"
+#include "MassEntityView.h"
 #include "MassAIBehaviorTypes.h"
 #include "MassCommonFragments.h"
 #include "MassLookAtFragments.h"
diff --git a/Engine/Plugins/AI/MassAI/Source/MassAIBehavior/Private/MassStateTreeProcessors.cpp b/Engine/Plugins/AI/MassAI/Source/MassAIBehavior/Private/MassStateTreeProcessors.cpp
--- a/Engine/Plugins/AI/MassAI/Source/MassAIBehavior/Private/MassStateTreeProcessors.cpp
+++ b/Engine/Plugins/AI/MassAI/Source/MassAIBehavior/Private/MassStateTreeProcessors.cpp
@@ -3,7 +3,7 @@
 #include "MassStateTreeProcessors.h"
 #include "StateTree.h"
 #include "MassStateTreeExecutionContext.h"
-#include "EntityView.h"
+#include "MassEntityView.h"
 #include "MassAIMovementTypes.h"
 #include "MassComponentHitTypes.h"
 #include "MassSmartObjectTypes.h"
diff --git a/Engine/Plugins/AI/MassAI/Source/MassAIBehavior/Private/MassStateTreeSchema.cpp b/Engine/Plugins/AI/MassAI/Source/MassAIBehavior/Private/MassStateTreeSchema.cpp
--- a/Engine/Plugins/AI/MassAI/Source/MassAIBehavior/Private/MassStateTreeSchema.cpp
+++ b/Engine/Plugins/AI/MassAI/Source/MassAIBehavior/Private/MassStateTreeSchema.cpp
@@ -1,7 +1,7 @@
 ﻿// Copyright Epic Games, Inc. All Rights Reserved.
 
 #include "MassStateTreeSchema.h"
-#include "LWComponentTypes.h"
+#include "MassEntityTypes.h"
 #include "MassStateTreeTypes.h"
 #include "Subsystems/WorldSubsystem.h"
 
diff --git a/Engine/Plugins/AI/MassAI/Source/MassAIBehavior/Public/MassComponentHitSubsystem.h b/Engine/Plugins/AI/MassAI/Source/MassAIBehavior/Public/MassComponentHitSubsystem.h
--- a/Engine/Plugins/AI/MassAI/Source/MassAIBehavior/Public/MassComponentHitSubsystem.h
+++ b/Engine/Plugins/AI/MassAI/Source/MassAIBehavior/Public/MassComponentHitSubsystem.h
@@ -2,7 +2,7 @@
 
 #pragma once
 
-#include "LWComponentTypes.h"
+#include "MassEntityTypes.h"
 #include "MassComponentHitTypes.h"
 #include "Subsystems/WorldSubsystem.h"
 
diff --git a/Engine/Plugins/AI/MassAI/Source/MassAIBehavior/Public/MassComponentHitTypes.h b/Engine/Plugins/AI/MassAI/Source/MassAIBehavior/Public/MassComponentHitTypes.h
--- a/Engine/Plugins/AI/MassAI/Source/MassAIBehavior/Public/MassComponentHitTypes.h
+++ b/Engine/Plugins/AI/MassAI/Source/MassAIBehavior/Public/MassComponentHitTypes.h
@@ -3,7 +3,7 @@
 #pragma once
 
 #include "Engine/EngineTypes.h"
-#include "LWComponentTypes.h"
+#include "MassEntityTypes.h"
 #include "StateTreeTypes.h"
 
 #include "MassComponentHitTypes.generated.h"
diff --git a/Engine/Plugins/AI/MassAI/Source/MassAIBehavior/Public/MassLookAtFragments.h b/Engine/Plugins/AI/MassAI/Source/MassAIBehavior/Public/MassLookAtFragments.h
--- a/Engine/Plugins/AI/MassAI/Source/MassAIBehavior/Public/MassLookAtFragments.h
+++ b/Engine/Plugins/AI/MassAI/Source/MassAIBehavior/Public/MassLookAtFragments.h
@@ -2,7 +2,7 @@
 
 #pragma once
 
-#include "LWComponentTypes.h"
+#include "MassEntityTypes.h"
 #include "MassCommonTypes.h"
 #include "ZoneGraphTypes.h"
 #include "Containers/StaticArray.h"
diff --git a/Engine/Plugins/AI/MassAI/Source/MassAIBehavior/Public/MassStateTreeExecutionContext.h b/Engine/Plugins/AI/MassAI/Source/MassAIBehavior/Public/MassStateTreeExecutionContext.h
--- a/Engine/Plugins/AI/MassAI/Source/MassAIBehavior/Public/MassStateTreeExecutionContext.h
+++ b/Engine/Plugins/AI/MassAI/Source/MassAIBehavior/Public/MassStateTreeExecutionContext.h
@@ -2,7 +2,7 @@
 
 #pragma once
 
-#include "LWComponentTypes.h"
+#include "MassEntityTypes.h"
 #include "StateTreeExecutionContext.h"
 #include "MassStateTreeExecutionContext.generated.h"
 
diff --git a/Engine/Plugins/AI/MassAI/Source/MassAIBehavior/Public/MassStateTreeFragments.h b/Engine/Plugins/AI/MassAI/Source/MassAIBehavior/Public/MassStateTreeFragments.h
--- a/Engine/Plugins/AI/MassAI/Source/MassAIBehavior/Public/MassStateTreeFragments.h
+++ b/Engine/Plugins/AI/MassAI/Source/MassAIBehavior/Public/MassStateTreeFragments.h
@@ -2,7 +2,7 @@
 
 #pragma once
 
-#include "LWComponentTypes.h"
+#include "MassEntityTypes.h"
 #include "MassStateTreeSubsystem.h"
 #include "MassStateTreeFragments.generated.h"
 
diff --git a/Engine/Plugins/AI/MassAI/Source/MassAIBehavior/Public/MassZoneGraphAnnotationFragments.h b/Engine/Plugins/AI/MassAI/Source/MassAIBehavior/Public/MassZoneGraphAnnotationFragments.h
--- a/Engine/Plugins/AI/MassAI/Source/MassAIBehavior/Public/MassZoneGraphAnnotationFragments.h
+++ b/Engine/Plugins/AI/MassAI/Source/MassAIBehavior/Public/MassZoneGraphAnnotationFragments.h
@@ -2,7 +2,7 @@
 
 #pragma once
 
-#include "LWComponentTypes.h"
+#include "MassEntityTypes.h"
 #include "ZoneGraphTypes.h"
 #include "MassZoneGraphAnnotationFragments.generated.h"
 
diff --git a/Engine/Plugins/AI/MassAI/Source/MassAIBehavior/Public/Tasks/MassLookAtTask.h b/Engine/Plugins/AI/MassAI/Source/MassAIBehavior/Public/Tasks/MassLookAtTask.h
--- a/Engine/Plugins/AI/MassAI/Source/MassAIBehavior/Public/Tasks/MassLookAtTask.h
+++ b/Engine/Plugins/AI/MassAI/Source/MassAIBehavior/Public/Tasks/MassLookAtTask.h
@@ -2,7 +2,7 @@
 
 #pragma once
 
-#include "LWComponentTypes.h"
+#include "MassEntityTypes.h"
 #include "MassStateTreeTypes.h"
 #include "MassLookAtFragments.h"
 #include "MassLookAtTask.generated.h"
diff --git a/Engine/Plugins/AI/MassAI/Source/MassAIDebug/Private/GameplayDebuggerCategory_Mass.cpp b/Engine/Plugins/AI/MassAI/Source/MassAIDebug/Private/GameplayDebuggerCategory_Mass.cpp
--- a/Engine/Plugins/AI/MassAI/Source/MassAIDebug/Private/GameplayDebuggerCategory_Mass.cpp
+++ b/Engine/Plugins/AI/MassAI/Source/MassAIDebug/Private/GameplayDebuggerCategory_Mass.cpp
@@ -3,7 +3,7 @@
 #include "GameplayDebuggerCategory_Mass.h"
 
 #if WITH_GAMEPLAY_DEBUGGER && WITH_MASSGAMEPLAY_DEBUG
-#include "EntityView.h"
+#include "MassEntityView.h"
 #include "GameplayDebuggerCategoryReplicator.h"
 #include "GameplayDebuggerPlayerManager.h"
 #include "MassDebuggerSubsystem.h"
diff --git a/Engine/Plugins/AI/MassAI/Source/MassAIMovement/Private/MassAvoidanceProcessors.cpp b/Engine/Plugins/AI/MassAI/Source/MassAIMovement/Private/MassAvoidanceProcessors.cpp
--- a/Engine/Plugins/AI/MassAI/Source/MassAIMovement/Private/MassAvoidanceProcessors.cpp
+++ b/Engine/Plugins/AI/MassAI/Source/MassAIMovement/Private/MassAvoidanceProcessors.cpp
@@ -3,7 +3,7 @@
 #include "MassAvoidanceProcessors.h"
 #include "NavigationSystem.h"
 #include "DrawDebugHelpers.h"
-#include "EntityView.h"
+#include "MassEntityView.h"
 #include "MassAIMovementFragments.h"
 #include "VisualLogger/VisualLogger.h"
 #include "Math/Vector2D.h"
diff --git a/Engine/Plugins/AI/MassAI/Source/MassAIMovement/Private/NavigationProcessor.cpp b/Engine/Plugins/AI/MassAI/Source/MassAIMovement/Private/NavigationProcessor.cpp
--- a/Engine/Plugins/AI/MassAI/Source/MassAIMovement/Private/NavigationProcessor.cpp
+++ b/Engine/Plugins/AI/MassAI/Source/MassAIMovement/Private/NavigationProcessor.cpp
@@ -3,7 +3,7 @@
 #include "NavigationProcessor.h"
 
 #include "MassCommonUtils.h"
-#include "LWCCommandBuffer.h"
+#include "MassCommandBuffer.h"
 #include "MassCommonFragments.h"
 #include "MassAIMovementFragments.h"
 #include "MassMovementTypes.h"
diff --git a/Engine/Plugins/AI/MassAI/Source/MassAIMovement/Public/MassAIMovementTypes.h b/Engine/Plugins/AI/MassAI/Source/MassAIMovement/Public/MassAIMovementTypes.h
--- a/Engine/Plugins/AI/MassAI/Source/MassAIMovement/Public/MassAIMovementTypes.h
+++ b/Engine/Plugins/AI/MassAI/Source/MassAIMovement/Public/MassAIMovementTypes.h
@@ -2,7 +2,7 @@
 
 #pragma once
 
-#include "LWComponentTypes.h"
+#include "MassEntityTypes.h"
 
 #include "ZoneGraphTypes.h"
 #include "MassMovementTypes.h"
diff --git a/Engine/Plugins/AI/MassAI/Source/MassAIMovement/Public/MassZoneGraphMovementFragments.h b/Engine/Plugins/AI/MassAI/Source/MassAIMovement/Public/MassZoneGraphMovementFragments.h
--- a/Engine/Plugins/AI/MassAI/Source/MassAIMovement/Public/MassZoneGraphMovementFragments.h
+++ b/Engine/Plugins/AI/MassAI/Source/MassAIMovement/Public/MassZoneGraphMovementFragments.h
@@ -2,7 +2,7 @@
 
 #pragma once
 
-#include "LWComponentTypes.h"
+#include "MassEntityTypes.h"
 #include "ZoneGraphTypes.h"
 #include "MassCommonTypes.h"
 #include "MassZoneGraphMovementTypes.h"
diff --git a/Engine/Plugins/AI/MassAI/Source/MassAIMovement/Public/MassZoneGraphMovementUtils.h b/Engine/Plugins/AI/MassAI/Source/MassAIMovement/Public/MassZoneGraphMovementUtils.h
--- a/Engine/Plugins/AI/MassAI/Source/MassAIMovement/Public/MassZoneGraphMovementUtils.h
+++ b/Engine/Plugins/AI/MassAI/Source/MassAIMovement/Public/MassZoneGraphMovementUtils.h
@@ -2,7 +2,7 @@
 
 #pragma once
 
-#include "LWComponentTypes.h"
+#include "MassEntityTypes.h"
 
 struct FZoneGraphShortPathRequest;
 struct FMassZoneGraphCachedLaneFragment;
diff --git a/Engine/Plugins/AI/MassAI/Source/MassAIReplication/Private/MassReplicationPathHandlers.cpp b/Engine/Plugins/AI/MassAI/Source/MassAIReplication/Private/MassReplicationPathHandlers.cpp
--- a/Engine/Plugins/AI/MassAI/Source/MassAIReplication/Private/MassReplicationPathHandlers.cpp
+++ b/Engine/Plugins/AI/MassAI/Source/MassAIReplication/Private/MassReplicationPathHandlers.cpp
@@ -1,7 +1,7 @@
 // Copyright Epic Games, Inc. All Rights Reserved.
 
 #include "MassReplicationPathHandlers.h"
-#include "EntityQuery.h"
+#include "MassEntityQuery.h"
 #include "MassCommonFragments.h"
 #include "MassMovementTypes.h"
 #include "MassSimulationSubsystem.h"
diff --git a/Engine/Plugins/AI/MassAI/Source/MassAIReplication/Public/MassReplicationPathHandlers.h b/Engine/Plugins/AI/MassAI/Source/MassAIReplication/Public/MassReplicationPathHandlers.h
--- a/Engine/Plugins/AI/MassAI/Source/MassAIReplication/Public/MassReplicationPathHandlers.h
+++ b/Engine/Plugins/AI/MassAI/Source/MassAIReplication/Public/MassReplicationPathHandlers.h
@@ -3,7 +3,7 @@
 #pragma once
 
 #include "MassReplicationTypes.h"
-#include "EntityView.h"
+#include "MassEntityView.h"
 #include "MassClientBubbleHandler.h"
 #include "MassAIMovementFragments.h"
 #include "MassZoneGraphMovementFragments.h"
diff --git a/Engine/Plugins/AI/MassCrowd/Source/MassCrowd/Private/MassCrowdBubble.cpp b/Engine/Plugins/AI/MassCrowd/Source/MassCrowd/Private/MassCrowdBubble.cpp
--- a/Engine/Plugins/AI/MassCrowd/Source/MassCrowd/Private/MassCrowdBubble.cpp
+++ b/Engine/Plugins/AI/MassCrowd/Source/MassCrowd/Private/MassCrowdBubble.cpp
@@ -4,7 +4,7 @@
 #include "Net/UnrealNetwork.h"
 #include "DrawDebugHelpers.h"
 #include "MassSpawnerTypes.h"
-#include "EntityView.h"
+#include "MassEntityView.h"
 #include "MassCommonFragments.h"
 #include "MassMovementFragments.h"
 
diff --git a/Engine/Plugins/AI/MassCrowd/Source/MassCrowd/Public/MassCrowdBubble.h b/Engine/Plugins/AI/MassCrowd/Source/MassCrowd/Public/MassCrowdBubble.h
--- a/Engine/Plugins/AI/MassCrowd/Source/MassCrowd/Public/MassCrowdBubble.h
+++ b/Engine/Plugins/AI/MassCrowd/Source/MassCrowd/Public/MassCrowdBubble.h
@@ -7,7 +7,7 @@
 #include "MassCrowdReplicatedAgent.h"
 #include "MassClientBubbleHandler.h"
 #include "MassClientBubbleInfoBase.h"
-#include "EntityView.h"
+#include "MassEntityView.h"
 
 #include "MassCrowdBubble.generated.h"
 
diff --git a/Engine/Plugins/AI/MassCrowd/Source/MassCrowd/Public/MassCrowdFragments.h b/Engine/Plugins/AI/MassCrowd/Source/MassCrowd/Public/MassCrowdFragments.h
--- a/Engine/Plugins/AI/MassCrowd/Source/MassCrowd/Public/MassCrowdFragments.h
+++ b/Engine/Plugins/AI/MassCrowd/Source/MassCrowd/Public/MassCrowdFragments.h
@@ -2,7 +2,7 @@
 
 #pragma once
 
-#include "LWComponentTypes.h"
+#include "MassEntityTypes.h"
 #include "ZoneGraphTypes.h"
 #include "Containers/StaticArray.h"
 #include "MassLODManager.h"
diff --git a/Engine/Plugins/Runtime/MassEntity/Source/MassEntity/Private/ArchetypeData.cpp b/Engine/Plugins/Runtime/MassEntity/Source/MassEntity/Private/ArchetypeData.cpp
deleted file mode 100644
diff --git a/Engine/Plugins/Runtime/MassEntity/Source/MassEntity/Private/ArchetypeData.h b/Engine/Plugins/Runtime/MassEntity/Source/MassEntity/Private/ArchetypeData.h
deleted file mode 100644
diff --git a/Engine/Plugins/Runtime/MassEntity/Source/MassEntity/Private/ArchetypeTypes.cpp b/Engine/Plugins/Runtime/MassEntity/Source/MassEntity/Private/ArchetypeTypes.cpp
deleted file mode 100644
diff --git a/Engine/Plugins/Runtime/MassEntity/Source/MassEntity/Private/EntityQuery.cpp b/Engine/Plugins/Runtime/MassEntity/Source/MassEntity/Private/EntityQuery.cpp
deleted file mode 100644
diff --git a/Engine/Plugins/Runtime/MassEntity/Source/MassEntity/Private/EntityView.cpp b/Engine/Plugins/Runtime/MassEntity/Source/MassEntity/Private/EntityView.cpp
deleted file mode 100644
diff --git a/Engine/Plugins/Runtime/MassEntity/Source/MassEntity/Private/LWCCommandBuffer.cpp b/Engine/Plugins/Runtime/MassEntity/Source/MassEntity/Private/LWCCommandBuffer.cpp
deleted file mode 100644
diff --git a/Engine/Plugins/Runtime/MassEntity/Source/MassEntity/Private/MassArchetypeData.cpp b/Engine/Plugins/Runtime/MassEntity/Source/MassEntity/Private/MassArchetypeData.cpp
new file mode 100644
--- /dev/null
+++ b/Engine/Plugins/Runtime/MassEntity/Source/MassEntity/Private/MassArchetypeData.cpp
@@ -0,0 +1,887 @@
+// Copyright Epic Games, Inc. All Rights Reserved.
+
+#include "MassArchetypeData.h"
+#include "MassEntityTypes.h"
+#include "Misc/StringBuilder.h"
+
+
+//////////////////////////////////////////////////////////////////////
+// FMassArchetypeData
+namespace UE { namespace AggregateTicking 
+{
+	constexpr static bool bBitwiseRelocateComponents = true;
+}} // UE::AggregateTicking
+
+void FMassArchetypeData::ForEachComponentType(TFunction< void(const UScriptStruct* /*Component*/)> Function) const
+{
+	for (const FMassArchetypeFragmentConfig& ComponentData : ComponentConfigs)
+	{
+		Function(ComponentData.ComponentType);
+	}
+}
+
+bool FMassArchetypeData::HasComponentType(const UScriptStruct* ComponentType) const
+{
+	return (ComponentType && CompositionDescriptor.Components.Contains(*ComponentType));
+}
+
+bool FMassArchetypeData::IsEquivalent(TConstArrayView<const UScriptStruct*> OtherList) const
+{
+	if (OtherList.Num() != ComponentConfigs.Num())
+	{
+		return false;
+	}
+
+	for (int32 Index = 0; Index < ComponentConfigs.Num(); ++Index)
+	{
+		if (OtherList[Index] != ComponentConfigs[Index].ComponentType)
+		{
+			return false;
+		}
+	}
+
+	return true;
+}
+
+void FMassArchetypeData::Initialize(const FMassFragmentBitSet& Components, const FMassTagBitSet& Tags, const FMassChunkFragmentBitSet& ChunkComponents)
+{
+	TArray<const UScriptStruct*, TInlineAllocator<16>> SortedComponentList;
+	Components.ExportTypes(SortedComponentList);
+
+	SortedComponentList.Sort(FMassSorterOperator<UScriptStruct>());
+
+	// Figure out how many bytes all of the individual components (and metadata) will cost per entity
+	int32 ComponentSizeTallyBytes = 0;
+
+	// Alignment padding computation is currently very conservative and over-estimated.
+	int32 AlignmentPadding = 0;
+	{
+		// Save room for the 'metadata' (entity array)
+		ComponentSizeTallyBytes += sizeof(FMassEntityHandle);
+
+		// Tally up the component sizes and place them in the index map
+		ComponentConfigs.AddDefaulted(SortedComponentList.Num());
+		ComponentIndexMap.Reserve(SortedComponentList.Num());
+
+		for (int32 ComponentIndex = 0; ComponentIndex < SortedComponentList.Num(); ++ComponentIndex)
+		{
+			const UScriptStruct* ComponentType = SortedComponentList[ComponentIndex];
+			checkSlow(ComponentType);
+			ComponentConfigs[ComponentIndex].ComponentType = ComponentType;
+			
+			AlignmentPadding += ComponentType->GetMinAlignment();
+			ComponentSizeTallyBytes += ComponentType->GetStructureSize();
+
+			ComponentIndexMap.Add(ComponentType, ComponentIndex);
+			CompositionDescriptor.Components.Add(*ComponentType);
+		}
+		checkSlow(Components.IsEquivalent(CompositionDescriptor.Components));
+	}
+
+	CompositionDescriptor.Tags = Tags;
+	CompositionDescriptor.ChunkComponents = ChunkComponents;
+	// create the chunk component instances
+	TArray<const UScriptStruct*, TInlineAllocator<16>> ChunkComponentList;
+	ChunkComponents.ExportTypes(ChunkComponentList);
+	ChunkComponentList.Sort(FMassSorterOperator<UScriptStruct>());
+	for (const UScriptStruct* ChunkComponentType : ChunkComponentList)
+	{
+		check(ChunkComponentType);
+		ChunkComponentTemplates.Add(FInstancedStruct(ChunkComponentType));
+	}
+
+	TotalBytesPerEntity = ComponentSizeTallyBytes;
+	int32 ChunkAvailableSize = GetChunkAllocSize() - AlignmentPadding;
+	check(TotalBytesPerEntity <= ChunkAvailableSize);
+
+	NumEntitiesPerChunk = ChunkAvailableSize / TotalBytesPerEntity;
+
+	// Set up the offsets for each component into the chunk data
+	EntityListOffsetWithinChunk = 0;
+	int32 CurrentOffset = NumEntitiesPerChunk * sizeof(FMassEntityHandle);
+	for (FMassArchetypeFragmentConfig& ComponentData : ComponentConfigs)
+	{
+		CurrentOffset = Align(CurrentOffset, ComponentData.ComponentType->GetMinAlignment());
+		ComponentData.ArrayOffsetWithinChunk = CurrentOffset;
+		const int32 SizeOfThisComponentArray = NumEntitiesPerChunk * ComponentData.ComponentType->GetStructureSize();
+		CurrentOffset += SizeOfThisComponentArray;
+	}
+}
+
+void FMassArchetypeData::InitializeWithSibling(const FMassArchetypeData& SiblingArchetype, const FMassTagBitSet& OverrideTags)
+{
+	checkf(IsInitialized() == false, TEXT("Trying to %s but this archetype has already been initialized"));
+
+	ComponentConfigs = SiblingArchetype.ComponentConfigs;
+	ComponentIndexMap = SiblingArchetype.ComponentIndexMap;
+	CompositionDescriptor.Components = SiblingArchetype.CompositionDescriptor.Components;
+	CompositionDescriptor.Tags = OverrideTags;
+	CompositionDescriptor.ChunkComponents = SiblingArchetype.CompositionDescriptor.ChunkComponents;
+	ChunkComponentTemplates = SiblingArchetype.ChunkComponentTemplates;
+
+	TotalBytesPerEntity = SiblingArchetype.TotalBytesPerEntity;
+	NumEntitiesPerChunk = SiblingArchetype.NumEntitiesPerChunk;
+
+	// Set up the offsets for each component into the chunk data
+	EntityListOffsetWithinChunk = 0;
+}
+
+void FMassArchetypeData::AddEntity(FMassEntityHandle Entity)
+{
+	AddEntityInternal(Entity, true/*bInitializeComponents*/);
+}
+
+int32 FMassArchetypeData::AddEntityInternal(FMassEntityHandle Entity, const bool bInitializeComponents)
+{
+	int32 IndexWithinChunk = 0;
+	int32 AbsoluteIndex = 0;
+	int32 ChunkIndex = 0;
+	int32 EmptyChunkIndex = INDEX_NONE;
+	int32 EmptyAbsoluteIndex = INDEX_NONE;
+
+	FMassArchetypeChunk* DestinationChunk = nullptr;
+	// Check chunks for a free spot (trying to reuse the earlier ones first so later ones might get freed up) 
+	//@TODO: This could be accelerated to include a cached index to the first chunk with free spots or similar
+	for (FMassArchetypeChunk& Chunk : Chunks)
+	{
+		if (Chunk.GetNumInstances() == 0)
+		{
+			// Remember first empty chunk but continue looking for a chunk that has space and same group tag
+			if (EmptyChunkIndex == INDEX_NONE)
+			{
+				EmptyChunkIndex = ChunkIndex;
+				EmptyAbsoluteIndex = AbsoluteIndex;
+			}
+		}
+		else if (Chunk.GetNumInstances() < NumEntitiesPerChunk)
+		{
+			IndexWithinChunk = Chunk.GetNumInstances();
+			AbsoluteIndex += IndexWithinChunk;
+
+			Chunk.AddInstance();
+
+			DestinationChunk = &Chunk;
+			break;
+		}
+		AbsoluteIndex += NumEntitiesPerChunk;
+		++ChunkIndex;
+	}
+
+	if (DestinationChunk == nullptr)
+	{
+		// Check if it is a recycled chunk
+		if (EmptyChunkIndex != INDEX_NONE)
+		{
+			DestinationChunk = &Chunks[EmptyChunkIndex];
+			DestinationChunk->Recycle(ChunkComponentTemplates);
+			AbsoluteIndex = EmptyAbsoluteIndex;
+		}
+		else
+		{
+			DestinationChunk = &Chunks.Emplace_GetRef(GetChunkAllocSize(), ChunkComponentTemplates);
+		}
+
+		check(DestinationChunk);
+		DestinationChunk->AddInstance();
+	}
+
+	// Initialize the component memory
+	if (bInitializeComponents)
+	{
+		for (const FMassArchetypeFragmentConfig& ComponentConfig : ComponentConfigs)
+		{
+			void* ComponentPtr = ComponentConfig.GetComponentData(DestinationChunk->GetRawMemory(), IndexWithinChunk);
+			ComponentConfig.ComponentType->InitializeStruct(ComponentPtr);
+		}
+	}
+
+	// Add to the table and map
+	EntityMap.Add(Entity.Index, AbsoluteIndex);
+	DestinationChunk->GetEntityArrayElementRef(EntityListOffsetWithinChunk, IndexWithinChunk) = Entity;
+
+	return AbsoluteIndex;
+}
+
+void FMassArchetypeData::RemoveEntity(FMassEntityHandle Entity)
+{
+	const int32 AbsoluteIndex = EntityMap.FindAndRemoveChecked(Entity.Index);
+	RemoveEntityInternal(AbsoluteIndex, true/*bDestroyComponents*/);
+}
+
+void FMassArchetypeData::RemoveEntityInternal(const int32 AbsoluteIndex, const bool bDestroyComponents)
+{
+	const int32 ChunkIndex = AbsoluteIndex / NumEntitiesPerChunk;
+	const int32 IndexWithinChunk = AbsoluteIndex % NumEntitiesPerChunk;
+
+	FMassArchetypeChunk& Chunk = Chunks[ChunkIndex];
+
+	Chunk.RemoveInstance();
+	const int32 IndexToSwapFrom = Chunk.GetNumInstances();
+
+	checkf(bDestroyComponents || UE::AggregateTicking::bBitwiseRelocateComponents, TEXT("We allow not to destroy components only in bit wise relocation mode."));
+
+	// Remove and swap the last entry in the chunk to the location of the removed item (if it's not the same as the dying entry)
+	if (IndexToSwapFrom != IndexWithinChunk)
+	{
+		for (const FMassArchetypeFragmentConfig& ComponentConfig : ComponentConfigs)
+		{
+			void* DyingComponentPtr = ComponentConfig.GetComponentData(Chunk.GetRawMemory(), IndexWithinChunk);
+			void* MovingComponentPtr = ComponentConfig.GetComponentData(Chunk.GetRawMemory(), IndexToSwapFrom);
+
+			if (UE::AggregateTicking::bBitwiseRelocateComponents)
+			{
+				// Destroy component data
+				if (bDestroyComponents)
+				{
+					ComponentConfig.ComponentType->DestroyStruct(DyingComponentPtr);
+				}
+
+				// Move last entry
+				FMemory::Memcpy(DyingComponentPtr, MovingComponentPtr, ComponentConfig.ComponentType->GetStructureSize());
+			}
+			else
+			{
+				// Destroy & initialize the component data
+				ComponentConfig.ComponentType->ClearScriptStruct(DyingComponentPtr);
+
+				// Copy last entry
+				ComponentConfig.ComponentType->CopyScriptStruct(DyingComponentPtr, MovingComponentPtr);
+
+				// Destroy last entry
+				ComponentConfig.ComponentType->DestroyStruct(MovingComponentPtr);
+			}
+		}
+
+		// Update the entity table and map
+		const FMassEntityHandle EntityBeingSwapped = Chunk.GetEntityArrayElementRef(EntityListOffsetWithinChunk, IndexToSwapFrom);
+		Chunk.GetEntityArrayElementRef(EntityListOffsetWithinChunk, IndexWithinChunk) = EntityBeingSwapped;
+		EntityMap.FindChecked(EntityBeingSwapped.Index) = AbsoluteIndex;
+	}
+	else if (!UE::AggregateTicking::bBitwiseRelocateComponents || bDestroyComponents)
+	{
+		for (const FMassArchetypeFragmentConfig& ComponentConfig : ComponentConfigs)
+		{
+			// Destroy the component data
+			void* DyingComponentPtr = ComponentConfig.GetComponentData(Chunk.GetRawMemory(), IndexWithinChunk);
+			ComponentConfig.ComponentType->DestroyStruct(DyingComponentPtr);
+		}
+	}
+
+	// If the chunk itself is empty now, see if we can remove it entirely
+	// Note: This is only possible for trailing chunks, to avoid messing up the absolute indices in the entities map
+	while ((Chunks.Num() > 0) && (Chunks.Last().GetNumInstances() == 0))
+	{
+		Chunks.RemoveAt(Chunks.Num() - 1, 1, /*bAllowShrinking=*/ false);
+	}
+}
+
+void FMassArchetypeData::BatchDestroyEntityChunks(const FArchetypeChunkCollection& ChunkCollection, TArray<FMassEntityHandle>& OutEntitiesRemoved)
+{
+	const int32 InitialOutEntitiesCount = OutEntitiesRemoved.Num();
+
+	// Sorting the subchunks info so that subchunks of a given chunk are processed "from the back". Otherwise removing 
+	// a subchunk from the front of the chunk would inevitably invalidate following subchunks' information.
+	TArray<FArchetypeChunkCollection::FChunkInfo> Subchunks(ChunkCollection.GetChunks());
+	Subchunks.Sort([](const FArchetypeChunkCollection::FChunkInfo& A, const FArchetypeChunkCollection::FChunkInfo& B) 
+		{ 
+			return A.ChunkIndex < B.ChunkIndex || (A.ChunkIndex == B.ChunkIndex && A.SubchunkStart > B.SubchunkStart);
+		});
+
+	for (const FArchetypeChunkCollection::FChunkInfo SubchunkInfo : Subchunks)
+	{ 
+		FMassArchetypeChunk& Chunk = Chunks[SubchunkInfo.ChunkIndex];
+
+		// gather entities we're about to remove
+		FMassEntityHandle* DyingEntityPtr = &Chunk.GetEntityArrayElementRef(EntityListOffsetWithinChunk, SubchunkInfo.SubchunkStart);
+		OutEntitiesRemoved.Append(DyingEntityPtr, SubchunkInfo.Length);
+
+		const int32 NumberToMove = FMath::Min(Chunk.GetNumInstances() - (SubchunkInfo.SubchunkStart + SubchunkInfo.Length), SubchunkInfo.Length);
+		checkf(NumberToMove >= 0, TEXT("Trying to move a negative number of elements indicates a problem with SubchunkInfo, it's possibly out of date."));
+		const int32 NumberToCut = FMath::Max(SubchunkInfo.Length - NumberToMove, 0);
+		
+		if (NumberToMove > 0)
+		{
+			const int32 SwapStartIndex = Chunk.GetNumInstances() - NumberToMove;
+			checkf((SubchunkInfo.SubchunkStart + NumberToMove - 1) < SwapStartIndex, TEXT("Remove and Move ranges overlap"));
+
+			for (const FMassArchetypeFragmentConfig& ComponentConfig : ComponentConfigs)
+			{
+				void* DyingComponentPtr = ComponentConfig.GetComponentData(Chunk.GetRawMemory(), SubchunkInfo.SubchunkStart);
+				void* MovingComponentPtr = ComponentConfig.GetComponentData(Chunk.GetRawMemory(), SwapStartIndex);
+
+				if (UE::AggregateTicking::bBitwiseRelocateComponents)
+				{
+					// Destroy the components we'll replace by the following copy
+					ComponentConfig.ComponentType->DestroyStruct(DyingComponentPtr, NumberToMove);
+
+					// Swap components to the empty space just created.
+					FMemory::Memcpy(DyingComponentPtr, MovingComponentPtr, ComponentConfig.ComponentType->GetStructureSize() * NumberToMove);
+				}
+				else
+				{
+					// Clear components that we will copy over. Clear destroys and initializes the components, which is needed for CopyScriptStruct().
+					ComponentConfig.ComponentType->ClearScriptStruct(DyingComponentPtr, NumberToMove);
+
+					// Swap components into the empty space just created.
+					ComponentConfig.ComponentType->CopyScriptStruct(DyingComponentPtr, MovingComponentPtr, NumberToMove);
+
+					// Destroy the components that were moved.
+					ComponentConfig.ComponentType->DestroyStruct(MovingComponentPtr, NumberToMove);
+				}
+			}
+
+			// Update the entity table and map
+			const FMassEntityHandle* MovingEntityPtr = &Chunk.GetEntityArrayElementRef(EntityListOffsetWithinChunk, SwapStartIndex);
+			int32 AbsoluteIndex = SubchunkInfo.ChunkIndex * NumEntitiesPerChunk + SubchunkInfo.SubchunkStart;
+
+			for (int i = 0; i < NumberToMove; ++i)
+			{
+				DyingEntityPtr[i] = MovingEntityPtr[i];
+				EntityMap.FindChecked(MovingEntityPtr[i].Index) = AbsoluteIndex++;
+			}
+		}
+
+		if (NumberToCut > 0)
+		{
+			// just clean up the rest. Note that we explicitly do not clean the spots vacated by entities moved from 
+			// the back of the chunk - if we did the risk calling DestroyStruct on them multiple times
+			const int32 CutStartIndex = SubchunkInfo.SubchunkStart + NumberToMove;
+			for (const FMassArchetypeFragmentConfig& ComponentConfig : ComponentConfigs)
+			{
+				// Destroy the component data
+				void* DyingComponentPtr = ComponentConfig.GetComponentData(Chunk.GetRawMemory(), CutStartIndex);
+				ComponentConfig.ComponentType->DestroyStruct(DyingComponentPtr, NumberToCut);
+			}
+		}
+
+		Chunk.RemoveMultipleInstances(SubchunkInfo.Length);
+	}
+
+	for (int i = InitialOutEntitiesCount; i < OutEntitiesRemoved.Num(); ++i)
+	{
+		EntityMap.FindAndRemoveChecked(OutEntitiesRemoved[i].Index);
+	}
+
+	// If the chunk itself is empty now, see if we can remove it entirely
+	// Note: This is only possible for trailing chunks, to avoid messing up the absolute indices in the entities map
+	while ((Chunks.Num() > 0) && (Chunks.Last().GetNumInstances() == 0))
+	{
+		Chunks.RemoveAt(Chunks.Num() - 1, 1, /*bAllowShrinking=*/ false);
+	}
+}
+
+bool FMassArchetypeData::HasComponentDataForEntity(const UScriptStruct* ComponentType, int32 EntityIndex) const
+{
+	return (ComponentType && CompositionDescriptor.Components.Contains(*ComponentType));
+}
+
+void* FMassArchetypeData::GetComponentDataForEntityChecked(const UScriptStruct* ComponentType, int32 EntityIndex) const
+{
+	const FInternalEntityHandle InternalIndex = MakeEntityHandle(EntityIndex);
+	
+	// failing the below Find means given entity's archetype is missing given ComponentType
+	const int32 ComponentIndex = ComponentIndexMap.FindChecked(ComponentType);
+	return GetComponentData(ComponentIndex, InternalIndex);
+}
+
+void* FMassArchetypeData::GetComponentDataForEntity(const UScriptStruct* ComponentType, int32 EntityIndex) const
+{
+	if (const int32* ComponentIndex = ComponentIndexMap.Find(ComponentType))
+	{
+		FInternalEntityHandle InternalIndex = MakeEntityHandle(EntityIndex);
+		// failing the below Find means given entity's archetype is missing given ComponentType
+		return GetComponentData(*ComponentIndex, InternalIndex);
+	}
+	return nullptr;
+}
+
+void FMassArchetypeData::SetComponentsData(const FMassEntityHandle Entity, TArrayView<const FInstancedStruct> ComponentInstances)
+{
+	FInternalEntityHandle InternalIndex = MakeEntityHandle(Entity);
+
+	for (const FInstancedStruct& Instance : ComponentInstances)
+	{
+		const UScriptStruct* ComponentType = Instance.GetScriptStruct();
+		check(ComponentType);
+		const int32 ComponentIndex = ComponentIndexMap.FindChecked(ComponentType);
+		void* ComponentMemory = GetComponentData(ComponentIndex, InternalIndex);
+		// No UE::AggregateTicking::bBitwiseRelocateComponents, this isn't a move component
+		ComponentType->CopyScriptStruct(ComponentMemory, Instance.GetMemory());
+	}
+}
+
+void FMassArchetypeData::SetComponentData(const FArchetypeChunkCollection& ChunkCollection, const FInstancedStruct& ComponentSource)
+{
+	check(ComponentSource.IsValid());
+	const UScriptStruct* ComponentType = ComponentSource.GetScriptStruct();
+	check(ComponentType);
+	const int32 ComponentIndex = ComponentIndexMap.FindChecked(ComponentType);
+	const int32 ComponentTypeSize = ComponentType->GetStructureSize();
+	const uint8* ComponentSourceMemory = ComponentSource.GetMemory();
+	check(ComponentSourceMemory);
+	
+	for (FMassArchetypeChunkIterator ChunkIterator(ChunkCollection); ChunkIterator; ++ChunkIterator)
+	{
+		uint8* ComponentMemory = (uint8*)ComponentConfigs[ComponentIndex].GetComponentData(Chunks[ChunkIterator->ChunkIndex].GetRawMemory(), ChunkIterator->SubchunkStart);
+		for (int i = ChunkIterator->Length; i; --i, ComponentMemory += ComponentTypeSize)
+		{
+			// No UE::AggregateTicking::bBitwiseRelocateComponents, this isn't a move of a component
+			ComponentType->CopyScriptStruct(ComponentMemory, ComponentSourceMemory);
+		}
+	}
+}
+
+void FMassArchetypeData::SetDefaultChunkComponentValue(FConstStructView InstancedStruct)
+{
+	if (ensure(InstancedStruct.GetScriptStruct()) && CompositionDescriptor.ChunkComponents.Contains(*InstancedStruct.GetScriptStruct()))
+	{
+		if (FInstancedStruct* Element = ChunkComponentTemplates.FindByPredicate([&InstancedStruct](const FInstancedStruct& Element) { return Element.GetScriptStruct() == InstancedStruct.GetScriptStruct(); }))
+		{
+			*Element = InstancedStruct;
+		}
+	}
+}
+
+void FMassArchetypeData::MoveEntityToAnotherArchetype(const FMassEntityHandle Entity, FMassArchetypeData& NewArchetype)
+{
+	check(&NewArchetype != this);
+
+	const int32 AbsoluteIndex = EntityMap.FindAndRemoveChecked(Entity.Index);
+	const int32 ChunkIndex = AbsoluteIndex / NumEntitiesPerChunk;
+	const int32 IndexWithinChunk = AbsoluteIndex % NumEntitiesPerChunk;
+	FMassArchetypeChunk& Chunk = Chunks[ChunkIndex];
+
+	constexpr bool bInitializeComponentsDuringCreation = !UE::AggregateTicking::bBitwiseRelocateComponents;
+	const int32 NewAbsoluteIndex = NewArchetype.AddEntityInternal(Entity, bInitializeComponentsDuringCreation);
+	const int32 NewChunkIndex = NewAbsoluteIndex / NewArchetype.NumEntitiesPerChunk;
+	const int32 NewIndexWithinChunk = NewAbsoluteIndex % NewArchetype.NumEntitiesPerChunk;
+	FMassArchetypeChunk& NewChunk = NewArchetype.Chunks[NewChunkIndex];
+
+	// for every NewArchetype's component see if it was in the old component as well and if so copy it's value. 
+	// If not then initialize the component (unless it has already been initialized based on bInitializeComponentsDuringCreation
+	// value
+	for (const FMassArchetypeFragmentConfig& NewComponentConfig : NewArchetype.ComponentConfigs)
+	{
+		const int32* OldComponentIndex = ComponentIndexMap.Find(NewComponentConfig.ComponentType);
+		void* Dst = NewComponentConfig.GetComponentData(NewChunk.GetRawMemory(), NewIndexWithinChunk);
+
+		// Only copy if the component type exists in both archetypes
+		if (OldComponentIndex)
+		{
+			const void* Src = ComponentConfigs[*OldComponentIndex].GetComponentData(Chunk.GetRawMemory(), IndexWithinChunk);
+			if (UE::AggregateTicking::bBitwiseRelocateComponents)
+			{
+				FMemory::Memcpy(Dst, Src, NewComponentConfig.ComponentType->GetStructureSize());
+			}
+			else
+			{
+				NewComponentConfig.ComponentType->CopyScriptStruct(Dst, Src);
+			}
+		}
+		else if (bInitializeComponentsDuringCreation == false)
+		{
+			// the component's unique to the NewArchetype need to be initialized
+			// @todo we're doing it for tags here as well. A tiny bit of perf lost. Probably not worth adding a check
+			// but something to keep in mind. Will go away once tags are more of an archetype fragment than entity's
+			NewComponentConfig.ComponentType->InitializeStruct(Dst);
+		}
+	}
+
+	constexpr bool bDestroyComponents = !UE::AggregateTicking::bBitwiseRelocateComponents;
+	RemoveEntityInternal(AbsoluteIndex, bDestroyComponents);
+}
+
+void FMassArchetypeData::ExecuteFunction(FMassExecutionContext& RunContext, const FMassExecuteFunction& Function, const FMassQueryRequirementIndicesMapping& RequirementMapping, const FArchetypeChunkCollection& ChunkCollection)
+{
+	check(ChunkCollection.GetArchetype() == this);
+
+	// mz@todo to be removed
+	RunContext.SetCurrentArchetypeData(*this);
+	for (FMassArchetypeChunkIterator ChunkIterator(ChunkCollection); ChunkIterator; ++ChunkIterator)
+	{
+		FMassArchetypeChunk& Chunk = Chunks[ChunkIterator->ChunkIndex];
+		
+		const int32 ChunkLength = ChunkIterator->Length > 0 ? ChunkIterator->Length : (Chunk.GetNumInstances() - ChunkIterator->SubchunkStart);
+		if (ChunkLength)
+		{
+			checkf((ChunkIterator->SubchunkStart + ChunkLength) <= Chunk.GetNumInstances() && ChunkLength > 0, TEXT("Invalid subchunk, it is going over the number of instances in the chunk or it is empty."));
+
+			RunContext.SetCurrentChunkSerialModificationNumber(Chunk.GetSerialModificationNumber());
+			BindChunkComponentRequirements(RunContext, RequirementMapping.ChunkComponents, Chunk);
+			BindEntityRequirements(RunContext, RequirementMapping.EntityComponents, Chunk, ChunkIterator->SubchunkStart, ChunkLength);
+
+			Function(RunContext);
+		}
+	}
+}
+
+void FMassArchetypeData::ExecuteFunction(FMassExecutionContext& RunContext, const FMassExecuteFunction& Function, const FMassQueryRequirementIndicesMapping& RequirementMapping, const FMassChunkConditionFunction& ChunkCondition)
+{
+	// mz@todo to be removed
+	RunContext.SetCurrentArchetypeData(*this);
+	for (FMassArchetypeChunk& Chunk : Chunks)
+	{
+		if (Chunk.GetNumInstances())
+		{
+			RunContext.SetCurrentChunkSerialModificationNumber(Chunk.GetSerialModificationNumber());
+			BindChunkComponentRequirements(RunContext, RequirementMapping.ChunkComponents, Chunk);
+
+			if (!ChunkCondition || ChunkCondition(RunContext))
+			{
+				BindEntityRequirements(RunContext, RequirementMapping.EntityComponents, Chunk, 0, Chunk.GetNumInstances());
+				Function(RunContext);
+			}
+		}
+	}
+}
+
+void FMassArchetypeData::ExecutionFunctionForChunk(FMassExecutionContext RunContext, const FMassExecuteFunction& Function, const FMassQueryRequirementIndicesMapping& RequirementMapping, const FArchetypeChunkCollection::FChunkInfo& ChunkInfo, const FMassChunkConditionFunction& ChunkCondition)
+{
+	FMassArchetypeChunk& Chunk = Chunks[ChunkInfo.ChunkIndex];
+	const int32 ChunkLength = ChunkInfo.Length > 0 ? ChunkInfo.Length : (Chunk.GetNumInstances() - ChunkInfo.SubchunkStart);
+
+	if (ChunkLength)
+	{
+		RunContext.SetCurrentArchetypeData(*this);
+		RunContext.SetCurrentChunkSerialModificationNumber(Chunk.GetSerialModificationNumber());
+		BindChunkComponentRequirements(RunContext, RequirementMapping.ChunkComponents, Chunk);
+
+		if (!ChunkCondition || ChunkCondition(RunContext))
+		{
+			BindEntityRequirements(RunContext, RequirementMapping.EntityComponents, Chunk, ChunkInfo.SubchunkStart, ChunkLength);
+			Function(RunContext);
+		}
+	}
+}
+
+void FMassArchetypeData::CompactEntities(const double TimeAllowed)
+{
+	const double TimeAllowedEnd = FPlatformTime::Seconds() + TimeAllowed;
+
+	TArray<FMassArchetypeChunk*> SortedChunks;
+	for (FMassArchetypeChunk& Chunk : Chunks)
+	{
+		// Skip already full chunks
+		const int32 NumInstances = Chunk.GetNumInstances();
+		if(NumInstances > 0 && NumInstances < NumEntitiesPerChunk)
+		{
+			SortedChunks.Add(&Chunk);
+		}
+	}
+
+	// Check if there is anything to compact at all
+	if (SortedChunks.Num() <= 1)
+	{
+		return;
+	}
+
+	SortedChunks.Sort([](const FMassArchetypeChunk& LHS, const FMassArchetypeChunk& RHS)
+	{
+		return LHS.GetNumInstances() < RHS.GetNumInstances();
+	});
+
+	int32 ChunkToFillSortedIdx = 0;
+	int32 ChunkToEmptySortedIdx = SortedChunks.Num() -1;
+	while (ChunkToFillSortedIdx < ChunkToEmptySortedIdx && FPlatformTime::Seconds() < TimeAllowedEnd)
+	{
+		while (ChunkToFillSortedIdx < SortedChunks.Num() && SortedChunks[ChunkToFillSortedIdx]->GetNumInstances() == NumEntitiesPerChunk)
+		{
+			ChunkToFillSortedIdx++;
+		}
+		while (ChunkToEmptySortedIdx >= 0 && SortedChunks[ChunkToEmptySortedIdx]->GetNumInstances() == 0)
+		{
+			ChunkToEmptySortedIdx--;
+		}
+		if (ChunkToFillSortedIdx >= ChunkToEmptySortedIdx)
+		{
+			break;
+		}
+
+		FMassArchetypeChunk* ChunkToFill = SortedChunks[ChunkToFillSortedIdx];
+		FMassArchetypeChunk* ChunkToEmpty = SortedChunks[ChunkToEmptySortedIdx];
+		const int32 NumberOfEntitiesToMove =  FMath::Min(NumEntitiesPerChunk-ChunkToFill->GetNumInstances(), ChunkToEmpty->GetNumInstances());
+		const int32 FromIndex = ChunkToEmpty->GetNumInstances() - NumberOfEntitiesToMove;
+		const int32 ToIndex = ChunkToFill->GetNumInstances();
+		check(NumberOfEntitiesToMove > 0);
+
+		for (const FMassArchetypeFragmentConfig& ComponentConfig : ComponentConfigs)
+		{
+			void* FromComponentPtr = ComponentConfig.GetComponentData(ChunkToEmpty->GetRawMemory(), FromIndex);
+			void* ToComponentPtr = ComponentConfig.GetComponentData(ChunkToFill->GetRawMemory(), ToIndex);
+
+			if (UE::AggregateTicking::bBitwiseRelocateComponents)
+			{
+				// Move all entries
+				FMemory::Memcpy(ToComponentPtr, FromComponentPtr, ComponentConfig.ComponentType->GetStructureSize() * NumberOfEntitiesToMove);
+			}
+			else
+			{
+				// Destroy & initialize the component data
+				ComponentConfig.ComponentType->ClearScriptStruct(ToComponentPtr, NumberOfEntitiesToMove);
+
+				// Copy all entries
+				ComponentConfig.ComponentType->CopyScriptStruct(ToComponentPtr, FromComponentPtr, NumberOfEntitiesToMove);
+
+				// Destroy all entries
+				ComponentConfig.ComponentType->DestroyStruct(FromComponentPtr, NumberOfEntitiesToMove);
+			}
+		}
+
+		FMassEntityHandle* FromEntity = &ChunkToEmpty->GetEntityArrayElementRef(EntityListOffsetWithinChunk, FromIndex);
+		FMassEntityHandle* ToEntity = &ChunkToFill->GetEntityArrayElementRef(EntityListOffsetWithinChunk, ToIndex);
+		FMemory::Memcpy(ToEntity, FromEntity, NumberOfEntitiesToMove * sizeof(FMassEntityHandle));
+		ChunkToFill->AddMultipleInstances(NumberOfEntitiesToMove);
+		ChunkToEmpty->RemoveMultipleInstances(NumberOfEntitiesToMove);
+
+		const int32 ChunkToFillIdx = ChunkToFill - &Chunks[0];
+		check(ChunkToFillIdx >=0 && ChunkToFillIdx < Chunks.Num());
+		const int32 AbsoluteIndex = ChunkToFillIdx * NumEntitiesPerChunk + ToIndex;
+
+		for (int32 i =0; i < NumberOfEntitiesToMove; i++, ++ToEntity)
+		{
+			EntityMap.FindChecked(ToEntity->Index) = AbsoluteIndex+i;
+		}
+	}
+}
+
+void FMassArchetypeData::GetRequirementsComponentMapping(TConstArrayView<FMassFragmentRequirement> Requirements, FMassFragmentIndicesMapping& OutComponentIndices)
+{
+	OutComponentIndices.Reset(Requirements.Num());
+	for (const FMassFragmentRequirement& Requirement : Requirements)
+	{
+		if (Requirement.RequiresBinding())
+		{
+			const int32* ComponentIndex = ComponentIndexMap.Find(Requirement.StructType);
+			check(ComponentIndex != nullptr || Requirement.IsOptional());
+			OutComponentIndices.Add(ComponentIndex ? *ComponentIndex : INDEX_NONE);
+		}
+	}
+}
+
+void FMassArchetypeData::GetRequirementsChunkComponentMapping(TConstArrayView<FMassFragmentRequirement> ChunkRequirements, FMassFragmentIndicesMapping& OutComponentIndices)
+{
+	int32 LastFoundComponentIndex = -1;
+	OutComponentIndices.Reset(ChunkRequirements.Num());
+	for (const FMassFragmentRequirement& Requirement : ChunkRequirements)
+	{
+		if (Requirement.RequiresBinding())
+		{
+			int32 ComponentIndex = INDEX_NONE;
+			for (int32 i = LastFoundComponentIndex + 1; i < ChunkComponentTemplates.Num(); ++i)
+			{
+				if (ChunkComponentTemplates[i].GetScriptStruct()->IsChildOf(Requirement.StructType))
+				{
+					ComponentIndex = i;
+					break;
+				}
+			}
+
+			check(ComponentIndex != INDEX_NONE);
+			OutComponentIndices.Add(ComponentIndex);
+			LastFoundComponentIndex = ComponentIndex;
+		}
+	}
+}
+
+void FMassArchetypeData::BindEntityRequirements(FMassExecutionContext& RunContext, const FMassFragmentIndicesMapping& EntityComponentsMapping, FMassArchetypeChunk& Chunk, const int32 SubchunkStart, const int32 SubchunkLength)
+{
+	// auto-correcting number of entities to process in case SubchunkStart +  SubchunkLength > Chunk.GetNumInstances()
+	const int32 NumEntities = SubchunkLength >= 0 ? FMath::Min(SubchunkLength, Chunk.GetNumInstances() - SubchunkStart) : Chunk.GetNumInstances();
+	check(SubchunkStart >= 0 && SubchunkStart < Chunk.GetNumInstances());
+
+	if (EntityComponentsMapping.Num() > 0)
+	{
+		check(RunContext.GetMutableRequirements().Num() == EntityComponentsMapping.Num());
+
+		for (int i = 0; i < EntityComponentsMapping.Num(); ++i)
+		{
+			FMassExecutionContext::FFragmentView& Requirement = RunContext.ComponentViews[i];
+			const int32 ComponentIndex = EntityComponentsMapping[i];
+
+			check(ComponentIndex != INDEX_NONE || Requirement.Requirement.IsOptional());
+			if (ComponentIndex != INDEX_NONE)
+			{
+				Requirement.ComponentView = TArrayView<FMassFragment>((FMassFragment*)GetComponentData(ComponentIndex, Chunk.GetRawMemory(), SubchunkStart), NumEntities);
+			}
+			else
+			{
+				// @todo this might not be needed
+				Requirement.ComponentView = TArrayView<FMassFragment>();
+			}
+		}
+	}
+	else
+	{
+		// Map in the required data arrays from the current chunk to the array views
+		for (FMassExecutionContext::FFragmentView& Requirement : RunContext.GetMutableRequirements())
+		{
+			const int32* ComponentIndex = ComponentIndexMap.Find(Requirement.Requirement.StructType);
+			check(ComponentIndex != nullptr || Requirement.Requirement.IsOptional());
+			if (ComponentIndex)
+			{
+				Requirement.ComponentView = TArrayView<FMassFragment>((FMassFragment*)GetComponentData(*ComponentIndex, Chunk.GetRawMemory(), SubchunkStart), NumEntities);
+			}
+			else
+			{
+				Requirement.ComponentView = TArrayView<FMassFragment>();
+			}
+		}
+	}
+
+	RunContext.EntityListView = TArrayView<FMassEntityHandle>(&Chunk.GetEntityArrayElementRef(EntityListOffsetWithinChunk, SubchunkStart), NumEntities);
+}
+
+void FMassArchetypeData::BindChunkComponentRequirements(FMassExecutionContext& RunContext, const FMassFragmentIndicesMapping& ChunkComponentsMapping, FMassArchetypeChunk& Chunk)
+{
+	if (ChunkComponentsMapping.Num() > 0)
+	{
+		check(RunContext.GetMutableChunkRequirements().Num() == ChunkComponentsMapping.Num());
+
+		for (int i = 0; i < ChunkComponentsMapping.Num(); ++i)
+		{
+			FMassExecutionContext::FChunkFragmentView& ChunkRequirement = RunContext.ChunkComponents[i];
+			const int32 ChunkComponentIndex = ChunkComponentsMapping[i];
+
+			check(ChunkComponentIndex != INDEX_NONE);
+			ChunkRequirement.ChunkComponentView = Chunk.GetMutableChunkComponentViewChecked(ChunkComponentIndex);
+		}
+	}
+	else
+	{
+		for (FMassExecutionContext::FChunkFragmentView& ChunkRequirement : RunContext.GetMutableChunkRequirements())
+		{
+			FInstancedStruct* ChunkComponentInstance = Chunk.FindMutableChunkComponent(ChunkRequirement.Requirement.StructType);
+			check(ChunkComponentInstance != nullptr);
+			ChunkRequirement.ChunkComponentView = FStructView(*ChunkComponentInstance);
+		}
+	}
+}
+
+SIZE_T FMassArchetypeData::GetAllocatedSize() const
+{
+	return sizeof(FMassArchetypeData) +
+		ComponentConfigs.GetAllocatedSize() +
+		Chunks.GetAllocatedSize() +
+		(Chunks.Num() * GetChunkAllocSize()) +
+		EntityMap.GetAllocatedSize() +
+		ComponentIndexMap.GetAllocatedSize();
+}
+
+FString FMassArchetypeData::DebugGetDescription() const
+{
+#if WITH_MASSENTITY_DEBUG
+	TStringBuilder<256> ArchetypeDebugName;
+	
+	ArchetypeDebugName.Append(TEXT("<"));
+
+	bool bNeedsComma = false;
+	for (const FMassArchetypeFragmentConfig& ComponentConfig : ComponentConfigs)
+	{
+		if (bNeedsComma)
+		{
+			ArchetypeDebugName.Append(TEXT(","));
+		}
+		ArchetypeDebugName.Append(ComponentConfig.ComponentType->GetName());
+		bNeedsComma = true;
+	}
+
+	ArchetypeDebugName.Append(TEXT(">"));
+	return ArchetypeDebugName.ToString();
+#else
+	return {};
+#endif
+}
+
+#if WITH_MASSENTITY_DEBUG
+void FMassArchetypeData::DebugPrintArchetype(FOutputDevice& Ar)
+{
+	FStringOutputDevice TagsDecription;
+	CompositionDescriptor.Tags.DebugGetStringDesc(TagsDecription);
+	Ar.Logf(ELogVerbosity::Log, TEXT("Tags: %s"), *TagsDecription);
+	Ar.Logf(ELogVerbosity::Log, TEXT("Components: %s"), *DebugGetDescription());
+	Ar.Logf(ELogVerbosity::Log, TEXT("\tChunks: %d x %d KB = %d KB total"), Chunks.Num(), GetChunkAllocSize() / 1024, (GetChunkAllocSize()*Chunks.Num()) / 1024);
+	
+	int ChunkWithComponentsCount = 0;
+	for (FMassArchetypeChunk& Chunk : Chunks)
+	{
+		ChunkWithComponentsCount += Chunk.DebugGetChunkComponentCount() > 0 ? 1 : 0;
+	}
+	if (ChunkWithComponentsCount)
+	{
+		Ar.Logf(ELogVerbosity::Log, TEXT("\tChunks with components: %d"), ChunkWithComponentsCount);
+	}
+
+	const int32 CurrentEntityCapacity = Chunks.Num() * NumEntitiesPerChunk;
+	Ar.Logf(ELogVerbosity::Log, TEXT("\tEntity Count    : %d"), EntityMap.Num());
+	Ar.Logf(ELogVerbosity::Log, TEXT("\tEntity Capacity : %d"), CurrentEntityCapacity);
+	if (Chunks.Num() > 1)
+	{
+		const float Scaler = 100.0f / (float)CurrentEntityCapacity;
+		// count non-last chunks to see how occupied they are
+		int EntitiesPerChunkMin = CurrentEntityCapacity;
+		int EntitiesPerChunkMax = 0;
+		for (int ChunkIndex = 0; ChunkIndex < Chunks.Num() - 1; ++ChunkIndex)
+		{
+			const int Population = Chunks[ChunkIndex].GetNumInstances();
+			EntitiesPerChunkMin = FMath::Min(Population, EntitiesPerChunkMin);
+			EntitiesPerChunkMax = FMath::Max(Population, EntitiesPerChunkMax);
+		}
+		Ar.Logf(ELogVerbosity::Log, TEXT("\tEntity Occupancy: %.1f%% (min: %.1f%%, max: %.1f%%)"), Scaler * EntityMap.Num(), Scaler * EntitiesPerChunkMin, Scaler * EntitiesPerChunkMax);
+	}
+	else 
+	{
+		Ar.Logf(ELogVerbosity::Log, TEXT("\tEntity Occupancy: %.1f%%"), CurrentEntityCapacity > 0 ? ((EntityMap.Num() * 100.0f) / (float)CurrentEntityCapacity) : 0.f);
+	}
+	Ar.Logf(ELogVerbosity::Log, TEXT("\tBytes / Entity  : %d"), TotalBytesPerEntity);
+	Ar.Logf(ELogVerbosity::Log, TEXT("\tEntities / Chunk: %d"), NumEntitiesPerChunk);
+
+	Ar.Logf(ELogVerbosity::Log, TEXT("\tOffset 0x%04X: Entity[] (%d bytes each)"), EntityListOffsetWithinChunk, sizeof(FMassEntityHandle));
+	int32 TotalBytesOfValidData = sizeof(FMassEntityHandle) * NumEntitiesPerChunk;
+	for (const FMassArchetypeFragmentConfig& ComponentConfig : ComponentConfigs)
+	{
+		TotalBytesOfValidData += ComponentConfig.ComponentType->GetStructureSize() * NumEntitiesPerChunk;
+		Ar.Logf(ELogVerbosity::Log, TEXT("\tOffset 0x%04X: %s[] (%d bytes each)"), ComponentConfig.ArrayOffsetWithinChunk, *ComponentConfig.ComponentType->GetName(), ComponentConfig.ComponentType->GetStructureSize());
+	}
+
+	//@TODO: Print out padding in between things?
+
+	const int32 UnusuablePaddingOffset = TotalBytesPerEntity * NumEntitiesPerChunk;
+	const int32 UnusuablePaddingAmount = GetChunkAllocSize() - UnusuablePaddingOffset;
+	if (UnusuablePaddingAmount > 0)
+	{
+		Ar.Logf(ELogVerbosity::Log, TEXT("\tOffset 0x%04X: WastePadding[] (%d bytes total)"), UnusuablePaddingOffset, UnusuablePaddingAmount);
+	}
+
+	if (GetChunkAllocSize() != TotalBytesOfValidData + UnusuablePaddingAmount)
+	{
+		Ar.Logf(ELogVerbosity::Log, TEXT("\t@TODO: EXTRA PADDING HERE:  TotalBytesOfValidData: %d (%d missing)"), TotalBytesOfValidData, GetChunkAllocSize() - TotalBytesOfValidData);
+	}
+}
+
+void FMassArchetypeData::DebugPrintEntity(FMassEntityHandle Entity, FOutputDevice& Ar, const TCHAR* InPrefix) const
+{
+	for (const FMassArchetypeFragmentConfig& ComponentConfig : ComponentConfigs)
+	{
+		void* Data = GetComponentDataForEntityChecked(ComponentConfig.ComponentType, Entity.Index);
+		
+		FString ComponentName = ComponentConfig.ComponentType->GetName();
+		ComponentName.RemoveFromStart(InPrefix);
+
+		FString ValueStr;
+		ComponentConfig.ComponentType->ExportText(ValueStr, Data, /*Default*/nullptr, /*OwnerObject*/nullptr, EPropertyPortFlags::PPF_IncludeTransient, /*ExportRootScope*/nullptr);
+
+		Ar.Logf(TEXT("%s: %s"), *ComponentName, *ValueStr);
+	}
+}
+
+#endif // WITH_MASSENTITY_DEBUG
+
+void FMassArchetypeData::REMOVEME_GetArrayViewForComponentInChunk(int32 ChunkIndex, const UScriptStruct* ComponentType, void*& OutChunkBase, int32& OutNumEntities)
+{
+	const FMassArchetypeChunk& Chunk = Chunks[ChunkIndex];
+	const int32 ComponentIndex = ComponentIndexMap.FindChecked(ComponentType);
+
+	OutChunkBase = ComponentConfigs[ComponentIndex].GetComponentData(Chunk.GetRawMemory(), 0);
+	OutNumEntities = Chunk.GetNumInstances();
+}
+
diff --git a/Engine/Plugins/Runtime/MassEntity/Source/MassEntity/Private/MassArchetypeData.h b/Engine/Plugins/Runtime/MassEntity/Source/MassEntity/Private/MassArchetypeData.h
new file mode 100644
--- /dev/null
+++ b/Engine/Plugins/Runtime/MassEntity/Source/MassEntity/Private/MassArchetypeData.h
@@ -0,0 +1,281 @@
+// Copyright Epic Games, Inc. All Rights Reserved.
+
+#pragma once
+
+#include "MassEntitySubsystem.h"
+#include "MassArchetypeTypes.h"
+
+struct FMassEntityQuery;
+struct FMassExecutionContext;
+class FOutputDevice;
+struct FArchetypeChunkCollection;
+
+// This is one chunk within an archetype
+struct FMassArchetypeChunk
+{
+private:
+	uint8* RawMemory = nullptr;
+	int32 NumInstances = 0;
+	int32 SerialModificationNumber = 0;
+	TArray<FInstancedStruct> ChunkComponentData;
+
+public:
+	explicit FMassArchetypeChunk(int32 AllocSize, TConstArrayView<FInstancedStruct> InChunkComponentTemplates)
+		: ChunkComponentData(InChunkComponentTemplates)
+	{
+		RawMemory = (uint8*)FMemory::Malloc(AllocSize);
+	}
+
+	~FMassArchetypeChunk()
+	{
+		FMemory::Free(RawMemory);
+	}
+
+	// Returns the Entity array element at the specified index
+	FMassEntityHandle& GetEntityArrayElementRef(int32 ChunkBase, int32 IndexWithinChunk)
+	{
+		return ((FMassEntityHandle*)(RawMemory + ChunkBase))[IndexWithinChunk];
+	}
+
+	uint8* GetRawMemory() const
+	{
+		return RawMemory;
+	}
+
+	int32 GetNumInstances() const
+	{
+		return NumInstances;
+	}
+
+	void AddMultipleInstances(uint32 Count)
+	{
+		NumInstances += Count;
+		SerialModificationNumber++;
+	}
+
+	void RemoveMultipleInstances(uint32 Count)
+	{
+		NumInstances -= Count;
+		check(NumInstances >= 0);
+		SerialModificationNumber++;
+	}
+
+	void AddInstance()
+	{
+		AddMultipleInstances(1);
+	}
+
+	void RemoveInstance()
+	{
+		RemoveMultipleInstances(1);
+	}
+
+	int32 GetSerialModificationNumber() const
+	{
+		return SerialModificationNumber;
+	}
+
+	FStructView GetMutableChunkComponentViewChecked(const int32 Index) { return FStructView(ChunkComponentData[Index]); }
+
+	FInstancedStruct* FindMutableChunkComponent(const UScriptStruct* Type)
+	{
+		return ChunkComponentData.FindByPredicate([Type](const FInstancedStruct& Element)
+			{
+				return Element.GetScriptStruct()->IsChildOf(Type);
+			});
+	}
+
+	void Recycle(TConstArrayView<FInstancedStruct> InChunkComponentsTemplate)
+	{
+		checkf(NumInstances == 0, TEXT("Recycling a chunk that is not empty."));
+		SerialModificationNumber++;
+		ChunkComponentData = InChunkComponentsTemplate;
+	}
+
+#if WITH_MASSENTITY_DEBUG
+	int32 DebugGetChunkComponentCount() const { return ChunkComponentData.Num(); }
+#endif // WITH_MASSENTITY_DEBUG
+};
+
+// Information for a single component type in an archetype
+struct FMassArchetypeFragmentConfig
+{
+	const UScriptStruct* ComponentType = nullptr;
+	int32 ArrayOffsetWithinChunk = 0;
+
+	void* GetComponentData(uint8* ChunkBase, int32 IndexWithinChunk) const
+	{
+		return ChunkBase + ArrayOffsetWithinChunk + (IndexWithinChunk * ComponentType->GetStructureSize());
+	}
+};
+
+// An archetype is defined by a collection of unique component types (no duplicates).
+// Order doesn't matter, there will only ever be one FMassArchetypeData per unique set of component types per entity manager subsystem
+struct FMassArchetypeData
+{
+private:
+	// One-stop-shop variable describing the archetype's component and tag composition 
+	FMassCompositionDescriptor CompositionDescriptor;
+
+	TArray<FInstancedStruct> ChunkComponentTemplates;
+
+	TArray<FMassArchetypeFragmentConfig, TInlineAllocator<16>> ComponentConfigs;
+	
+	TArray<FMassArchetypeChunk> Chunks;
+
+	// Entity ID to index within archetype
+	//@TODO: Could be folded into FEntityData in the entity manager at the expense of a bit
+	// of loss of encapsulation and extra complexity during archetype changes
+	TMap<int32, int32> EntityMap;
+	
+	TMap<const UScriptStruct*, int32> ComponentIndexMap;
+
+	int32 NumEntitiesPerChunk;
+	int32 TotalBytesPerEntity;
+	int32 EntityListOffsetWithinChunk;
+
+	friend FMassEntityQuery;
+	friend FArchetypeChunkCollection;
+
+public:
+	TConstArrayView<FMassArchetypeFragmentConfig> GetComponentConfigs() const { return ComponentConfigs; }
+	const FMassFragmentBitSet& GetComponentBitSet() const { return CompositionDescriptor.Components; }
+	const FMassTagBitSet& GetTagBitSet() const { return CompositionDescriptor.Tags; }
+	const FMassChunkFragmentBitSet& GetChunkComponentBitSet() const { return CompositionDescriptor.ChunkComponents; }
+	const FMassCompositionDescriptor& GetCompositionDescriptor() const { return CompositionDescriptor; }
+
+	/** Method to iterate on all the component types */
+	void ForEachComponentType(TFunction< void(const UScriptStruct* /*ComponentType*/)> Function) const;
+	bool HasComponentType(const UScriptStruct* ComponentType) const;
+	bool HasTagType(const UScriptStruct* ComponentType) const { check(ComponentType); return CompositionDescriptor.Tags.Contains(*ComponentType); }
+	bool IsEquivalent(TConstArrayView<const UScriptStruct*> OtherList) const;
+	bool IsEquivalent(const FMassFragmentBitSet& InComponentBitSet, const FMassTagBitSet& InTagBitSet, const FMassChunkFragmentBitSet& InChunkComponentsBitSet) const
+	{ 
+		return CompositionDescriptor.IsEquivalent(InComponentBitSet, InTagBitSet, InChunkComponentsBitSet);
+	}
+	bool IsEquivalent(const FMassCompositionDescriptor& OtherCompositionDescriptor) const
+	{
+		return CompositionDescriptor.IsEquivalent(OtherCompositionDescriptor);
+	}
+
+	void Initialize(const FMassFragmentBitSet& Components, const FMassTagBitSet& Tags, const FMassChunkFragmentBitSet& ChunkComponents);
+
+	/** 
+	 * A special way of initializing an archetype resulting in a copy of SiblingArchetype's setup with OverrideTags
+	 * replacing original tags of SiblingArchetype
+	 */
+	void InitializeWithSibling(const FMassArchetypeData& SiblingArchetype, const FMassTagBitSet& OverrideTags);
+
+	void AddEntity(FMassEntityHandle Entity);
+	void RemoveEntity(FMassEntityHandle Entity);
+	void BatchDestroyEntityChunks(const FArchetypeChunkCollection& ChunkCollection, TArray<FMassEntityHandle>& OutEntitiesRemoved);
+
+	bool HasComponentDataForEntity(const UScriptStruct* ComponentType, int32 EntityIndex) const;
+	void* GetComponentDataForEntityChecked(const UScriptStruct* ComponentType, int32 EntityIndex) const;
+	void* GetComponentDataForEntity(const UScriptStruct* ComponentType, int32 EntityIndex) const;
+
+	FORCEINLINE int32 GetInternalIndexForEntity(const int32 EntityIndex) const { return EntityMap.FindChecked(EntityIndex); }
+	int32 GetNumEntitiesPerChunk() const { return NumEntitiesPerChunk; }
+
+	int32 GetNumEntities() const { return EntityMap.Num(); }
+
+	int32 GetChunkAllocSize() const { return 64*1024; }
+
+	int32 GetChunkCount() const { return Chunks.Num(); }
+
+	void ExecuteFunction(FMassExecutionContext& RunContext, const FMassExecuteFunction& Function, const FMassQueryRequirementIndicesMapping& RequirementMapping, const FArchetypeChunkCollection& ChunkCollection);
+	void ExecuteFunction(FMassExecutionContext& RunContext, const FMassExecuteFunction& Function, const FMassQueryRequirementIndicesMapping& RequirementMapping, const FMassChunkConditionFunction& ChunkCondition = FMassChunkConditionFunction());
+
+	void ExecutionFunctionForChunk(FMassExecutionContext RunContext, const FMassExecuteFunction& Function, const FMassQueryRequirementIndicesMapping& RequirementMapping, const FArchetypeChunkCollection::FChunkInfo& ChunkInfo, const FMassChunkConditionFunction& ChunkCondition = FMassChunkConditionFunction());
+
+	/**
+	 * Compacts entities to fill up chunks as much as possible
+	 */
+	void CompactEntities(const double TimeAllowed);
+
+	/**
+	 * Moves the entity from this archetype to another, will only copy all matching component types
+	 * @param Entity is the entity to move
+	 * @param NewArchetype the archetype to move to
+	 */
+	void MoveEntityToAnotherArchetype(const FMassEntityHandle Entity, FMassArchetypeData& NewArchetype);
+
+	/**
+	 * Set all component sources data on specified entity, will check if there are component sources type that does not exist in the archetype
+	 * @param Entity is the entity to set the data of all components
+	 * @param ComponentSources are the components to copy the data from
+	 */
+	void SetComponentsData(const FMassEntityHandle Entity, TArrayView<const FInstancedStruct> ComponentSources);
+
+	/** For all entities indicated by ChunkCollection the function sets the value of component of type
+	 *  ComponentSource.GetScriptStruct to the value represented by ComponentSource.GetMemory */
+	void SetComponentData(const FArchetypeChunkCollection& ChunkCollection, const FInstancedStruct& ComponentSource);
+
+	void SetDefaultChunkComponentValue(FConstStructView InstancedStruct);
+
+	/** Returns conversion from given Requirements to archetype's component indices */
+	void GetRequirementsComponentMapping(TConstArrayView<FMassFragmentRequirement> Requirements, FMassFragmentIndicesMapping& OutComponentIndices);
+
+	/** Returns conversion from given ChunkRequirements to archetype's chunk component indices */
+	void GetRequirementsChunkComponentMapping(TConstArrayView<FMassFragmentRequirement> ChunkRequirements, FMassFragmentIndicesMapping& OutComponentIndices);
+
+	SIZE_T GetAllocatedSize() const;
+
+	// Converts the list of components into a user-readable debug string
+	FString DebugGetDescription() const;
+
+#if WITH_MASSENTITY_DEBUG
+	/**
+	 * Prints out debug information about the archetype
+	 */
+	void DebugPrintArchetype(FOutputDevice& Ar);
+
+	/**
+	 * Prints out component's values for the specified entity. 
+	 * @param Entity The entity for which we want to print component values
+	 * @param Ar The output device
+	 * @param InPrefix Optional prefix to remove from component names
+	 */
+	void DebugPrintEntity(FMassEntityHandle Entity, FOutputDevice& Ar, const TCHAR* InPrefix = TEXT("")) const;
+#endif // WITH_MASSENTITY_DEBUG
+
+	void REMOVEME_GetArrayViewForComponentInChunk(int32 ChunkIndex, const UScriptStruct* ComponentType, void*& OutChunkBase, int32& OutNumEntities);
+
+	//////////////////////////////////////////////////////////////////////
+	// low level api
+	FORCEINLINE const int32* GetComponentIndex(const UScriptStruct* ComponentType) const { return ComponentIndexMap.Find(ComponentType); }
+	FORCEINLINE int32 GetComponentIndexChecked(const UScriptStruct* ComponentType) const { return ComponentIndexMap.FindChecked(ComponentType); }
+
+	FORCEINLINE void* GetComponentData(const int32 ComponentIndex, const FInternalEntityHandle EntityIndex) const
+	{
+		return ComponentConfigs[ComponentIndex].GetComponentData(EntityIndex.ChunkRawMemory, EntityIndex.IndexWithinChunk);
+	}
+
+	FORCEINLINE FInternalEntityHandle MakeEntityHandle(int32 EntityIndex) const
+	{
+		const int32 AbsoluteIndex = EntityMap.FindChecked(EntityIndex);
+		const int32 ChunkIndex = AbsoluteIndex / NumEntitiesPerChunk;
+	
+		return FInternalEntityHandle(Chunks[ChunkIndex].GetRawMemory(), AbsoluteIndex % NumEntitiesPerChunk); 
+	}
+
+	FORCEINLINE FInternalEntityHandle MakeEntityHandle(FMassEntityHandle Entity) const
+	{
+		return MakeEntityHandle(Entity.Index); 
+	}
+
+	bool IsInitialized() const { return TotalBytesPerEntity > 0 && ComponentConfigs.IsEmpty() == false; }
+
+protected:
+	FORCEINLINE void* GetComponentData(const int32 ComponentIndex, uint8* ChunkRawMemory, const int32 IndexWithinChunk) const
+	{
+		return ComponentConfigs[ComponentIndex].GetComponentData(ChunkRawMemory, IndexWithinChunk);
+	}
+
+	void BindEntityRequirements(FMassExecutionContext& RunContext, const FMassFragmentIndicesMapping& EntityComponentsMapping, FMassArchetypeChunk& Chunk, const int32 SubchunkStart, const int32 SubchunkLength);
+	void BindChunkComponentRequirements(FMassExecutionContext& RunContext, const FMassFragmentIndicesMapping& ChunkComponentsMapping, FMassArchetypeChunk& Chunk);
+
+private:
+	int32 AddEntityInternal(FMassEntityHandle Entity, const bool bInitializeComponents);
+	void RemoveEntityInternal(const int32 AbsoluteIndex, const bool bDestroyComponents);
+};
diff --git a/Engine/Plugins/Runtime/MassEntity/Source/MassEntity/Private/MassArchetypeTypes.cpp b/Engine/Plugins/Runtime/MassEntity/Source/MassEntity/Private/MassArchetypeTypes.cpp
new file mode 100644
--- /dev/null
+++ b/Engine/Plugins/Runtime/MassEntity/Source/MassEntity/Private/MassArchetypeTypes.cpp
@@ -0,0 +1,92 @@
+// Copyright Epic Games, Inc. All Rights Reserved.
+
+#include "MassArchetypeTypes.h"
+#include "MassEntitySubsystem.h"
+#include "MassArchetypeData.h"
+#include "MassCommandBuffer.h"
+
+
+//////////////////////////////////////////////////////////////////////
+// FArchetypeHandle
+bool FArchetypeHandle::operator==(const FMassArchetypeData* Other) const
+{
+	return DataPtr.Get() == Other;
+}
+
+uint32 GetTypeHash(const FArchetypeHandle& Instance)
+{
+	return GetTypeHash(Instance.DataPtr.Get());
+}
+
+//////////////////////////////////////////////////////////////////////
+// FArchetypeChunkCollection
+
+FArchetypeChunkCollection::FArchetypeChunkCollection(const FArchetypeHandle& InArchetype, TConstArrayView<FMassEntityHandle> InEntities)
+	: Archetype(InArchetype)
+{
+	check(InArchetype.IsValid());
+	const FMassArchetypeData& ArchetypeData = *InArchetype.DataPtr.Get();
+	const int32 NumEntitiesPerChunk = ArchetypeData.GetNumEntitiesPerChunk();
+
+	// InEntities has a real chance of not being sorted by AbsoluteIndex. We gotta fix that to optimize how we process the data 
+	TArray<int32> TrueIndices;
+	TrueIndices.AddUninitialized(InEntities.Num());
+	int32 i = 0;
+	for (const FMassEntityHandle& Entity : InEntities)
+	{
+		TrueIndices[i] = ArchetypeData.GetInternalIndexForEntity(Entity.Index);
+		++i;
+	}
+	TrueIndices.Sort();
+
+	// the following block of code is splitting up sorted AbsoluteIndices into 
+	// continuous chunks
+	int32 ChunkEnd = INDEX_NONE;
+	FChunkInfo DummyChunk;
+	FChunkInfo* SubChunkPtr = &DummyChunk;
+	int32 SubchunkLen = 0;
+	int32 PrevAbsoluteIndex = INDEX_NONE;
+	for (const int32 Index : TrueIndices)
+	{
+		// if run across a chunk border or run into an index discontinuity 
+		if (Index >= ChunkEnd || Index != (PrevAbsoluteIndex + 1))
+		{
+			SubChunkPtr->Length = SubchunkLen;
+			// note that both ChunkIndex and ChunkEnd will change only if AbsoluteIndex >= ChunkEnd
+			const int32 ChunkIndex = Index / NumEntitiesPerChunk;
+			ChunkEnd = (ChunkIndex + 1) * NumEntitiesPerChunk;
+			SubchunkLen = 0;
+			// new subchunk
+			const int32 SubchunkStart = Index % NumEntitiesPerChunk;
+			SubChunkPtr = &Chunks.Add_GetRef(FChunkInfo(ChunkIndex, SubchunkStart));
+		}
+		++SubchunkLen;
+		PrevAbsoluteIndex = Index;
+	}
+
+	SubChunkPtr->Length = SubchunkLen;
+}
+
+FArchetypeChunkCollection::FArchetypeChunkCollection(FArchetypeHandle& InArchetypeHandle)
+{
+	check(InArchetypeHandle.DataPtr.IsValid());
+	GatherChunksFromArchetype(InArchetypeHandle.DataPtr);
+}
+
+FArchetypeChunkCollection::FArchetypeChunkCollection(TSharedPtr<FMassArchetypeData>& InArchetype)
+{
+	GatherChunksFromArchetype(InArchetype);
+}
+
+void FArchetypeChunkCollection::GatherChunksFromArchetype(TSharedPtr<FMassArchetypeData>& InArchetype)
+{
+	check(InArchetype.IsValid());
+	Archetype.DataPtr = InArchetype;
+
+	const int32 ChunkCount = InArchetype->GetChunkCount();
+	Chunks.Reset(ChunkCount);
+	for (int32 i = 0; i < ChunkCount; ++i)
+	{
+		Chunks.Add(FChunkInfo(i));
+	}
+}
diff --git a/Engine/Plugins/Runtime/MassEntity/Source/MassEntity/Private/MassCommandBuffer.cpp b/Engine/Plugins/Runtime/MassEntity/Source/MassEntity/Private/MassCommandBuffer.cpp
new file mode 100644
--- /dev/null
+++ b/Engine/Plugins/Runtime/MassEntity/Source/MassEntity/Private/MassCommandBuffer.cpp
@@ -0,0 +1,196 @@
+// Copyright Epic Games, Inc. All Rights Reserved.
+
+#include "MassCommandBuffer.h"
+#include "HAL/IConsoleManager.h"
+#include "ProfilingDebugging/CsvProfiler.h"
+
+//////////////////////////////////////////////////////////////////////
+// FMassCommandBuffer
+
+CSV_DEFINE_CATEGORY(MassEntities, true);
+CSV_DEFINE_CATEGORY(MassEntitiesCounters, true);
+
+namespace UE::FLWCCommand {
+
+#if CSV_PROFILER
+bool bEnableDetailedStats = false;
+
+FAutoConsoleVariableRef CVarEnableDetailedCommandStats(TEXT("massentities.EnableCommandDetailedStats"), bEnableDetailedStats,
+	TEXT("Set to true create a dedicated stat per type of command."), ECVF_Default);
+
+/** CSV stat names */
+static FString DefaultName = TEXT("Command");
+static TArray<FName> CommandFNames;
+
+/** CSV custom stat names (ANSI) */
+static const int32 MaxNameLength = 64;
+typedef ANSICHAR ANSIName[MaxNameLength];
+static ANSIName DefaultANSIName = "Command";
+static TArray<ANSIName> CommandANSINames;
+
+/**
+ * Provides valid names for CSV profiling.
+ * @param Entry is the view on the command
+ * @param OutName is the name to use for csv custom stats
+ * @param OutANSIName is the name to use for csv stats
+ */
+void GetCommandStatNames(FStructView Entry, FString& OutName, ANSIName*& OutANSIName)
+{
+	OutANSIName = &DefaultANSIName;
+	OutName = DefaultName;
+	if (!bEnableDetailedStats ||
+		!ensureMsgf(Entry.GetScriptStruct() != nullptr, TEXT("Unable to get stat name for an invalid entry. Using default name.")))
+	{
+		return;
+	}
+
+	const FName CommandFName = Entry.GetScriptStruct()->GetFName();
+	OutName = CommandFName.ToString();
+
+	const int32 Index = CommandFNames.Find(CommandFName);
+	if (Index == INDEX_NONE)
+	{
+		CommandFNames.Emplace(CommandFName);
+		OutANSIName = &CommandANSINames.AddZeroed_GetRef();
+		// Use prefix for easier parsing in reports
+		//const FString CounterName = FString::Printf(TEXT("Num%s"), *OutName);
+		//FMemory::Memcpy(OutANSIName, StringCast<ANSICHAR>(*CounterName).Get(), FMath::Min(CounterName.Len(), MaxNameLength - 1) * sizeof(ANSICHAR));
+		FMemory::Memcpy(OutANSIName, StringCast<ANSICHAR>(*OutName).Get(), FMath::Min(OutName.Len(), MaxNameLength - 1) * sizeof(ANSICHAR));
+	}
+	else
+	{
+		OutANSIName = &CommandANSINames[Index];
+	}
+}
+
+#endif
+} // UE::FLWCCommand
+
+void FMassCommandBuffer::ReplayBufferAgainstSystem(UMassEntitySubsystem* System)
+{
+	check(System);
+
+	UE_MT_SCOPED_WRITE_ACCESS(PendingCommandsDetector);
+	PendingCommands.ForEach([System](FStructView Entry)
+	{
+		const FCommandBufferEntryBase* Command = Entry.GetPtr<FCommandBufferEntryBase>();
+		checkf(Command, TEXT("Either the entry is null or the command does not derive from FCommandBufferEntryBase"));
+
+#if CSV_PROFILER
+		using namespace UE::FLWCCommand;
+
+		// Extract name (default or detailed)
+		ANSIName* ANSIName = &DefaultANSIName;
+		FString Name = DefaultName;
+		GetCommandStatNames(Entry, Name, ANSIName);
+
+		// Push stats
+		FScopedCsvStat ScopedCsvStat(*ANSIName, CSV_CATEGORY_INDEX(MassEntities));
+		FCsvProfiler::RecordCustomStat(*Name, CSV_CATEGORY_INDEX(MassEntitiesCounters), 1, ECsvCustomStatOp::Accumulate);
+#endif // CSV_PROFILER
+
+		Command->Execute(*System);
+	});
+
+	// Using Clear() instead of Reset(), as otherwise the chunks moved into the PendingCommands in MoveAppend() can accumulate.
+	PendingCommands.Clear();
+
+	if (EntitiesToDestroy.Num())
+	{
+		TRACE_CPUPROFILER_EVENT_SCOPE_STR("ECS Deferred Commands Destroy Entities");
+		System->BatchDestroyEntities(MakeArrayView(EntitiesToDestroy));
+		EntitiesToDestroy.Reset();
+	}
+}
+
+void FMassCommandBuffer::MoveAppend(FMassCommandBuffer& Other)
+{
+	// @todo optimize, there surely a way to do faster then this.
+	UE_MT_SCOPED_READ_ACCESS(Other.PendingCommandsDetector);
+	if (Other.GetPendingCommandsCount() || Other.EntitiesToDestroy.Num())
+	{
+		FScopeLock Lock(&AppendingCommandsCS);
+		UE_MT_SCOPED_WRITE_ACCESS(PendingCommandsDetector);
+		PendingCommands.Append(MoveTemp(Other.PendingCommands));
+		EntitiesToDestroy.Append(MoveTemp(Other.EntitiesToDestroy));
+	}
+}
+
+//////////////////////////////////////////////////////////////////////
+// Command implementations
+
+void FCommandAddComponentInstance::Execute(UMassEntitySubsystem& System) const
+{
+	System.AddComponentInstanceListToEntity(TargetEntity, MakeArrayView(&Struct, 1));
+}
+
+void FCommandSwapTags::Execute(UMassEntitySubsystem& System) const
+{
+	if (System.IsEntityValid(TargetEntity) == false)
+	{
+		return;
+}
+
+	if (OldTagType && NewTagType)
+	{
+		System.SwapTagsForEntity(TargetEntity, OldTagType, NewTagType);
+	}
+	else if (OldTagType)
+	{
+		System.RemoveTagFromEntity(TargetEntity, OldTagType);
+	}
+	else if (NewTagType)
+	{
+		System.AddTagToEntity(TargetEntity, NewTagType);
+	}
+}
+
+void FBuildEntityFromComponentInstance::Execute(UMassEntitySubsystem& System) const
+{
+	System.BuildEntity(TargetEntity, MakeArrayView(&Struct, 1));
+}
+
+void FBuildEntityFromComponentInstances::Execute(UMassEntitySubsystem& System) const
+{
+	System.BuildEntity(TargetEntity, Instances);
+}
+
+void FCommandAddComponent::Execute(UMassEntitySubsystem& System) const
+{
+	System.AddComponentToEntity(TargetEntity, StructParam);
+}
+
+void FCommandRemoveComponent::Execute(UMassEntitySubsystem& System) const
+{
+	System.RemoveComponentFromEntity(TargetEntity, StructParam);
+}
+
+void FCommandAddComponentList::Execute(UMassEntitySubsystem& System) const
+{
+	System.AddComponentListToEntity(TargetEntity, ComponentList);
+}
+
+void FCommandRemoveComponentList::Execute(UMassEntitySubsystem& System) const
+{
+	System.RemoveComponentListFromEntity(TargetEntity, ComponentList);
+}
+
+void FCommandAddTag::Execute(UMassEntitySubsystem & System) const 
+{
+	if (System.IsEntityValid(TargetEntity) == false)
+{
+		return;
+	}
+
+	System.AddTagToEntity(TargetEntity, StructParam);
+}
+
+void FCommandRemoveTag::Execute(UMassEntitySubsystem& System) const
+{
+	System.RemoveTagFromEntity(TargetEntity, StructParam);
+}
+
+void FCommandRemoveComposition::Execute(UMassEntitySubsystem& System) const
+{
+	System.RemoveCompositionFromEntity(TargetEntity, Descriptor);
+}
diff --git a/Engine/Plugins/Runtime/MassEntity/Source/MassEntity/Private/MassEntityQuery.cpp b/Engine/Plugins/Runtime/MassEntity/Source/MassEntity/Private/MassEntityQuery.cpp
new file mode 100644
--- /dev/null
+++ b/Engine/Plugins/Runtime/MassEntity/Source/MassEntity/Private/MassEntityQuery.cpp
@@ -0,0 +1,308 @@
+// Copyright Epic Games, Inc. All Rights Reserved.
+
+#include "MassEntityQuery.h"
+#include "MassEntitySubsystem.h"
+#include "MassArchetypeData.h"
+#include "MassCommandBuffer.h"
+#include "VisualLogger/VisualLogger.h"
+#include "Async/ParallelFor.h"
+
+//////////////////////////////////////////////////////////////////////
+// FMassEntityQuery
+
+
+FMassEntityQuery::FMassEntityQuery()
+{
+	ReadCommandlineParams();
+}
+
+FMassEntityQuery::FMassEntityQuery(std::initializer_list<UScriptStruct*> InitList)
+	: FMassEntityQuery()
+{
+	for (const UScriptStruct* ComponentType : InitList)
+	{
+		AddRequirement(ComponentType, EMassFragmentAccess::ReadWrite, EMassFragmentPresence::All);
+	}
+}
+
+FMassEntityQuery::FMassEntityQuery(TConstArrayView<const UScriptStruct*> InitList)
+	: FMassEntityQuery()
+{
+	for (const UScriptStruct* ComponentType : InitList)
+	{
+		AddRequirement(ComponentType, EMassFragmentAccess::ReadWrite, EMassFragmentPresence::All);
+	}
+}
+
+void FMassEntityQuery::ReadCommandlineParams()
+{
+	int AllowParallelQueries = -1;
+	if (FParse::Value(FCommandLine::Get(), TEXT("ParallelMassQueries="), AllowParallelQueries))
+	{
+		bAllowParallelExecution = (AllowParallelQueries != 0);
+	}
+}
+
+void FMassEntityQuery::SortRequirements()
+{
+	// we're sorting the Requirements the same way ArchetypeData's ComponentConfig is sorted (see FMassArchetypeData::Initialize)
+	// so that when we access ArchetypeData.ComponentConfigs in FMassArchetypeData::BindRequirementsWithMapping
+	// (via GetComponentData call) the access is sequential (i.e. not random) and there's a higher chance the memory
+	// ComponentConfigs we want to access have already been fetched and are available in processor cache.
+	Requirements.Sort(FMassSorterOperator<FMassFragmentRequirement>());
+	ChunkRequirements.Sort(FMassSorterOperator<FMassFragmentRequirement>());
+}
+
+void FMassEntityQuery::CacheArchetypes(UMassEntitySubsystem& InEntitySubsystem)
+{
+	const uint32 InEntitySubsystemHash = PointerHash(&InEntitySubsystem);
+	if (EntitySubsystemHash != InEntitySubsystemHash || InEntitySubsystem.GetArchetypeDataVersion() != ArchetypeDataVersion)
+	{
+		if (CheckValidity())
+		{
+			SortRequirements();
+
+			EntitySubsystemHash = InEntitySubsystemHash;
+			ValidArchetypes.Reset();
+			InEntitySubsystem.GetValidArchetypes(*this, ValidArchetypes);
+			ArchetypeDataVersion = InEntitySubsystem.GetArchetypeDataVersion();
+
+			TRACE_CPUPROFILER_EVENT_SCOPE_STR("Pipe RequirementsBinding")
+			const TConstArrayView<FMassFragmentRequirement> LocalRequirements = GetRequirements();
+			ArchetypeComponentMapping.Reset(ValidArchetypes.Num());
+			ArchetypeComponentMapping.AddDefaulted(ValidArchetypes.Num());
+			for (int i = 0; i < ValidArchetypes.Num(); ++i)
+			{
+				ValidArchetypes[i].DataPtr->GetRequirementsComponentMapping(LocalRequirements, ArchetypeComponentMapping[i].EntityComponents);
+				if (ChunkRequirements.Num())
+				{
+					ValidArchetypes[i].DataPtr->GetRequirementsChunkComponentMapping(ChunkRequirements, ArchetypeComponentMapping[i].ChunkComponents);
+				}
+			}
+		}
+		else
+		{
+			UE_VLOG_UELOG(&InEntitySubsystem, LogAggregateTicking, Error, TEXT("FMassEntityQuery::CacheArchetypes: requirements not valid: %s"), *DebugGetDescription());
+		}
+	}
+}
+
+bool FMassEntityQuery::CheckValidity() const
+{
+	return RequiredAllComponents.IsEmpty() == false || RequiredAnyComponents.IsEmpty() == false || RequiredOptionalComponents.IsEmpty() == false;
+}
+
+bool FMassEntityQuery::DoesArchetypeMatchRequirements(const FArchetypeHandle& ArchetypeHandle) const
+{
+	check(ArchetypeHandle.IsValid());
+	const FMassArchetypeData* Archetype = ArchetypeHandle.DataPtr.Get();
+	CA_ASSUME(Archetype);
+
+	for (const FMassFragmentRequirement& Requirement : GetRequirements())
+	{
+		check(Requirement.StructType);
+		if (Requirement.IsOptional())
+		{
+			// at this stage we don't care if the optional is present
+			continue;
+		}
+
+		const bool bNeedsComponent = (Requirement.Presence == EMassFragmentPresence::All);
+		const bool bHasComponent = Archetype->HasComponentType(Requirement.StructType);
+		if (bNeedsComponent != bHasComponent)
+		{
+			return false;
+		}
+	}
+	return true;
+}
+
+void FMassEntityQuery::ForEachEntityChunk(const FArchetypeChunkCollection& Chunks, UMassEntitySubsystem& EntitySubsystem, FMassExecutionContext& ExecutionContext, const FMassExecuteFunction& ExecuteFunction)
+{
+	// mz@todo I don't like that we're copying data here.
+	ExecutionContext.SetChunkCollection(Chunks);
+	ForEachEntityChunk(EntitySubsystem, ExecutionContext, ExecuteFunction);
+	ExecutionContext.ClearChunkCollection();
+}
+
+void FMassEntityQuery::ForEachEntityChunk(UMassEntitySubsystem& EntitySubsystem, FMassExecutionContext& ExecutionContext, const FMassExecuteFunction& ExecuteFunction)
+{
+#if WITH_MASSENTITY_DEBUG
+	int32 NumEntitiesToProcess = 0;
+#endif
+
+	// if there's a chunk collection set by the external code - use that
+	if (ExecutionContext.GetChunkCollection().IsSet())
+	{
+		// verify the archetype matches requirements
+		if (DoesArchetypeMatchRequirements(ExecutionContext.GetChunkCollection().GetArchetype()) == false)
+		{
+			// mz@todo add a unit test for this message
+			UE_VLOG_UELOG(&EntitySubsystem, LogAggregateTicking, Error, TEXT("Attempted to execute FMassEntityQuery with an incompatible Archetype"));
+			return;
+		}
+		ExecutionContext.SetRequirements(Requirements, ChunkRequirements);
+		ExecutionContext.GetChunkCollection().GetArchetype().DataPtr->ExecuteFunction(ExecutionContext, ExecuteFunction, {}, ExecutionContext.GetChunkCollection());
+#if WITH_MASSENTITY_DEBUG
+		NumEntitiesToProcess = ExecutionContext.GetChunkCollection().GetArchetype().DataPtr->GetNumEntities();
+#endif
+	}
+	else
+	{
+		CacheArchetypes(EntitySubsystem);
+		// it's important to set requirements after caching archetypes due to that call potentially sorting the requirements and the order is relevant here.
+		ExecutionContext.SetRequirements(Requirements, ChunkRequirements);
+
+		for (int i = 0; i < ValidArchetypes.Num(); ++i)
+		{
+			FArchetypeHandle& Archetype = ValidArchetypes[i];
+			check(Archetype.IsValid());
+			Archetype.DataPtr->ExecuteFunction(ExecutionContext, ExecuteFunction, ArchetypeComponentMapping[i], ChunkCondition);
+			ExecutionContext.ClearComponentViews();
+#if WITH_MASSENTITY_DEBUG
+			NumEntitiesToProcess += Archetype.DataPtr->GetNumEntities();
+#endif
+		}
+	}
+
+#if WITH_MASSENTITY_DEBUG
+	// Not using VLOG to be thread safe
+	UE_CLOG(!ExecutionContext.DebugGetExecutionDesc().IsEmpty(), LogAggregateTicking, VeryVerbose,
+		TEXT("%s: %d entities sent for processing"), *ExecutionContext.DebugGetExecutionDesc(), NumEntitiesToProcess);
+#endif
+
+	ExecutionContext.ClearExecutionData();
+	ExecutionContext.FlushDeferred(EntitySubsystem);
+}
+
+void FMassEntityQuery::ParallelForEachEntityChunk(UMassEntitySubsystem& EntitySubsystem, FMassExecutionContext& ExecutionContext, const FMassExecuteFunction& ExecuteFunction)
+{
+	if (bAllowParallelExecution == false)
+	{
+		ForEachEntityChunk(EntitySubsystem, ExecutionContext, ExecuteFunction);
+		return;
+	}
+
+	struct FChunkJob
+	{
+		FMassArchetypeData& Archetype;
+		const int32 ArchetypeIndex;
+		const FArchetypeChunkCollection::FChunkInfo ChunkInfo;
+	};
+	TArray<FChunkJob> Jobs;
+
+	// if there's a chunk collection set by the external code - use that
+	if (ExecutionContext.GetChunkCollection().IsSet())
+	{
+		// verify the archetype matches requirements
+		FArchetypeHandle ArchetypeHandle = ExecutionContext.GetChunkCollection().GetArchetype();
+		if (DoesArchetypeMatchRequirements(ArchetypeHandle) == false)
+		{
+			// mz@todo add a unit test for this message
+			UE_VLOG_UELOG(&EntitySubsystem, LogAggregateTicking, Error, TEXT("Attempted to execute FMassEntityQuery with an incompatible Archetype"));
+			return;
+		}
+
+		ExecutionContext.SetRequirements(Requirements, ChunkRequirements);
+		check(ArchetypeHandle.IsValid());
+		FMassArchetypeData& ArchetypeRef = *ArchetypeHandle.DataPtr.Get();
+		const FArchetypeChunkCollection AsChunkCollection(ArchetypeHandle.DataPtr);
+		for (const FArchetypeChunkCollection::FChunkInfo& ChunkInfo : AsChunkCollection.GetChunks())
+		{
+			Jobs.Add({ ArchetypeRef, INDEX_NONE, ChunkInfo });
+		}
+	}
+	else
+	{
+		CacheArchetypes(EntitySubsystem);
+		ExecutionContext.SetRequirements(Requirements, ChunkRequirements);
+		for (int ArchetypeIndex = 0; ArchetypeIndex < ValidArchetypes.Num(); ++ArchetypeIndex)
+		{
+			FArchetypeHandle& Archetype = ValidArchetypes[ArchetypeIndex];
+			check(Archetype.IsValid());
+			FMassArchetypeData& ArchetypeRef = *Archetype.DataPtr.Get();
+			const FArchetypeChunkCollection AsChunkCollection(Archetype.DataPtr);
+			for (const FArchetypeChunkCollection::FChunkInfo& ChunkInfo : AsChunkCollection.GetChunks())
+			{
+				Jobs.Add({ArchetypeRef, ArchetypeIndex, ChunkInfo});
+			}
+		}
+	}
+	// ExecutionContext passed by copy on purpose
+	ParallelFor(Jobs.Num(), [this, ExecutionContext, &ExecuteFunction, Jobs](const int32 JobIndex)
+	{
+		Jobs[JobIndex].Archetype.ExecutionFunctionForChunk(ExecutionContext, ExecuteFunction
+			, Jobs[JobIndex].ArchetypeIndex != INDEX_NONE ? ArchetypeComponentMapping[Jobs[JobIndex].ArchetypeIndex] : FMassQueryRequirementIndicesMapping()
+			, Jobs[JobIndex].ChunkInfo
+			, ChunkCondition);
+	});
+
+	ExecutionContext.ClearExecutionData();
+	ExecutionContext.FlushDeferred(EntitySubsystem);
+}
+
+int32 FMassEntityQuery::GetMatchingEntitiesNum(UMassEntitySubsystem& InEntitySubsystem)
+{
+	CacheArchetypes(InEntitySubsystem);
+	int32 TotalEntities = 0;
+	for (FArchetypeHandle& ArchetypeHandle : ValidArchetypes)
+	{
+		if (const FMassArchetypeData* Archetype = ArchetypeHandle.DataPtr.Get())
+		{
+			TotalEntities += Archetype->GetNumEntities();
+		}
+	}
+	return TotalEntities;
+}
+
+bool FMassEntityQuery::HasMatchingEntities(UMassEntitySubsystem& InEntitySubsystem)
+{
+	CacheArchetypes(InEntitySubsystem);
+
+	for (FArchetypeHandle& ArchetypeHandle : ValidArchetypes)
+	{
+		const FMassArchetypeData* Archetype = ArchetypeHandle.DataPtr.Get();
+		if (Archetype && Archetype->GetNumEntities() > 0)
+		{
+			return true;
+		}
+	}
+	return false;
+}
+
+FString FMassEntityQuery::DebugGetDescription() const
+{
+#if WITH_MASSENTITY_DEBUG
+	TStringBuilder<256> StringBuilder;
+	StringBuilder.Append(TEXT("<"));
+
+	bool bNeedsComma = false;
+	for (const FMassFragmentRequirement& Requirement : Requirements)
+	{
+		if (bNeedsComma)
+		{
+			StringBuilder.Append(TEXT(","));
+		}
+		StringBuilder.Append(*Requirement.DebugGetDescription());
+		bNeedsComma = true;
+	}
+
+	StringBuilder.Append(TEXT(">"));
+	return StringBuilder.ToString();
+#else
+	return {};
+#endif
+}
+
+//////////////////////////////////////////////////////////////////////
+// FMassFragmentRequirement
+
+FString FMassFragmentRequirement::DebugGetDescription() const
+{
+#if WITH_MASSENTITY_DEBUG
+	return FString::Printf(TEXT("%s%s[%s]"), IsOptional() ? TEXT("?") : (Presence == EMassFragmentPresence::None ? TEXT("-") : TEXT("+"))
+		, *GetNameSafe(StructType), *UE::AggregateTicking::DebugGetComponentAccessString(AccessMode));
+#else
+	return {};
+#endif
+}
diff --git a/Engine/Plugins/Runtime/MassEntity/Source/MassEntity/Private/MassEntitySubsystem.cpp b/Engine/Plugins/Runtime/MassEntity/Source/MassEntity/Private/MassEntitySubsystem.cpp
--- a/Engine/Plugins/Runtime/MassEntity/Source/MassEntity/Private/MassEntitySubsystem.cpp
+++ b/Engine/Plugins/Runtime/MassEntity/Source/MassEntity/Private/MassEntitySubsystem.cpp
@@ -1,8 +1,8 @@
 // Copyright Epic Games, Inc. All Rights Reserved.
 
 #include "MassEntitySubsystem.h"
-#include "ArchetypeData.h"
-#include "LWCCommandBuffer.h"
+#include "MassArchetypeData.h"
+#include "MassCommandBuffer.h"
 #include "HAL/IConsoleManager.h"
 #include "Engine/World.h"
 #include "UObject/UObjectIterator.h"
diff --git a/Engine/Plugins/Runtime/MassEntity/Source/MassEntity/Private/MassEntityTypes.cpp b/Engine/Plugins/Runtime/MassEntity/Source/MassEntity/Private/MassEntityTypes.cpp
deleted file mode 100644
diff --git a/Engine/Plugins/Runtime/MassEntity/Source/MassEntity/Private/MassEntityView.cpp b/Engine/Plugins/Runtime/MassEntity/Source/MassEntity/Private/MassEntityView.cpp
new file mode 100644
--- /dev/null
+++ b/Engine/Plugins/Runtime/MassEntity/Source/MassEntity/Private/MassEntityView.cpp
@@ -0,0 +1,49 @@
+// Copyright Epic Games, Inc. All Rights Reserved.
+
+#include "MassEntityView.h"
+#include "MassEntitySubsystem.h"
+#include "MassArchetypeData.h"
+
+
+//////////////////////////////////////////////////////////////////////
+// FMassEntityView
+FMassEntityView::FMassEntityView(const FArchetypeHandle& ArchetypeHandle, FMassEntityHandle InEntity)
+{
+	Entity = InEntity;
+	check(ArchetypeHandle.IsValid());
+	Archetype = ArchetypeHandle.DataPtr.Get();
+	EntityHandle = Archetype->MakeEntityHandle(Entity);
+}
+
+FMassEntityView::FMassEntityView(const UMassEntitySubsystem& EntitySubsystem, FMassEntityHandle InEntity)
+{
+	Entity = InEntity;
+	const FArchetypeHandle ArchetypeHandle = EntitySubsystem.GetArchetypeForEntity(Entity);
+	check(ArchetypeHandle.IsValid());
+	Archetype = ArchetypeHandle.DataPtr.Get();
+	EntityHandle = Archetype->MakeEntityHandle(Entity);
+}
+
+void* FMassEntityView::GetComponentPtr(const UScriptStruct& ComponentType) const
+{
+	checkSlow(Archetype && EntityHandle.IsValid());
+	if (const int32* ComponentIndex = Archetype->GetComponentIndex(&ComponentType))
+	{
+		// failing the below Find means given entity's archetype is missing given ComponentType
+		return Archetype->GetComponentData(*ComponentIndex, EntityHandle);
+	}
+	return nullptr;
+}
+
+void* FMassEntityView::GetComponentPtrChecked(const UScriptStruct& ComponentType) const
+{
+	checkSlow(Archetype && EntityHandle.IsValid());
+	const int32 ComponentIndex = Archetype->GetComponentIndexChecked(&ComponentType);
+	return Archetype->GetComponentData(ComponentIndex, EntityHandle);
+}
+
+bool FMassEntityView::HasTag(const UScriptStruct& TagType) const
+{
+	checkSlow(Archetype && EntityHandle.IsValid());
+	return Archetype->HasTagType(&TagType);
+}
diff --git a/Engine/Plugins/Runtime/MassEntity/Source/MassEntity/Private/MassExecutor.cpp b/Engine/Plugins/Runtime/MassEntity/Source/MassEntity/Private/MassExecutor.cpp
--- a/Engine/Plugins/Runtime/MassEntity/Source/MassEntity/Private/MassExecutor.cpp
+++ b/Engine/Plugins/Runtime/MassEntity/Source/MassEntity/Private/MassExecutor.cpp
@@ -1,9 +1,9 @@
 // Copyright Epic Games, Inc. All Rights Reserved.
 
 #include "MassExecutor.h"
-#include "MassEntityTypes.h"
+#include "MassProcessingTypes.h"
 #include "MassProcessor.h"
-#include "LWCCommandBuffer.h"
+#include "MassCommandBuffer.h"
 #include "ProfilingDebugging/CpuProfilerTrace.h"
 
 namespace UE::Pipe::Executor
diff --git a/Engine/Plugins/Runtime/MassEntity/Source/MassEntity/Private/MassProcessingPhase.cpp b/Engine/Plugins/Runtime/MassEntity/Source/MassEntity/Private/MassProcessingPhase.cpp
--- a/Engine/Plugins/Runtime/MassEntity/Source/MassEntity/Private/MassProcessingPhase.cpp
+++ b/Engine/Plugins/Runtime/MassEntity/Source/MassEntity/Private/MassProcessingPhase.cpp
@@ -1,7 +1,7 @@
 // Copyright Epic Games, Inc. All Rights Reserved.
 
 #include "MassProcessingPhase.h"
-#include "MassEntityTypes.h"
+#include "MassProcessingTypes.h"
 #include "MassEntityDebug.h"
 #include "MassEntitySettings.h"
 #include "MassProcessor.h"
@@ -9,7 +9,7 @@
 #include "MassEntitySubsystem.h"
 #include "VisualLogger/VisualLogger.h"
 #include "Engine/World.h"
-#include "LWCCommandBuffer.h"
+#include "MassCommandBuffer.h"
 
 #define LOCTEXT_NAMESPACE "Mass"
 
diff --git a/Engine/Plugins/Runtime/MassEntity/Source/MassEntity/Private/MassProcessingTypes.cpp b/Engine/Plugins/Runtime/MassEntity/Source/MassEntity/Private/MassProcessingTypes.cpp
new file mode 100644
--- /dev/null
+++ b/Engine/Plugins/Runtime/MassEntity/Source/MassEntity/Private/MassProcessingTypes.cpp
@@ -0,0 +1,225 @@
+// Copyright Epic Games, Inc. All Rights Reserved.
+
+#include "MassProcessingTypes.h"
+#include "MassProcessor.h"
+#include "MassSchematic.h"
+#include "Misc/OutputDevice.h"
+#include "MassEntityUtils.h"
+#include "VisualLogger/VisualLogger.h"
+#include "MassEntityDebug.h"
+
+DEFINE_LOG_CATEGORY(LogPipe);
+
+//----------------------------------------------------------------------//
+//  FMassProcessingContext
+//----------------------------------------------------------------------//
+FMassProcessingContext::FMassProcessingContext(UMassEntitySubsystem& InEntities, const float InDeltaSeconds)
+	: EntitySubsystem(&InEntities), DeltaSeconds(InDeltaSeconds)
+{
+
+}
+
+//----------------------------------------------------------------------//
+//  FMassRuntimePipeline
+//----------------------------------------------------------------------//
+void FMassRuntimePipeline::Reset()
+{
+	Processors.Reset();
+}
+
+void FMassRuntimePipeline::Initialize(UObject& Owner)
+{
+	for (UMassProcessor* Proc : Processors)
+	{
+		if (Proc)
+		{
+			REDIRECT_OBJECT_TO_VLOG(Proc, &Owner);
+			Proc->Initialize(Owner);
+		}
+	}
+}
+
+void FMassRuntimePipeline::SetProcessors(TArray<UMassProcessor*>&& InProcessors)
+{
+	Processors = InProcessors;
+}
+
+void FMassRuntimePipeline::InitializeFromSchematics(TConstArrayView<TSoftObjectPtr<UMassSchematic>> Schematics, UObject& InOwner)
+{
+	Reset();
+	
+	// @todo we'll sometimes end up with duplicated MassProcessors in the resulting array. We need to come up with a consistent policy for handling that 
+	for (const TSoftObjectPtr<UMassSchematic>& Schematic : Schematics)
+	{
+		UMassSchematic* Pipe = Schematic.LoadSynchronous();
+		if (Pipe)
+		{
+			AppendOrOverrideRuntimeProcessorCopies(Pipe->GetProcessors(), InOwner);
+		}
+		else
+		{
+			UE_LOG(LogPipe, Error, TEXT("Unable to resolve MassSchematic %s while creating FMassRuntimePipeline"), *Schematic.GetLongPackageName());
+		}
+	}
+
+	Initialize(InOwner);
+}
+
+void FMassRuntimePipeline::CreateFromArray(TConstArrayView<const UMassProcessor*> InProcessors, UObject& InOwner)
+{
+	Reset();
+	AppendOrOverrideRuntimeProcessorCopies(InProcessors, InOwner);
+}
+
+void FMassRuntimePipeline::InitializeFromArray(TConstArrayView<const UMassProcessor*> InProcessors, UObject& InOwner)
+{
+	CreateFromArray(InProcessors, InOwner);
+	Initialize(InOwner);
+}
+
+void FMassRuntimePipeline::InitializeFromClassArray(TConstArrayView<TSubclassOf<UMassProcessor>> InProcessorClasses, UObject& InOwner)
+{
+	Reset();
+
+	const UWorld* World = InOwner.GetWorld();
+	const EProcessorExecutionFlags WorldExecutionFlags = World ? UE::Pipe::Utils::GetProcessorExecutionFlagsForWold(*World) : EProcessorExecutionFlags::All;
+
+	for (const TSubclassOf<UMassProcessor>& ProcessorClass : InProcessorClasses)
+	{
+		if (ProcessorClass)
+		{
+			UMassProcessor* CDO = ProcessorClass.GetDefaultObject();
+			if (CDO && CDO->ShouldExecute(WorldExecutionFlags))
+			{
+				UMassProcessor* ProcInstance = NewObject<UMassProcessor>(&InOwner, ProcessorClass);
+				Processors.Add(ProcInstance);
+			}
+			else
+			{
+				UE_CVLOG(CDO, &InOwner, LogPipe, Log, TEXT("Skipping %s due to ExecutionFlags"), *CDO->GetName());
+			}
+		}
+	}
+
+	Initialize(InOwner);
+}
+
+bool FMassRuntimePipeline::HasProcessorOfExactClass(TSubclassOf<UMassProcessor> InClass) const
+{
+	UClass* TestClass = InClass.Get();
+	return Processors.FindByPredicate([TestClass](const UMassProcessor* Proc){ return Proc != nullptr && Proc->GetClass() == TestClass; }) != nullptr;
+}
+
+void FMassRuntimePipeline::AppendUniqueRuntimeProcessorCopies(TConstArrayView<const UMassProcessor*> InProcessors, UObject& InOwner)
+{
+	const UWorld* World = InOwner.GetWorld();
+	const EProcessorExecutionFlags WorldExecutionFlags = World ? UE::Pipe::Utils::GetProcessorExecutionFlagsForWold(*World) : EProcessorExecutionFlags::All;
+	const int32 StartingCount = Processors.Num();
+		
+	for (const UMassProcessor* Proc : InProcessors)
+	{
+		if (Proc && Proc->ShouldExecute(WorldExecutionFlags)
+			&& (Proc->AllowDuplicates() || (HasProcessorOfExactClass(Proc->GetClass()) == false)))
+		{
+			// unfortunately the const cast is required since NewObject doesn't support const Template object
+			UMassProcessor* ProcCopy = NewObject<UMassProcessor>(&InOwner, Proc->GetClass(), FName(), RF_NoFlags, const_cast<UMassProcessor*>(Proc));
+			Processors.Add(ProcCopy);
+		}
+#if WITH_MASSENTITY_DEBUG
+		else if (Proc)
+		{
+			if (Proc->ShouldExecute(WorldExecutionFlags) == false)
+			{
+				UE_VLOG(&InOwner, LogPipe, Log, TEXT("Skipping %s due to ExecutionFlags"), *Proc->GetName());
+			}
+			else if (Proc->AllowDuplicates() == false)
+			{
+				UE_VLOG(&InOwner, LogPipe, Log, TEXT("Skipping %s due to it being a duplicate"), *Proc->GetName());
+			}
+		}
+#endif // WITH_MASSENTITY_DEBUG
+	}
+
+	for (int32 NewProcIndex = StartingCount; NewProcIndex < Processors.Num(); ++NewProcIndex)
+	{
+		UMassProcessor* Proc = Processors[NewProcIndex];
+		check(Proc);
+		REDIRECT_OBJECT_TO_VLOG(Proc, &InOwner);
+		Proc->Initialize(InOwner);
+	}
+}
+
+void FMassRuntimePipeline::AppendOrOverrideRuntimeProcessorCopies(TConstArrayView<const UMassProcessor*> InProcessors, UObject& InOwner)
+{
+	const UWorld* World = InOwner.GetWorld();
+	const EProcessorExecutionFlags WorldExecutionFlags = World ? UE::Pipe::Utils::GetProcessorExecutionFlagsForWold(*World) : EProcessorExecutionFlags::All;
+	
+	for (const UMassProcessor* Proc : InProcessors)
+	{
+		if (Proc && Proc->ShouldExecute(WorldExecutionFlags))
+		{
+			// unfortunately the const cast is required since NewObject doesn't support const Template object
+			UMassProcessor* ProcCopy = NewObject<UMassProcessor>(&InOwner, Proc->GetClass(), FName(), RF_NoFlags, const_cast<UMassProcessor*>(Proc));
+			check(ProcCopy);
+
+			if (ProcCopy->AllowDuplicates())
+			{
+				// we don't care if there are instances of this class in Processors already
+				Processors.Add(ProcCopy);
+			}
+			else 
+			{
+				const UClass* TestClass = Proc->GetClass();
+				UMassProcessor** PrevProcessor = Processors.FindByPredicate([TestClass, ProcCopy](const UMassProcessor* Proc) {
+					return Proc != nullptr && Proc->GetClass() == TestClass;
+				});
+
+				if (PrevProcessor)
+				{
+					*PrevProcessor = ProcCopy;
+				}
+				else
+				{
+					Processors.Add(ProcCopy);
+				}
+			}
+		}
+		else
+		{
+			UE_CVLOG(Proc, &InOwner, LogPipe, Log, TEXT("Skipping %s due to ExecutionFlags"), *Proc->GetName());
+		}
+	}
+}
+
+void FMassRuntimePipeline::AppendProcessor(UMassProcessor& Processor)
+{
+	Processors.Add(&Processor);
+}
+
+UMassCompositeProcessor* FMassRuntimePipeline::FindTopLevelGroupByName(FName GroupName)
+{
+	for (UMassProcessor* Processor : Processors)
+	{
+		UMassCompositeProcessor* CompositeProcessor = Cast<UMassCompositeProcessor>(Processor);
+		if (CompositeProcessor && CompositeProcessor->GetGroupName() == GroupName)
+		{
+			return CompositeProcessor;
+		}
+	}
+	return nullptr;
+}
+
+void FMassRuntimePipeline::DebugOutputDescription(FOutputDevice& Ar) const
+{
+	UE::Pipe::Debug::DebugOutputDescription(Processors, Ar);
+}
+
+uint32 GetTypeHash(const FMassRuntimePipeline& Instance)
+{ 
+	uint32 Hash = 0;
+	for (const UMassProcessor* Proc : Instance.Processors)
+	{
+		Hash = HashCombine(Hash, PointerHash(Proc));
+	}
+	return Hash;
+}
diff --git a/Engine/Plugins/Runtime/MassEntity/Source/MassEntity/Private/MassProcessor.cpp b/Engine/Plugins/Runtime/MassEntity/Source/MassEntity/Private/MassProcessor.cpp
--- a/Engine/Plugins/Runtime/MassEntity/Source/MassEntity/Private/MassProcessor.cpp
+++ b/Engine/Plugins/Runtime/MassEntity/Source/MassEntity/Private/MassProcessor.cpp
@@ -5,7 +5,7 @@
 #include "MassProcessorDependencySolver.h"
 #include "VisualLogger/VisualLogger.h"
 #include "Engine/World.h"
-#include "LWCCommandBuffer.h"
+#include "MassCommandBuffer.h"
 
 DECLARE_CYCLE_STAT(TEXT("MassProcessor Group Completed"), Pipe_GroupCompletedTask, STATGROUP_TaskGraphTasks);
 
diff --git a/Engine/Plugins/Runtime/MassEntity/Source/MassEntity/Private/MassProcessorDependencySolver.cpp b/Engine/Plugins/Runtime/MassEntity/Source/MassEntity/Private/MassProcessorDependencySolver.cpp
--- a/Engine/Plugins/Runtime/MassEntity/Source/MassEntity/Private/MassProcessorDependencySolver.cpp
+++ b/Engine/Plugins/Runtime/MassEntity/Source/MassEntity/Private/MassProcessorDependencySolver.cpp
@@ -1,7 +1,7 @@
 // Copyright Epic Games, Inc. All Rights Reserved.
 
 #include "MassProcessorDependencySolver.h"
-#include "MassEntityTypes.h"
+#include "MassProcessingTypes.h"
 #include "MassProcessor.h"
 #include "Logging/MessageLog.h"
 #include "HAL/FileManager.h"
diff --git a/Engine/Plugins/Runtime/MassEntity/Source/MassEntity/Public/ArchetypeTypes.h b/Engine/Plugins/Runtime/MassEntity/Source/MassEntity/Public/ArchetypeTypes.h
deleted file mode 100644
diff --git a/Engine/Plugins/Runtime/MassEntity/Source/MassEntity/Public/EntityQuery.h b/Engine/Plugins/Runtime/MassEntity/Source/MassEntity/Public/EntityQuery.h
deleted file mode 100644
diff --git a/Engine/Plugins/Runtime/MassEntity/Source/MassEntity/Public/EntityView.h b/Engine/Plugins/Runtime/MassEntity/Source/MassEntity/Public/EntityView.h
deleted file mode 100644
diff --git a/Engine/Plugins/Runtime/MassEntity/Source/MassEntity/Public/LWCCommandBuffer.h b/Engine/Plugins/Runtime/MassEntity/Source/MassEntity/Public/LWCCommandBuffer.h
deleted file mode 100644
diff --git a/Engine/Plugins/Runtime/MassEntity/Source/MassEntity/Public/LWComponentTypes.h b/Engine/Plugins/Runtime/MassEntity/Source/MassEntity/Public/LWComponentTypes.h
deleted file mode 100644
diff --git a/Engine/Plugins/Runtime/MassEntity/Source/MassEntity/Public/MassArchetypeTypes.h b/Engine/Plugins/Runtime/MassEntity/Source/MassEntity/Public/MassArchetypeTypes.h
new file mode 100644
--- /dev/null
+++ b/Engine/Plugins/Runtime/MassEntity/Source/MassEntity/Public/MassArchetypeTypes.h
@@ -0,0 +1,157 @@
+// Copyright Epic Games, Inc. All Rights Reserved.
+
+#pragma once
+
+#include "UObject/Class.h"
+#include "Containers/ArrayView.h"
+#include "Containers/UnrealString.h"
+#include "MassEntityTypes.h"
+
+class UMassEntitySubsystem;
+struct FMassArchetypeData;
+struct FMassExecutionContext;
+struct FMassFragment;
+struct FMassArchetypeChunkIterator;
+struct FMassEntityQuery;
+struct FArchetypeChunkCollection;
+struct FMassEntityView;
+
+typedef TFunction< void(FMassExecutionContext& /*ExecutionContext*/) > FMassExecuteFunction;
+typedef TFunction< bool(const FMassExecutionContext& /*ExecutionContext*/) > FMassChunkConditionFunction;
+
+
+//////////////////////////////////////////////////////////////////////
+//
+
+// An opaque handle to an archetype
+struct FArchetypeHandle final
+{
+	FArchetypeHandle() = default;
+	bool IsValid() const { return DataPtr.IsValid(); }
+
+	MASSENTITY_API bool operator==(const FMassArchetypeData* Other) const;
+	bool operator==(const FArchetypeHandle& Other) const { return DataPtr == Other.DataPtr; }
+	bool operator!=(const FArchetypeHandle& Other) const { return DataPtr != Other.DataPtr; }
+
+	MASSENTITY_API friend uint32 GetTypeHash(const FArchetypeHandle& Instance);
+private:
+	FArchetypeHandle(const TSharedPtr<FMassArchetypeData>& InDataPtr)
+	: DataPtr(InDataPtr)
+	{}
+	TSharedPtr<FMassArchetypeData> DataPtr;
+
+	friend UMassEntitySubsystem;
+	friend FArchetypeChunkCollection;
+	friend FMassEntityQuery;
+	friend FMassEntityView;
+};
+
+
+//////////////////////////////////////////////////////////////////////
+//
+
+/** A struct that converts an arbitrary array of entities of given Archetype into a sequence of continuous
+ *  entity chunks. The goal is to have the user create an instance of this struct once and run through a bunch of
+ *  systems. The runtime code usually uses FMassArchetypeChunkIterator to iterate on the chunk collection
+ */
+struct MASSENTITY_API FArchetypeChunkCollection
+{
+public:
+	struct FChunkInfo
+	{
+		int32 ChunkIndex = INDEX_NONE;
+		int32 SubchunkStart = 0;
+		/** negative or 0-length means "all available entities within chunk" */
+		int32 Length = 0;
+
+		FChunkInfo() = default;
+		explicit FChunkInfo(const int32 InChunkIndex, const int32 InSubchunkStart = 0, const int32 InLength = 0) : ChunkIndex(InChunkIndex), SubchunkStart(InSubchunkStart), Length(InLength) {}
+		/** Note that we consider invalid-length chunks valid as long as ChunkIndex and SubchunkStart are valid */
+		bool IsSet() const { return ChunkIndex != INDEX_NONE && SubchunkStart >= 0; }
+	};
+private:
+	TArray<FChunkInfo> Chunks;
+	/** entity indices indicated by SubChunks are only valid with given Archetype */
+	FArchetypeHandle Archetype;
+
+public:
+	FArchetypeChunkCollection() = default;
+	FArchetypeChunkCollection(const FArchetypeHandle& InArchetype, TConstArrayView<FMassEntityHandle> InEntities);
+	explicit FArchetypeChunkCollection(FArchetypeHandle& InArchetypeHandle);
+	explicit FArchetypeChunkCollection(TSharedPtr<FMassArchetypeData>& InArchetype);
+
+	TArrayView<const FChunkInfo> GetChunks() const { return Chunks; }
+	const FArchetypeHandle& GetArchetype() const { return Archetype; }
+	bool IsEmpty() const { return Chunks.Num() == 0 && Archetype.IsValid() == false; }
+	bool IsSet() const { return Archetype.IsValid(); }
+	void Reset() 
+	{ 
+		Archetype = FArchetypeHandle();
+		Chunks.Reset();
+	}
+
+private:
+	void GatherChunksFromArchetype(TSharedPtr<FMassArchetypeData>& InArchetype);
+};
+
+//////////////////////////////////////////////////////////////////////
+//
+
+/**
+ *  The type used to iterate over given archetype's chunks, be it full, continuous chunks or sparse subchunks. It hides
+ *  this details from the rest of the system.
+ */
+struct MASSENTITY_API FMassArchetypeChunkIterator
+{
+private:
+	const FArchetypeChunkCollection& ChunkData;
+	int32 CurrentChunkIndex = 0;
+
+public:
+	explicit FMassArchetypeChunkIterator(const FArchetypeChunkCollection& InChunkData) : ChunkData(InChunkData), CurrentChunkIndex(0) {}
+
+	operator bool() const { return ChunkData.GetChunks().IsValidIndex(CurrentChunkIndex) && ChunkData.GetChunks()[CurrentChunkIndex].IsSet(); }
+	FMassArchetypeChunkIterator& operator++() { ++CurrentChunkIndex; return *this; }
+
+	const FArchetypeChunkCollection::FChunkInfo* operator->() const { check(bool(*this)); return &ChunkData.GetChunks()[CurrentChunkIndex]; }
+	const FArchetypeChunkCollection::FChunkInfo& operator*() const { check(bool(*this)); return ChunkData.GetChunks()[CurrentChunkIndex]; }
+};
+
+//////////////////////////////////////////////////////////////////////
+//
+struct FInternalEntityHandle 
+{
+	FInternalEntityHandle() = default;
+	FInternalEntityHandle(uint8* InChunkRawMemory, const int32 InIndexWithinChunk)
+        : ChunkRawMemory(InChunkRawMemory), IndexWithinChunk(InIndexWithinChunk)
+	{}
+	bool IsValid() const { return ChunkRawMemory != nullptr && IndexWithinChunk != INDEX_NONE; }
+	bool operator==(const FInternalEntityHandle & Other) const { return ChunkRawMemory == Other.ChunkRawMemory && IndexWithinChunk == Other.IndexWithinChunk; }
+
+	uint8* ChunkRawMemory = nullptr;
+	int32 IndexWithinChunk = INDEX_NONE;
+};
+
+typedef TArray<int32, TInlineAllocator<16>> FMassFragmentIndicesMapping;
+typedef TConstArrayView<int32> FMassFragmentIndicesMappingView;
+struct FMassQueryRequirementIndicesMapping
+{
+	FMassQueryRequirementIndicesMapping() = default;
+
+	FMassFragmentIndicesMapping EntityComponents;
+	FMassFragmentIndicesMapping ChunkComponents;
+	FORCEINLINE bool IsEmpty() const
+	{
+		return EntityComponents.Num() == 0 || ChunkComponents.Num() == 0;
+	}
+};
+
+template<typename T>
+struct FMassSorterOperator
+{
+	bool operator()(const T& A, const T& B) const
+	{
+		return (A.GetStructureSize() > B.GetStructureSize())
+			|| (A.GetStructureSize() == B.GetStructureSize() && B.GetFName().FastLess(A.GetFName()));
+	}
+};
diff --git a/Engine/Plugins/Runtime/MassEntity/Source/MassEntity/Public/MassCommandBuffer.h b/Engine/Plugins/Runtime/MassEntity/Source/MassEntity/Public/MassCommandBuffer.h
new file mode 100644
--- /dev/null
+++ b/Engine/Plugins/Runtime/MassEntity/Source/MassEntity/Public/MassCommandBuffer.h
@@ -0,0 +1,368 @@
+// Copyright Epic Games, Inc. All Rights Reserved.
+
+#pragma once
+
+#include "MassEntityTypes.h"
+#include "MassEntitySubsystem.h"
+#include "InstancedStructStream.h"
+#include "Misc/MTAccessDetector.h"
+
+#include "MassCommandBuffer.generated.h"
+
+//@TODO: Consider debug information in case there is an assert when replaying the command buffer
+// (e.g., which system added the command, or even file/line number in development builds for the specific call via a macro)
+//@TODO: Support more commands
+
+USTRUCT()
+struct MASSENTITY_API FCommandBufferEntryBase
+{
+	GENERATED_BODY()
+
+	FMassEntityHandle TargetEntity;
+
+	FCommandBufferEntryBase() = default;
+	virtual ~FCommandBufferEntryBase() = default;
+
+	FCommandBufferEntryBase(FMassEntityHandle InEntity)
+		: TargetEntity(InEntity)
+	{}
+
+	virtual void Execute(UMassEntitySubsystem& System) const PURE_VIRTUAL(FCommandBufferEntryBase::Execute, );
+};
+template<> struct TStructOpsTypeTraits<FCommandBufferEntryBase> : public TStructOpsTypeTraitsBase2<FCommandBufferEntryBase> { enum { WithPureVirtual = true, }; };
+
+/**
+ * Command dedicated to execute a captured lambda 
+ */
+USTRUCT()
+struct MASSENTITY_API FDeferredCommand : public FCommandBufferEntryBase
+{
+	GENERATED_BODY()
+
+	FDeferredCommand() = default;
+
+	FDeferredCommand(const TFunction<void(UMassEntitySubsystem& System)>& InDeferredFunction)
+		: DeferredFunction(InDeferredFunction)
+	{}
+
+	virtual void Execute(UMassEntitySubsystem& System) const override
+	{
+		DeferredFunction(System);
+	}
+
+protected:
+	TFunction<void(UMassEntitySubsystem& System)> DeferredFunction;
+};
+
+/**
+ * Command dedicated to building an entity from an existing component instance
+ */
+USTRUCT()
+struct MASSENTITY_API FBuildEntityFromComponentInstance : public FCommandBufferEntryBase
+{
+	GENERATED_BODY()
+
+	FBuildEntityFromComponentInstance() = default;
+	FBuildEntityFromComponentInstance(const FMassEntityHandle Entity, FStructView InStruct)
+		: FCommandBufferEntryBase(Entity)
+		, Struct(InStruct)
+	{}
+
+protected:
+	virtual void Execute(UMassEntitySubsystem& System) const override;
+
+	FInstancedStruct Struct;
+};
+
+/**
+* Command dedicated to building an entity from a list of existing component instances
+*/
+USTRUCT()
+struct MASSENTITY_API FBuildEntityFromComponentInstances : public FCommandBufferEntryBase
+{
+	GENERATED_BODY()
+
+	FBuildEntityFromComponentInstances() = default;
+	FBuildEntityFromComponentInstances(const FMassEntityHandle Entity, TConstArrayView<FInstancedStruct> InInstances)
+		: FCommandBufferEntryBase(Entity)
+		, Instances(InInstances)
+	{}
+
+protected:
+	virtual void Execute(UMassEntitySubsystem& System) const override;
+
+	TArray<FInstancedStruct> Instances;
+};
+
+/**
+ * Command dedicated to add a new component to an existing entity
+ */
+USTRUCT()
+struct MASSENTITY_API FCommandAddComponent : public FCommandBufferEntryBase
+{
+	GENERATED_BODY()
+
+	FCommandAddComponent() = default;
+	FCommandAddComponent(FMassEntityHandle InEntity, UScriptStruct* InStruct)
+		: FCommandBufferEntryBase(InEntity)
+		, StructParam(InStruct)
+	{}
+
+protected:
+	virtual void Execute(UMassEntitySubsystem& System) const override;
+
+	UScriptStruct* StructParam = nullptr;
+};
+
+/**
+ * Command dedicated to add a new component from an existing instance
+ */
+USTRUCT()
+struct MASSENTITY_API FCommandAddComponentInstance : public FCommandBufferEntryBase
+{
+	GENERATED_BODY()
+
+	FCommandAddComponentInstance() = default;
+	FCommandAddComponentInstance(const FMassEntityHandle InEntity, FStructView InStruct)
+        : FCommandBufferEntryBase(InEntity)
+        , Struct(InStruct)
+	{}
+
+protected:
+	virtual void Execute(UMassEntitySubsystem& System) const override;
+
+	FInstancedStruct Struct;
+};
+
+/**
+ * Command dedicated to remove a component from an existing entity
+ */
+USTRUCT()
+struct MASSENTITY_API FCommandRemoveComponent : public FCommandBufferEntryBase
+{
+	GENERATED_BODY()
+
+	FCommandRemoveComponent() = default;
+	FCommandRemoveComponent(FMassEntityHandle InEntity, UScriptStruct* InStruct)
+		: FCommandBufferEntryBase(InEntity)
+		, StructParam(InStruct)
+	{}
+
+protected:
+	virtual void Execute(UMassEntitySubsystem& System) const override;
+
+	UScriptStruct* StructParam = nullptr;
+};
+
+/**
+ * Command performing addition of a list of components to an existing entity
+ */
+USTRUCT()
+struct MASSENTITY_API FCommandAddComponentList : public FCommandBufferEntryBase
+{
+	GENERATED_BODY()
+
+	FCommandAddComponentList() = default;
+	FCommandAddComponentList(FMassEntityHandle InEntity, TConstArrayView<const UScriptStruct*> InComponentList)
+		: FCommandBufferEntryBase(InEntity)
+		, ComponentList(InComponentList)
+	{}
+
+protected:
+	virtual void Execute(UMassEntitySubsystem& System) const override;
+
+	TArray<const UScriptStruct*> ComponentList;
+};
+
+/**
+ * Command performing a removal of a list of components from an existing entity
+ */
+USTRUCT()
+struct MASSENTITY_API FCommandRemoveComponentList : public FCommandBufferEntryBase
+{
+	GENERATED_BODY()
+
+	FCommandRemoveComponentList() = default;
+	FCommandRemoveComponentList(FMassEntityHandle InEntity, TConstArrayView<const UScriptStruct*> InComponentList)
+		: FCommandBufferEntryBase(InEntity)
+		, ComponentList(InComponentList)
+	{}
+
+protected:
+	virtual void Execute(UMassEntitySubsystem& System) const override;
+
+	TArray<const UScriptStruct*> ComponentList;
+};
+
+/**
+ * Command dedicated to add a new tag to an existing entity
+ */
+USTRUCT()
+struct MASSENTITY_API FCommandAddTag: public FCommandBufferEntryBase
+{
+	GENERATED_BODY()
+
+	FCommandAddTag() = default;
+	FCommandAddTag(FMassEntityHandle InEntity, UScriptStruct* InStruct)
+		: FCommandBufferEntryBase(InEntity)
+		, StructParam(InStruct)
+	{}
+
+protected:
+	virtual void Execute(UMassEntitySubsystem& System) const override;
+
+	UScriptStruct* StructParam = nullptr;
+};
+
+/**
+ * Command dedicated to remove a tag from an existing entity
+ */
+USTRUCT()
+struct MASSENTITY_API FCommandRemoveTag : public FCommandBufferEntryBase
+{
+	GENERATED_BODY()
+
+	FCommandRemoveTag() = default;
+	FCommandRemoveTag(FMassEntityHandle InEntity, UScriptStruct* InStruct)
+		: FCommandBufferEntryBase(InEntity)
+		, StructParam(InStruct)
+	{}
+
+protected:
+	virtual void Execute(UMassEntitySubsystem& System) const override;
+
+	UScriptStruct* StructParam = nullptr;
+};
+
+/**
+ * Command for swapping given tags for a given entity. Note that the entity doesn't need to own the "old tag", the 
+ * "new tag" will be added regardless.
+ */
+USTRUCT()
+struct MASSENTITY_API FCommandSwapTags : public FCommandBufferEntryBase
+{
+	GENERATED_BODY()
+
+	FCommandSwapTags() = default;
+
+	FCommandSwapTags(const FMassEntityHandle InEntity, const UScriptStruct* InOldTagType, const UScriptStruct* InNewTagType)
+		: FCommandBufferEntryBase(InEntity)
+		, OldTagType(InOldTagType)
+		, NewTagType(InNewTagType)
+	{
+		checkf((InOldTagType == nullptr) || InOldTagType->IsChildOf(FMassTag::StaticStruct()), TEXT("FCommandSwapTags works only with tags while '%s' is not one."), *GetPathNameSafe(InOldTagType));
+		checkf((InNewTagType == nullptr) || InNewTagType->IsChildOf(FMassTag::StaticStruct()), TEXT("FCommandSwapTags works only with tags while '%s' is not one."), *GetPathNameSafe(InNewTagType));
+	}
+
+protected:
+	virtual void Execute(UMassEntitySubsystem& System) const override;
+
+	const UScriptStruct* OldTagType = nullptr;
+	const UScriptStruct* NewTagType = nullptr;
+};
+
+/**
+ * Command performing a removal of a collection of components and tags
+ */
+USTRUCT()
+struct MASSENTITY_API FCommandRemoveComposition : public FCommandBufferEntryBase
+{
+	GENERATED_BODY()
+
+	FCommandRemoveComposition() = default;
+	FCommandRemoveComposition(FMassEntityHandle InEntity, const FMassCompositionDescriptor& InDescriptor)
+		: FCommandBufferEntryBase(InEntity)
+		, Descriptor(InDescriptor)
+	{}
+
+protected:
+	virtual void Execute(UMassEntitySubsystem& System) const override;
+
+	FMassCompositionDescriptor Descriptor;
+};
+
+struct MASSENTITY_API FMassCommandBuffer
+{
+public:
+
+	/** Push any command, requires to derive from FCommandBufferEntryBase */
+	template< typename T, typename = typename TEnableIf<TIsDerivedFrom<typename TRemoveReference<T>::Type, FCommandBufferEntryBase>::IsDerived, void>::Type >
+	void PushCommand(const T& Command)
+	{
+		EmplaceCommand<T>(Command);
+	}
+
+	/** Emplace any command, requires to derive from FCommandBufferEntryBase */
+	template< typename T, typename = typename TEnableIf<TIsDerivedFrom<typename TRemoveReference<T>::Type, FCommandBufferEntryBase>::IsDerived, void>::Type, typename... TArgs >
+	void EmplaceCommand(TArgs&&... InArgs)
+	{
+		UE_MT_SCOPED_WRITE_ACCESS(PendingCommandsDetector);
+		return PendingCommands.Emplace<T>(Forward<TArgs>(InArgs)...);
+	}
+
+	/** Emplace any command and return its ref, requires to derive from FCommandBufferEntryBase */
+	template< typename T, typename = typename TEnableIf<TIsDerivedFrom<typename TRemoveReference<T>::Type, FCommandBufferEntryBase>::IsDerived, void>::Type, typename... TArgs >
+	T& EmplaceCommand_GetRef(TArgs&&... InArgs)
+	{
+		UE_MT_SCOPED_WRITE_ACCESS(PendingCommandsDetector);
+		return PendingCommands.Emplace_GetRef<T>(Forward<TArgs>(InArgs)...);
+	}
+
+	template<typename T>
+	void AddComponent(FMassEntityHandle Entity)
+	{
+		static_assert(TIsDerivedFrom<T, FMassFragment>::IsDerived, "Given struct type is not a valid component type.");
+		EmplaceCommand<FCommandAddComponent>(Entity, T::StaticStruct());
+	}
+
+	template<typename T>
+	void RemoveComponent(FMassEntityHandle Entity)
+	{
+		static_assert(TIsDerivedFrom<T, FMassFragment>::IsDerived, "Given struct type is not a valid component type.");
+		EmplaceCommand<FCommandRemoveComponent>(Entity, T::StaticStruct());
+	}
+
+	template<typename T>
+	void AddTag(FMassEntityHandle Entity)
+	{
+		static_assert(TIsDerivedFrom<T, FMassTag>::IsDerived, "Given struct type is not a valid tag type.");
+		EmplaceCommand<FCommandAddTag>(Entity, T::StaticStruct());
+	}
+
+	template<typename T>
+	void RemoveTag(FMassEntityHandle Entity)
+	{
+		static_assert(TIsDerivedFrom<T, FMassTag>::IsDerived, "Given struct type is not a valid tag type.");
+		EmplaceCommand<FCommandRemoveTag>(Entity, T::StaticStruct());
+	}
+
+	void DestroyEntity(FMassEntityHandle Entity)
+	{
+		EntitiesToDestroy.Add(Entity);
+	}
+
+	void BatchDestroyEntities(const TArray<FMassEntityHandle>& InEntitiesToDestroy)
+	{
+		EntitiesToDestroy.Append(InEntitiesToDestroy);
+	}
+
+	void ReplayBufferAgainstSystem(UMassEntitySubsystem* System);
+
+	SIZE_T GetAllocatedSize() const { return PendingCommands.GetAllocatedSize(); }
+
+	/** 
+	 * Appends the commands from the passed buffer into this one
+	 * @param InOutOther the source buffer to copy the commands from. Note that after the call the InOutOther will be 
+	 *	emptied due to the function using Move semantics
+	 */
+	void MoveAppend(FMassCommandBuffer& InOutOther);
+
+	int32 GetPendingCommandsCount() const { return PendingCommands.Num(); }
+
+private:
+	FInstancedStructStream PendingCommands;
+	UE_MT_DECLARE_RW_ACCESS_DETECTOR(PendingCommandsDetector);
+	FCriticalSection AppendingCommandsCS;
+
+	TArray<FMassEntityHandle> EntitiesToDestroy;
+};
diff --git a/Engine/Plugins/Runtime/MassEntity/Source/MassEntity/Public/MassEntityDebug.h b/Engine/Plugins/Runtime/MassEntity/Source/MassEntity/Public/MassEntityDebug.h
--- a/Engine/Plugins/Runtime/MassEntity/Source/MassEntity/Public/MassEntityDebug.h
+++ b/Engine/Plugins/Runtime/MassEntity/Source/MassEntity/Public/MassEntityDebug.h
@@ -2,7 +2,7 @@
 
 #pragma once
 #include "VisualLogger/VisualLoggerTypes.h"
-#include "MassEntityTypes.h"
+#include "MassProcessingTypes.h"
 
 class FOutputDevice;
 class UMassProcessor;
diff --git a/Engine/Plugins/Runtime/MassEntity/Source/MassEntity/Public/MassEntityQuery.h b/Engine/Plugins/Runtime/MassEntity/Source/MassEntity/Public/MassEntityQuery.h
new file mode 100644
--- /dev/null
+++ b/Engine/Plugins/Runtime/MassEntity/Source/MassEntity/Public/MassEntityQuery.h
@@ -0,0 +1,375 @@
+// Copyright Epic Games, Inc. All Rights Reserved.
+
+#pragma once
+
+#include "UObject/Class.h"
+#include "Containers/ArrayView.h"
+#include "Containers/UnrealString.h"
+#include "MassEntityTypes.h"
+#include "MassArchetypeTypes.h"
+#include "MassEntityQuery.generated.h"
+
+
+class UMassEntitySubsystem;
+struct FMassArchetypeData;
+struct FMassExecutionContext;
+struct FMassFragment;
+struct FArchetypeHandle;
+
+enum class EMassFragmentAccess : uint8
+{
+	// no binding required
+	None, 
+
+	// We want to read the data for the fragment
+	ReadOnly,
+
+	// We want to read and write the data for the fragment
+	ReadWrite,
+
+	MAX
+};
+
+enum class EMassFragmentPresence : uint8
+{
+	// All of the required fragments must be present
+	All,
+
+	// One of the required fragments must be present
+	Any,
+
+	// None of the required fragments can be present
+	None,
+
+	// If fragment is present we'll use it, but it missing stop processing of a given archetype
+	Optional,
+
+	MAX
+};
+
+
+struct MASSENTITY_API FMassFragmentRequirement
+{
+	const UScriptStruct* StructType = nullptr;
+	EMassFragmentAccess AccessMode = EMassFragmentAccess::None;
+	EMassFragmentPresence Presence = EMassFragmentPresence::Optional;
+
+public:
+	FMassFragmentRequirement(){}
+	FMassFragmentRequirement(const UScriptStruct* InStruct, const EMassFragmentAccess InAccessMode, const EMassFragmentPresence InPresence)
+		: StructType(InStruct)
+		, AccessMode(InAccessMode)
+		, Presence(InPresence)
+	{
+		check(InStruct);
+		checkf((Presence != EMassFragmentPresence::Any && Presence != EMassFragmentPresence::Optional)
+			|| AccessMode == EMassFragmentAccess::ReadOnly || AccessMode == EMassFragmentAccess::ReadWrite, TEXT("Only ReadOnly and ReadWrite modes are suppored for optional requirements"));
+	}
+
+	bool RequiresBinding() const { return (AccessMode != EMassFragmentAccess::None); }
+	bool IsOptional() const { return (Presence == EMassFragmentPresence::Optional || Presence == EMassFragmentPresence::Any); }
+
+	FString DebugGetDescription() const;
+
+	// these functions are used for sorting. See FMassSorterOperator
+	int32 GetStructureSize() const
+	{
+		return StructType->GetStructureSize();
+	}
+
+	FName GetFName() const
+	{
+		return StructType->GetFName();
+	}
+};
+
+
+/** 
+ *  FMassEntityQuery is a structure that serves two main purposes:
+ *  1. Describe properties required of an archetype that's a subject of calculations
+ *  2. Trigger calculations on cached set of valid archetypes as described by Requirements
+ * 
+ *  A query to be considered valid needs declared at least one EMassFragmentPresence::All, EMassFragmentPresence::Any 
+ *  EMassFragmentPresence::Optional fragment requirement.
+ */
+USTRUCT()
+struct MASSENTITY_API FMassEntityQuery
+{
+	GENERATED_BODY()
+
+public:
+	FMassEntityQuery();
+	FMassEntityQuery(std::initializer_list<UScriptStruct*> InitList);
+	FMassEntityQuery(TConstArrayView<const UScriptStruct*> InitList);
+
+	/** Runs ExecuteFunction on all entities matching Requirements */
+	void ForEachEntityChunk(UMassEntitySubsystem& EntitySubsystem, FMassExecutionContext& ExecutionContext, const FMassExecuteFunction& ExecuteFunction);
+	
+	/** Will first verify that the archetype given with Chunks matches the query's requirements, and if so will run the other, more generic ForEachEntityChunk implementation */
+	void ForEachEntityChunk(const FArchetypeChunkCollection& Chunks, UMassEntitySubsystem& EntitySubsystem, FMassExecutionContext& ExecutionContext, const FMassExecuteFunction& ExecuteFunction);
+
+	/**
+	 * Attempts to process every chunk of every affected archetype in parallel.
+	 */
+	void ParallelForEachEntityChunk(UMassEntitySubsystem& EntitySubsystem, FMassExecutionContext& ExecutionContext, const FMassExecuteFunction& ExecuteFunction);
+
+	/** Will gather all archetypes from InEntitySubsystem matching this->Requirements.
+	 *  Note that no work will be done if the cached data is up to date (as tracked by EntitySubsystemHash and 
+	 *	ArchetypeDataVersion properties). */
+	void CacheArchetypes(UMassEntitySubsystem& InEntitySubsystem);
+
+	FMassEntityQuery& AddRequirement(const UScriptStruct* ComponentType, const EMassFragmentAccess AccessMode, const EMassFragmentPresence Presence = EMassFragmentPresence::All)
+	{
+		checkf(Requirements.FindByPredicate([ComponentType](const FMassFragmentRequirement& Item){ return Item.StructType == ComponentType; }) == nullptr
+			, TEXT("Duplicated requirements are not supported. %s already present"), *GetNameSafe(ComponentType));
+		
+		if (Presence != EMassFragmentPresence::None)
+		{
+			Requirements.Emplace(ComponentType, AccessMode, Presence);
+		}
+
+		switch (Presence)
+		{
+		case EMassFragmentPresence::All:
+			RequiredAllComponents.Add(*ComponentType);
+			break;
+		case EMassFragmentPresence::Any:
+			RequiredAnyComponents.Add(*ComponentType);
+			break;
+		case EMassFragmentPresence::Optional:
+			RequiredOptionalComponents.Add(*ComponentType);
+			break;
+		case EMassFragmentPresence::None:
+			RequiredNoneComponents.Add(*ComponentType);
+			break;
+		}
+		// force recaching the next time this query is used or the following CacheArchetypes call.
+		DirtyCachedData();
+		return *this;
+	}
+
+	/** FMassEntityQuery ref returned for chaining */
+	template<typename T>
+	FMassEntityQuery& AddRequirement(const EMassFragmentAccess AccessMode, const EMassFragmentPresence Presence = EMassFragmentPresence::All)
+	{
+		checkf(Requirements.FindByPredicate([](const FMassFragmentRequirement& Item) { return Item.StructType == T::StaticStruct(); }) == nullptr
+			, TEXT("Duplicated requirements are not supported. %s already present"), *T::StaticStruct()->GetName());
+
+		static_assert(TIsDerivedFrom<T, FMassFragment>::IsDerived, "Given struct doesn't represent a valid fragment type. Make sure to inherit from FMassFragment or one of its child-types.");
+		
+		if (Presence != EMassFragmentPresence::None)
+		{
+			Requirements.Emplace(T::StaticStruct(), AccessMode, Presence);
+		}
+		
+		switch (Presence)
+		{
+		case EMassFragmentPresence::All:
+			RequiredAllComponents.Add<T>();
+			break;
+		case EMassFragmentPresence::Any:
+			RequiredAnyComponents.Add<T>();
+			break;
+		case EMassFragmentPresence::Optional:
+			RequiredOptionalComponents.Add<T>();
+			break;
+		case EMassFragmentPresence::None:
+			RequiredNoneComponents.Add<T>();
+			break;
+		}
+		// force recaching the next time this query is used or the following CacheArchetypes call.
+		DirtyCachedData();
+		return *this;
+	}
+
+	void AddTagRequirement(const UScriptStruct& TagType, const EMassFragmentPresence Presence)
+	{
+		checkfSlow(int(Presence) < int(EMassFragmentPresence::Optional), TEXT("Optional and MAX presence are not valid calues for AddTagRequirement"));
+		switch (Presence)
+		{
+		case EMassFragmentPresence::All:
+			RequiredAllTags.Add(TagType);
+			break;
+		case EMassFragmentPresence::Any:
+			RequiredAnyTags.Add(TagType);
+			break;
+		case EMassFragmentPresence::None:
+			RequiredNoneTags.Add(TagType);
+			break;
+		}
+	}
+
+	template<typename T>
+	FMassEntityQuery& AddTagRequirement(const EMassFragmentPresence Presence)
+	{
+		checkfSlow(int(Presence) < int(EMassFragmentPresence::Optional), TEXT("Optional and MAX presence are not valid calues for AddTagRequirement"));
+		static_assert(TIsDerivedFrom<T, FMassTag>::IsDerived, "Given struct doesn't represent a valid tag type. Make sure to inherit from FMassFragment or one of its child-types.");
+		switch (Presence)
+		{
+			case EMassFragmentPresence::All:
+				RequiredAllTags.Add<T>();
+				break;
+			case EMassFragmentPresence::Any:
+				RequiredAnyTags.Add<T>();
+				break;
+			case EMassFragmentPresence::None:
+				RequiredNoneTags.Add<T>();
+				break;
+		}
+
+		return *this;
+	}
+
+	// actual implementation in specializations
+	template<EMassFragmentPresence Presence> 
+	FMassEntityQuery& AddTagRequirements(const FMassTagBitSet& TagBitSet)
+	{
+		static_assert(Presence == EMassFragmentPresence::None || Presence == EMassFragmentPresence::All || Presence == EMassFragmentPresence::Any
+			, "The only valid values for AddTagRequirements are All, Any and None");
+		return *this;
+	}
+
+	template<typename T>
+	FMassEntityQuery& AddChunkRequirement(const EMassFragmentAccess AccessMode, const EMassFragmentPresence Presence = EMassFragmentPresence::All)
+	{
+		static_assert(TIsDerivedFrom<T, FMassChunkFragment>::IsDerived, "Given struct doesn't represent a valid chunk fragment type. Make sure to inherit from FMassChunkFragment or one of its child-types.");
+		checkf(ChunkRequirements.FindByPredicate([](const FMassFragmentRequirement& Item) { return Item.StructType == T::StaticStruct(); }) == nullptr
+			, TEXT("Duplicated requirements are not supported. %s already present"), *T::StaticStruct()->GetName());
+		checkfSlow(Presence == EMassFragmentPresence::None || Presence == EMassFragmentPresence::All
+			, TEXT("The only valid Presence values for AddChunkRequirement are All and None"));
+		
+		if (Presence == EMassFragmentPresence::All)
+		{
+			RequiredAllChunkComponents.Add<T>();
+			ChunkRequirements.Emplace(T::StaticStruct(), AccessMode, Presence);
+		}
+		else
+		{
+			RequiredNoneChunkComponents.Add<T>();
+		}
+
+		return *this;
+	}
+
+	void Clear()
+	{
+		Requirements.Empty();
+		DirtyCachedData();
+	}
+
+	FORCEINLINE void DirtyCachedData()
+	{
+		EntitySubsystemHash = 0;
+		ArchetypeDataVersion = 0;
+	}
+	
+	bool DoesArchetypeMatchRequirements(const FArchetypeHandle& ArchetypeHandle) const;
+
+	/** The function validates requirements we make for queries. See the FMassEntityQuery struct description for details.
+	 *  Note that this function is non-trivial and end users are not expected to need to use it. 
+	 *  @return whether this query's requirements follow the rules. */
+	bool CheckValidity() const;
+
+	FString DebugGetDescription() const;
+
+	TConstArrayView<FMassFragmentRequirement> GetRequirements() const { return Requirements; }
+	const FMassFragmentBitSet& GetRequiredAllComponents() const { return RequiredAllComponents; }
+	const FMassFragmentBitSet& GetRequiredAnyComponents() const { return RequiredAnyComponents; }
+	const FMassFragmentBitSet& GetRequiredOptionalComponents() const { return RequiredOptionalComponents; }
+	const FMassFragmentBitSet& GetRequiredNoneComponents() const { return RequiredNoneComponents; }
+	const FMassTagBitSet& GetRequiredAllTags() const { return RequiredAllTags; }
+	const FMassTagBitSet& GetRequiredAnyTags() const { return RequiredAnyTags; }
+	const FMassTagBitSet& GetRequiredNoneTags() const { return RequiredNoneTags; }
+	const FMassChunkFragmentBitSet& GetRequiredAllChunkComponents() const { return RequiredAllChunkComponents; }
+	const FMassChunkFragmentBitSet& GetRequiredNoneChunkComponents() const { return RequiredNoneChunkComponents; }
+
+	const TArray<FArchetypeHandle>& GetArchetypes() const
+	{ 
+		return ValidArchetypes; 
+	}
+
+	/** 
+	 * Goes through ValidArchetypes and sums up the number of entities contained in them.
+	 * Note that the function is not const because calling it can result in re-caching of ValidArchetypes 
+	 * @return the number of entities this given query would process if called "now"
+	 */
+	int32 GetMatchingEntitiesNum(UMassEntitySubsystem& InEntitySubsystem);
+
+	/**
+	 * Checks if any of ValidArchetypes has any entities.
+	 * Note that the function is not const because calling it can result in re-caching of ValidArchetypes
+	 * @return "true" if any of the ValidArchetypes has any entities, "false" otherwise
+	 */
+	bool HasMatchingEntities(UMassEntitySubsystem& InEntitySubsystem);
+
+	/** 
+	 * Sets a chunk filter condition that will applied to each chunk of all valid archetypes. Note 
+	 * that this condition won't be applied when a specific collection of chunks is used (via FArchetypeChunkCollection)
+	 * The value returned by InFunction controls whether to allow execution (true) or block it (false).
+	 */
+	void SetChunkFilter(const FMassChunkConditionFunction& InFunction) { ChunkCondition = InFunction; }
+
+	void ClearChunkFilter() { ChunkCondition.Reset(); }
+
+	bool HasChunkFilter() const { return bool(ChunkCondition); }
+
+protected:
+	void SortRequirements();
+	void ReadCommandlineParams();
+
+protected:
+	TArray<FMassFragmentRequirement> Requirements;
+	TArray<FMassFragmentRequirement> ChunkRequirements;
+	FMassTagBitSet RequiredAllTags;
+	FMassTagBitSet RequiredAnyTags;
+	FMassTagBitSet RequiredNoneTags;
+	FMassFragmentBitSet RequiredAllComponents;
+	FMassFragmentBitSet RequiredAnyComponents;
+	FMassFragmentBitSet RequiredOptionalComponents;
+	FMassFragmentBitSet RequiredNoneComponents;
+	FMassChunkFragmentBitSet RequiredAllChunkComponents;
+	FMassChunkFragmentBitSet RequiredNoneChunkComponents;
+
+private:
+	/** 
+	 * This function represents a condition that will be called for every chunk to be processed before the actual 
+	 * execution function is called. The chunk fragment requirements are already bound and ready to be used by the time 
+	 * ChunkCondition is executed.
+	 */
+	FMassChunkConditionFunction ChunkCondition;
+
+	uint32 EntitySubsystemHash = 0;
+	uint32 ArchetypeDataVersion = 0;
+
+	TArray<FArchetypeHandle> ValidArchetypes;
+	TArray<FMassQueryRequirementIndicesMapping> ArchetypeComponentMapping;
+
+	bool bAllowParallelExecution = false;
+};
+
+template<>
+FORCEINLINE FMassEntityQuery& FMassEntityQuery::AddTagRequirements<EMassFragmentPresence::All>(const FMassTagBitSet& TagBitSet)
+{
+	RequiredAllTags += TagBitSet;
+	// force recaching the next time this query is used or the following CacheArchetypes call.
+	DirtyCachedData();
+	return *this;
+}
+
+template<>
+FORCEINLINE FMassEntityQuery& FMassEntityQuery::AddTagRequirements<EMassFragmentPresence::Any>(const FMassTagBitSet& TagBitSet)
+{
+	RequiredAnyTags += TagBitSet;
+	// force recaching the next time this query is used or the following CacheArchetypes call.
+	DirtyCachedData();
+	return *this;
+}
+
+template<>
+FORCEINLINE FMassEntityQuery& FMassEntityQuery::AddTagRequirements<EMassFragmentPresence::None>(const FMassTagBitSet& TagBitSet)
+{
+	RequiredNoneTags += TagBitSet;
+	// force recaching the next time this query is used or the following CacheArchetypes call.
+	DirtyCachedData();
+	return *this;
+}
diff --git a/Engine/Plugins/Runtime/MassEntity/Source/MassEntity/Public/MassEntitySubsystem.h b/Engine/Plugins/Runtime/MassEntity/Source/MassEntity/Public/MassEntitySubsystem.h
--- a/Engine/Plugins/Runtime/MassEntity/Source/MassEntity/Public/MassEntitySubsystem.h
+++ b/Engine/Plugins/Runtime/MassEntity/Source/MassEntity/Public/MassEntitySubsystem.h
@@ -3,10 +3,10 @@
 #pragma once
 
 #include "Subsystems/WorldSubsystem.h"
-#include "LWComponentTypes.h"
 #include "MassEntityTypes.h"
+#include "MassProcessingTypes.h"
 #include "InstancedStruct.h"
-#include "EntityQuery.h"
+#include "MassEntityQuery.h"
 #include "MassEntitySubsystem.generated.h"
 
 MASSENTITY_API DECLARE_LOG_CATEGORY_EXTERN(LogAggregateTicking, Warning, All);
diff --git a/Engine/Plugins/Runtime/MassEntity/Source/MassEntity/Public/MassEntityTypes.h b/Engine/Plugins/Runtime/MassEntity/Source/MassEntity/Public/MassEntityTypes.h
--- a/Engine/Plugins/Runtime/MassEntity/Source/MassEntity/Public/MassEntityTypes.h
+++ b/Engine/Plugins/Runtime/MassEntity/Source/MassEntity/Public/MassEntityTypes.h
@@ -2,116 +2,144 @@
 
 #pragma once
 
-#include "StructUtilsTypes.h"
-#include "InstancedStruct.h"
+#include "ScriptStructTypeBitSet.h"
 #include "MassEntityTypes.generated.h"
 
-#define WITH_MASSENTITY_DEBUG (!(UE_BUILD_SHIPPING || UE_BUILD_SHIPPING_WITH_EDITOR || UE_BUILD_TEST) && WITH_STRUCTUTILS_DEBUG && 1)
-
-MASSENTITY_API DECLARE_LOG_CATEGORY_EXTERN(LogPipe, Warning, All);
-
-class UMassEntitySubsystem;
-class UMassProcessor;
-class UMassSchematic;
-class UMassCompositeProcessor;
-
-UENUM(meta = (Bitflags, UseEnumValuesAsMaskValuesInEditor = "true"))
-enum class EProcessorExecutionFlags : uint8
+// This is the base class for all lightweight components
+USTRUCT()
+struct FMassFragment
 {
-	None = 0 UMETA(Hidden),
-	Standalone = 1 << 0,
-	Server = 1 << 1,
-	Client = 1 << 2,
-	All = Standalone | Server | Client UMETA(Hidden)
+	GENERATED_BODY()
+
+	FMassFragment() {}
 };
-ENUM_CLASS_FLAGS(EProcessorExecutionFlags);
 
+// This is the base class for types that will only be tested for presence/absence, i.e. Tags.
+// Subclasses should never contain any member properties.
 USTRUCT()
-struct FProcessorAuxDataBase
+struct FMassTag
 {
 	GENERATED_BODY()
+
+	FMassTag() {}
 };
 
-USTRUCT(BlueprintType)
-struct MASSENTITY_API FMassProcessingContext
+USTRUCT()
+struct FMassChunkFragment
 {
 	GENERATED_BODY()
 
-	UPROPERTY()
-	UMassEntitySubsystem* EntitySubsystem = nullptr;
-	
-	UPROPERTY()
-	float DeltaSeconds = 0.f;
-
-	UPROPERTY()
-	FInstancedStruct AuxData;
-
-	FMassProcessingContext() = default;
-	FMassProcessingContext(UMassEntitySubsystem& InEntities, const float InDeltaSeconds);
+	FMassChunkFragment() {}
 };
 
-/** 
- *  A helper type that converts a set of UMassSchematics into a runtime-usable array of MassProcessor copies
- */
+// A handle to a lightweight entity.  An entity is used in conjunction with the UMassEntitySubsystem
+// for the current world and can contain lightweight components.
 USTRUCT()
-struct MASSENTITY_API FMassRuntimePipeline
+struct FMassEntityHandle
 {
 	GENERATED_BODY()
 
-	UPROPERTY()
-	TArray<UMassProcessor*> Processors;
-
-	void Reset();
-	void Initialize(UObject& Owner);
+	FMassEntityHandle() {}
 	
-	/** Creates runtime copies of the given UMassProcessors collection. */
-	void SetProcessors(TArray<UMassProcessor*>&& InProcessors);
-
-	/** Creates runtime copies of UMassProcessors declared in Schematics using InOwner as new UMassProcessors' outer. */
-	void InitializeFromSchematics(TConstArrayView<TSoftObjectPtr<UMassSchematic>> Schematics, UObject& InOwner);
-
-	/** Creates runtime copies of UMassProcessors given in InProcessors input parameter, using InOwner as new UMassProcessors' outer. */
-	void CreateFromArray(TConstArrayView<const UMassProcessor*> InProcessors, UObject& InOwner);
-
-	/** Calls CreateFromArray and calls Initialize on all processors afterwards. */
-	void InitializeFromArray(TConstArrayView<const UMassProcessor*> InProcessors, UObject& InOwner);
+	UPROPERTY(VisibleAnywhere, Category = "AggregateTicking|Debug", Transient)
+	int32 Index = 0;
 	
-	/** Creates runtime instances of UMassProcessors for each processor class given via InProcessorClasses. 
-	 *  The instances will be created with InOwner as outer. */
-	void InitializeFromClassArray(TConstArrayView<TSubclassOf<UMassProcessor>> InProcessorClasses, UObject& InOwner);
-
-	/** Creates a runtime instance of every processors in the given InProcessors array. If a processor of that class
-	 *  already exists in Processors array it gets overridden. Otherwise it gets added to the end of the collection.*/
-	void AppendOrOverrideRuntimeProcessorCopies(TConstArrayView<const UMassProcessor*> InProcessors, UObject& InOwner);
-
-	/** Creates a runtime instance of every processors in the given array if there's no processor of that class in Processors already.
-	 *  Call this function when adding processors to an already configured FMassRuntimePipeline instance. If you're creating 
-	 *  one from scratch calling any of the InitializeFrom* methods will be more efficient (and will produce same results)
-	 *  or call AppendOrOverrideRuntimeProcessorCopies.*/
-	void AppendUniqueRuntimeProcessorCopies(TConstArrayView<const UMassProcessor*> InProcessors, UObject& InOwner);
-
-	/** Adds InProcessor to Processors without any additional checks */
-	void AppendProcessor(UMassProcessor& Processor);
-
-	/** goes through Processor looking for a UMassCompositeProcessor instance which GroupName matches the one given as the parameter */
-	UMassCompositeProcessor* FindTopLevelGroupByName(const FName GroupName);
-
-	void DebugOutputDescription(FOutputDevice& Ar) const;
-
-	bool HasProcessorOfExactClass(TSubclassOf<UMassProcessor> InClass) const;
-	bool IsEmpty() const { return Processors.IsEmpty();}
-
-	MASSENTITY_API friend uint32 GetTypeHash(const FMassRuntimePipeline& Instance);
+	UPROPERTY(VisibleAnywhere, Category = "AggregateTicking|Debug", Transient)
+	int32 SerialNumber = 0;
+
+	bool operator==(const FMassEntityHandle& Other) const
+	{
+		return Index == Other.Index && SerialNumber == Other.SerialNumber;
+	}
+
+	bool operator!=(const FMassEntityHandle& Other) const
+	{
+		return !operator==(Other);
+	}
+
+	/** Note that this function is merely checking if Index and SerialNumber are set. There's no way to validate if 
+	 *  these indicate a valid entity in an EntitySubsystem without asking the system. */
+	bool IsSet() const
+	{
+		return Index != 0 && SerialNumber != 0;
+	}
+
+	void Reset()
+	{
+		Index = SerialNumber = 0;
+	}
+
+	friend uint32 GetTypeHash(const FMassEntityHandle& Entity)
+	{
+		return HashCombine(Entity.Index, Entity.SerialNumber);
+	}
+
+	FString DebugGetDescription() const
+	{
+		return FString::Printf(TEXT("i: %d sn: %d"), Index, SerialNumber);
+	}
 };
 
-UENUM()
-enum class EMassProcessingPhase : uint8
+template struct MASSENTITY_API TScriptStructTypeBitSet<FMassFragment>;
+using FMassFragmentBitSet = TScriptStructTypeBitSet<FMassFragment>;
+template struct MASSENTITY_API TScriptStructTypeBitSet<FMassTag>;
+using FMassTagBitSet = TScriptStructTypeBitSet<FMassTag>;
+template struct MASSENTITY_API TScriptStructTypeBitSet<FMassChunkFragment>;
+using FMassChunkFragmentBitSet = TScriptStructTypeBitSet<FMassChunkFragment>;
+
+/** The type summarily describing a composition of an entity or an archetype. It contains information on both the
+ *  components as well as tags */
+struct FMassCompositionDescriptor
 {
-	PrePhysics,
-	StartPhysics,
-	DuringPhysics,
-	EndPhysics,
-	PostPhysics,
-	FrameEnd,
-	MAX,
+	FMassCompositionDescriptor() = default;
+	FMassCompositionDescriptor(const FMassFragmentBitSet& InComponents, const FMassTagBitSet& InTags, const FMassChunkFragmentBitSet& InChunkComponents)
+		: Components(InComponents), Tags(InTags), ChunkComponents(InChunkComponents)
+	{}
+
+	FMassCompositionDescriptor(TConstArrayView<const UScriptStruct*> InComponents, const FMassTagBitSet& InTags, const FMassChunkFragmentBitSet& InChunkComponents)
+		: FMassCompositionDescriptor(FMassFragmentBitSet(InComponents), InTags, InChunkComponents)
+	{}
+
+	FMassCompositionDescriptor(TConstArrayView<FInstancedStruct> InComponentInstances, const FMassTagBitSet& InTags, const FMassChunkFragmentBitSet& InChunkComponents)
+		: FMassCompositionDescriptor(FMassFragmentBitSet(InComponentInstances), InTags, InChunkComponents)
+	{}
+
+	FMassCompositionDescriptor(FMassFragmentBitSet&& InComponents, FMassTagBitSet&& InTags, FMassChunkFragmentBitSet&& InChunkComponents)
+		: Components(MoveTemp(InComponents)), Tags(MoveTemp(InTags)), ChunkComponents(MoveTemp(InChunkComponents))
+	{}
+
+	void Reset()
+	{
+		Components.Reset();
+		Tags.Reset();
+	}
+
+	bool IsEquivalent(const FMassFragmentBitSet& InComponentBitSet, const FMassTagBitSet& InTagBitSet, const FMassChunkFragmentBitSet& InChunkComponentsBitSet) const
+	{
+		return Components.IsEquivalent(InComponentBitSet) && Tags.IsEquivalent(InTagBitSet) && ChunkComponents.IsEquivalent(InChunkComponentsBitSet);
+	}
+
+	bool IsEquivalent(const FMassCompositionDescriptor& OtherDescriptor) const
+	{
+		return IsEquivalent(OtherDescriptor.Components, OtherDescriptor.Tags, OtherDescriptor.ChunkComponents);
+	}
+
+	bool IsEmpty() const { return Components.IsEmpty() && Tags.IsEmpty() && ChunkComponents.IsEmpty(); }
+
+	static uint32 CalculateHash(const FMassFragmentBitSet& InComponents, const FMassTagBitSet& InTags, const FMassChunkFragmentBitSet& InChunkComponents)
+	{
+		const uint32 ComponentsHash = GetTypeHash(InComponents);
+		const uint32 TagsHash = GetTypeHash(InTags);
+		const uint32 ChunkComponentsHash = GetTypeHash(InChunkComponents);
+		return HashCombine(HashCombine(ComponentsHash, TagsHash), ChunkComponentsHash);
+	}	
+
+	uint32 CalculateHash() const 
+	{
+		return CalculateHash(Components, Tags, ChunkComponents);
+	}
+
+	FMassFragmentBitSet Components;
+	FMassTagBitSet Tags;
+	FMassChunkFragmentBitSet ChunkComponents;
 };
diff --git a/Engine/Plugins/Runtime/MassEntity/Source/MassEntity/Public/MassEntityUtils.h b/Engine/Plugins/Runtime/MassEntity/Source/MassEntity/Public/MassEntityUtils.h
--- a/Engine/Plugins/Runtime/MassEntity/Source/MassEntity/Public/MassEntityUtils.h
+++ b/Engine/Plugins/Runtime/MassEntity/Source/MassEntity/Public/MassEntityUtils.h
@@ -2,7 +2,7 @@
 
 #pragma once
 
-#include "MassEntityTypes.h"
+#include "MassProcessingTypes.h"
 
 namespace UE { namespace Pipe { namespace Utils
 {
diff --git a/Engine/Plugins/Runtime/MassEntity/Source/MassEntity/Public/MassEntityView.h b/Engine/Plugins/Runtime/MassEntity/Source/MassEntity/Public/MassEntityView.h
new file mode 100644
--- /dev/null
+++ b/Engine/Plugins/Runtime/MassEntity/Source/MassEntity/Public/MassEntityView.h
@@ -0,0 +1,91 @@
+// Copyright Epic Games, Inc. All Rights Reserved.
+
+#pragma once
+
+#include "MassArchetypeTypes.h"
+#include "InstancedStruct.h"
+#include "MassEntityView.generated.h"
+
+
+class UMassEntitySubsystem;
+struct FMassArchetypeData;
+struct FArchetypeHandle;
+struct FMassArchetypeFragmentConfig;
+
+/** 
+ * The type representing a single entity in a single archetype. It's of a very transient nature so we guarantee it's 
+ * validity only within the scope it has been created in. Don't store it. 
+ */
+USTRUCT()
+struct MASSENTITY_API FMassEntityView
+{
+	GENERATED_BODY()
+
+	FMassEntityView() = default;
+
+	/** 
+	 *  Resolves Entity against ArchetypeHandle. Note that this approach requires the caller to ensure that Entity
+	 *  indeed belongs to ArchetypeHandle. If not the call will fail a check. As a remedy calling the 
+	 *  UMassEntitySubsystem-flavored constructor is recommended since it will first find the appropriate archetype for
+	 *  Entity. 
+	 */
+	FMassEntityView(const FArchetypeHandle& ArchetypeHandle, FMassEntityHandle Entity);
+
+	/** 
+	 *  Finds the archetype Entity belongs to and then resolves against it. The caller is responsible for ensuring
+	 *  that the given Entity is in fact a valid ID tied to any of the archetypes 
+	 */
+	FMassEntityView(const UMassEntitySubsystem& EntitySubsystem, FMassEntityHandle Entity);
+
+	FMassEntityHandle GetEntity() const	{ return Entity; }
+
+	/** will fail a check if the viewed entity doesn't have the given component */	
+	template<typename T>
+	T& GetComponentData() const
+	{
+		static_assert(!TIsDerivedFrom<T, FMassTag>::IsDerived,
+			"Given struct doesn't represent a valid fragment type but a tag. Use HasTag instead.");
+		static_assert(TIsDerivedFrom<T, FMassTag>::IsDerived || TIsDerivedFrom<T, FMassFragment>::IsDerived,
+			"Given struct doesn't represent a valid fragment type. Make sure to inherit from FMassFragment or one of its child-types.");
+
+		return *((T*)GetComponentPtrChecked(*T::StaticStruct()));
+	}
+		
+	/** if the viewed entity doesn't have the given component the function will return null */
+	template<typename T>
+	T* GetComponentDataPtr() const
+	{
+		static_assert(!TIsDerivedFrom<T, FMassTag>::IsDerived,
+			"Given struct doesn't represent a valid fragment type but a tag. Use HasTag instead.");
+		static_assert(TIsDerivedFrom<T, FMassTag>::IsDerived || TIsDerivedFrom<T, FMassFragment>::IsDerived,
+			"Given struct doesn't represent a valid fragment type. Make sure to inherit from FMassFragment or one of its child-types.");
+
+		return (T*)GetComponentPtr(*T::StaticStruct());
+	}
+
+	FStructView GetComponentDataStruct(const UScriptStruct* ComponentType) const
+	{
+		check(ComponentType);
+		return FStructView(ComponentType, static_cast<uint8*>(GetComponentPtr(*ComponentType)));
+	}
+
+	template<typename T>
+	bool HasTag() const
+	{
+		static_assert(TIsDerivedFrom<T, FMassTag>::IsDerived, "Given struct doesn't represent a valid tag type. Make sure to inherit from FMassTag or one of its child-types.");
+		return HasTag(*T::StaticStruct());
+	}
+
+	bool IsSet() const { return Archetype != nullptr && EntityHandle.IsValid(); }
+	bool operator==(const FMassEntityView& Other) const { return Archetype == Other.Archetype && EntityHandle == Other.EntityHandle; }
+
+protected:
+	void* GetComponentPtr(const UScriptStruct& ComponentType) const;
+	void* GetComponentPtrChecked(const UScriptStruct& ComponentType) const;
+	bool HasTag(const UScriptStruct& TagType) const;
+
+private:
+	FMassEntityHandle Entity;
+	FInternalEntityHandle EntityHandle;
+	FMassArchetypeData* Archetype = nullptr;
+};
diff --git a/Engine/Plugins/Runtime/MassEntity/Source/MassEntity/Public/MassProcessingPhase.h b/Engine/Plugins/Runtime/MassEntity/Source/MassEntity/Public/MassProcessingPhase.h
--- a/Engine/Plugins/Runtime/MassEntity/Source/MassEntity/Public/MassProcessingPhase.h
+++ b/Engine/Plugins/Runtime/MassEntity/Source/MassEntity/Public/MassProcessingPhase.h
@@ -6,7 +6,7 @@
 #include "UObject/UObjectGlobals.h"
 #include "UObject/Object.h"
 #include "Engine/EngineBaseTypes.h"
-#include "MassEntityTypes.h"
+#include "MassProcessingTypes.h"
 #include "MassProcessingPhase.generated.h"
 
 
diff --git a/Engine/Plugins/Runtime/MassEntity/Source/MassEntity/Public/MassProcessingTypes.h b/Engine/Plugins/Runtime/MassEntity/Source/MassEntity/Public/MassProcessingTypes.h
new file mode 100644
--- /dev/null
+++ b/Engine/Plugins/Runtime/MassEntity/Source/MassEntity/Public/MassProcessingTypes.h
@@ -0,0 +1,117 @@
+﻿// Copyright Epic Games, Inc. All Rights Reserved.
+
+#pragma once
+
+#include "StructUtilsTypes.h"
+#include "InstancedStruct.h"
+#include "MassProcessingTypes.generated.h"
+
+#define WITH_MASSENTITY_DEBUG (!(UE_BUILD_SHIPPING || UE_BUILD_SHIPPING_WITH_EDITOR || UE_BUILD_TEST) && WITH_STRUCTUTILS_DEBUG && 1)
+
+MASSENTITY_API DECLARE_LOG_CATEGORY_EXTERN(LogPipe, Warning, All);
+
+class UMassEntitySubsystem;
+class UMassProcessor;
+class UMassSchematic;
+class UMassCompositeProcessor;
+
+UENUM(meta = (Bitflags, UseEnumValuesAsMaskValuesInEditor = "true"))
+enum class EProcessorExecutionFlags : uint8
+{
+	None = 0 UMETA(Hidden),
+	Standalone = 1 << 0,
+	Server = 1 << 1,
+	Client = 1 << 2,
+	All = Standalone | Server | Client UMETA(Hidden)
+};
+ENUM_CLASS_FLAGS(EProcessorExecutionFlags);
+
+USTRUCT()
+struct FProcessorAuxDataBase
+{
+	GENERATED_BODY()
+};
+
+USTRUCT(BlueprintType)
+struct MASSENTITY_API FMassProcessingContext
+{
+	GENERATED_BODY()
+
+	UPROPERTY()
+	UMassEntitySubsystem* EntitySubsystem = nullptr;
+	
+	UPROPERTY()
+	float DeltaSeconds = 0.f;
+
+	UPROPERTY()
+	FInstancedStruct AuxData;
+
+	FMassProcessingContext() = default;
+	FMassProcessingContext(UMassEntitySubsystem& InEntities, const float InDeltaSeconds);
+};
+
+/** 
+ *  A helper type that converts a set of UMassSchematics into a runtime-usable array of MassProcessor copies
+ */
+USTRUCT()
+struct MASSENTITY_API FMassRuntimePipeline
+{
+	GENERATED_BODY()
+
+	UPROPERTY()
+	TArray<UMassProcessor*> Processors;
+
+	void Reset();
+	void Initialize(UObject& Owner);
+	
+	/** Creates runtime copies of the given UMassProcessors collection. */
+	void SetProcessors(TArray<UMassProcessor*>&& InProcessors);
+
+	/** Creates runtime copies of UMassProcessors declared in Schematics using InOwner as new UMassProcessors' outer. */
+	void InitializeFromSchematics(TConstArrayView<TSoftObjectPtr<UMassSchematic>> Schematics, UObject& InOwner);
+
+	/** Creates runtime copies of UMassProcessors given in InProcessors input parameter, using InOwner as new UMassProcessors' outer. */
+	void CreateFromArray(TConstArrayView<const UMassProcessor*> InProcessors, UObject& InOwner);
+
+	/** Calls CreateFromArray and calls Initialize on all processors afterwards. */
+	void InitializeFromArray(TConstArrayView<const UMassProcessor*> InProcessors, UObject& InOwner);
+	
+	/** Creates runtime instances of UMassProcessors for each processor class given via InProcessorClasses. 
+	 *  The instances will be created with InOwner as outer. */
+	void InitializeFromClassArray(TConstArrayView<TSubclassOf<UMassProcessor>> InProcessorClasses, UObject& InOwner);
+
+	/** Creates a runtime instance of every processors in the given InProcessors array. If a processor of that class
+	 *  already exists in Processors array it gets overridden. Otherwise it gets added to the end of the collection.*/
+	void AppendOrOverrideRuntimeProcessorCopies(TConstArrayView<const UMassProcessor*> InProcessors, UObject& InOwner);
+
+	/** Creates a runtime instance of every processors in the given array if there's no processor of that class in Processors already.
+	 *  Call this function when adding processors to an already configured FMassRuntimePipeline instance. If you're creating 
+	 *  one from scratch calling any of the InitializeFrom* methods will be more efficient (and will produce same results)
+	 *  or call AppendOrOverrideRuntimeProcessorCopies.*/
+	void AppendUniqueRuntimeProcessorCopies(TConstArrayView<const UMassProcessor*> InProcessors, UObject& InOwner);
+
+	/** Adds InProcessor to Processors without any additional checks */
+	void AppendProcessor(UMassProcessor& Processor);
+
+	/** goes through Processor looking for a UMassCompositeProcessor instance which GroupName matches the one given as the parameter */
+	UMassCompositeProcessor* FindTopLevelGroupByName(const FName GroupName);
+
+	void DebugOutputDescription(FOutputDevice& Ar) const;
+
+	bool HasProcessorOfExactClass(TSubclassOf<UMassProcessor> InClass) const;
+	bool IsEmpty() const { return Processors.IsEmpty();}
+
+	MASSENTITY_API friend uint32 GetTypeHash(const FMassRuntimePipeline& Instance);
+};
+
+UENUM()
+enum class EMassProcessingPhase : uint8
+{
+	PrePhysics,
+	StartPhysics,
+	DuringPhysics,
+	EndPhysics,
+	PostPhysics,
+	FrameEnd,
+	MAX,
+};
diff --git a/Engine/Plugins/Runtime/MassEntity/Source/MassEntity/Public/MassProcessor.h b/Engine/Plugins/Runtime/MassEntity/Source/MassEntity/Public/MassProcessor.h
--- a/Engine/Plugins/Runtime/MassEntity/Source/MassEntity/Public/MassProcessor.h
+++ b/Engine/Plugins/Runtime/MassEntity/Source/MassEntity/Public/MassProcessor.h
@@ -4,9 +4,9 @@
 
 #include "CoreMinimal.h"
 #include "MassEntitySubsystem.h"
-#include "MassEntityTypes.h"
+#include "MassProcessingTypes.h"
 #include "Async/TaskGraphInterfaces.h"
-#include "LWCCommandBuffer.h"
+#include "MassCommandBuffer.h"
 #include "MassProcessorDependencySolver.h"
 #include "MassProcessor.generated.h"
 
diff --git a/Engine/Plugins/Runtime/MassEntity/Source/MassEntity/Public/MassProcessorDependencySolver.h b/Engine/Plugins/Runtime/MassEntity/Source/MassEntity/Public/MassProcessorDependencySolver.h
--- a/Engine/Plugins/Runtime/MassEntity/Source/MassEntity/Public/MassProcessorDependencySolver.h
+++ b/Engine/Plugins/Runtime/MassEntity/Source/MassEntity/Public/MassProcessorDependencySolver.h
@@ -2,7 +2,7 @@
 
 #pragma once
 
-#include "MassEntityTypes.h"
+#include "MassProcessingTypes.h"
 
 
 class UMassProcessor;
diff --git a/Engine/Plugins/Runtime/MassEntity/Source/MassEntity/Public/MassSchematic.h b/Engine/Plugins/Runtime/MassEntity/Source/MassEntity/Public/MassSchematic.h
--- a/Engine/Plugins/Runtime/MassEntity/Source/MassEntity/Public/MassSchematic.h
+++ b/Engine/Plugins/Runtime/MassEntity/Source/MassEntity/Public/MassSchematic.h
@@ -4,7 +4,7 @@
 
 #include "CoreMinimal.h"
 #include "Engine/DataAsset.h"
-#include "MassEntityTypes.h"
+#include "MassProcessingTypes.h"
 #include "MassSchematic.generated.h"
 
 
diff --git a/Engine/Plugins/Runtime/MassEntity/Source/MassEntityTestSuite/Private/MassCompositeProcessorTest.cpp b/Engine/Plugins/Runtime/MassEntity/Source/MassEntityTestSuite/Private/MassCompositeProcessorTest.cpp
--- a/Engine/Plugins/Runtime/MassEntity/Source/MassEntityTestSuite/Private/MassCompositeProcessorTest.cpp
+++ b/Engine/Plugins/Runtime/MassEntity/Source/MassEntityTestSuite/Private/MassCompositeProcessorTest.cpp
@@ -5,7 +5,7 @@
 
 #include "Engine/World.h"
 #include "MassEntitySubsystem.h"
-#include "MassEntityTypes.h"
+#include "MassProcessingTypes.h"
 #include "MassEntityTestTypes.h"
 #include "MassExecutor.h"
 
diff --git a/Engine/Plugins/Runtime/MassEntity/Source/MassEntityTestSuite/Private/MassEntityTest.cpp b/Engine/Plugins/Runtime/MassEntity/Source/MassEntityTestSuite/Private/MassEntityTest.cpp
--- a/Engine/Plugins/Runtime/MassEntity/Source/MassEntityTestSuite/Private/MassEntityTest.cpp
+++ b/Engine/Plugins/Runtime/MassEntity/Source/MassEntityTestSuite/Private/MassEntityTest.cpp
@@ -5,7 +5,7 @@
 
 #include "Engine/World.h"
 #include "MassEntitySubsystem.h"
-#include "MassEntityTypes.h"
+#include "MassProcessingTypes.h"
 #include "MassEntityTestTypes.h"
 #include "MassExecutor.h"
 
diff --git a/Engine/Plugins/Runtime/MassEntity/Source/MassEntityTestSuite/Private/MassExecutionTest.cpp b/Engine/Plugins/Runtime/MassEntity/Source/MassEntityTestSuite/Private/MassExecutionTest.cpp
--- a/Engine/Plugins/Runtime/MassEntity/Source/MassEntityTestSuite/Private/MassExecutionTest.cpp
+++ b/Engine/Plugins/Runtime/MassEntity/Source/MassEntityTestSuite/Private/MassExecutionTest.cpp
@@ -3,7 +3,7 @@
 #include "CoreMinimal.h"
 #include "Engine/World.h"
 #include "MassEntitySubsystem.h"
-#include "MassEntityTypes.h"
+#include "MassProcessingTypes.h"
 #include "MassEntityTestTypes.h"
 #include "MassExecutor.h"
 
diff --git a/Engine/Plugins/Runtime/MassEntity/Source/MassEntityTestSuite/Private/MassProcessorTest.cpp b/Engine/Plugins/Runtime/MassEntity/Source/MassEntityTestSuite/Private/MassProcessorTest.cpp
--- a/Engine/Plugins/Runtime/MassEntity/Source/MassEntityTestSuite/Private/MassProcessorTest.cpp
+++ b/Engine/Plugins/Runtime/MassEntity/Source/MassEntityTestSuite/Private/MassProcessorTest.cpp
@@ -5,7 +5,7 @@
 
 #include "Engine/World.h"
 #include "MassEntitySubsystem.h"
-#include "MassEntityTypes.h"
+#include "MassProcessingTypes.h"
 #include "MassEntityTestTypes.h"
 #include "MassExecutor.h"
 
diff --git a/Engine/Plugins/Runtime/MassEntity/Source/MassEntityTestSuite/Private/MassQueryTest.cpp b/Engine/Plugins/Runtime/MassEntity/Source/MassEntityTestSuite/Private/MassQueryTest.cpp
--- a/Engine/Plugins/Runtime/MassEntity/Source/MassEntityTestSuite/Private/MassQueryTest.cpp
+++ b/Engine/Plugins/Runtime/MassEntity/Source/MassEntityTestSuite/Private/MassQueryTest.cpp
@@ -5,7 +5,7 @@
 
 #include "Engine/World.h"
 #include "MassEntitySubsystem.h"
-#include "MassEntityTypes.h"
+#include "MassProcessingTypes.h"
 #include "MassEntityTestTypes.h"
 #include "MassExecutor.h"
 
diff --git a/Engine/Plugins/Runtime/MassEntity/Source/MassEntityTestSuite/Public/LWComponentTestFarmPlot.h b/Engine/Plugins/Runtime/MassEntity/Source/MassEntityTestSuite/Public/LWComponentTestFarmPlot.h
--- a/Engine/Plugins/Runtime/MassEntity/Source/MassEntityTestSuite/Public/LWComponentTestFarmPlot.h
+++ b/Engine/Plugins/Runtime/MassEntity/Source/MassEntityTestSuite/Public/LWComponentTestFarmPlot.h
@@ -6,9 +6,9 @@
 #include "GameFramework/Actor.h"
 #include "Engine/DataTable.h"
 
-#include "LWComponentTypes.h"
+#include "MassEntityTypes.h"
 #include "MassEntitySubsystem.h"
-#include "LWCCommandBuffer.h"
+#include "MassCommandBuffer.h"
 
 #include "LWComponentTestFarmPlot.generated.h"
 
diff --git a/Engine/Plugins/Runtime/MassEntity/Source/MassEntityTestSuite/Public/MassEntityTestTypes.h b/Engine/Plugins/Runtime/MassEntity/Source/MassEntityTestSuite/Public/MassEntityTestTypes.h
--- a/Engine/Plugins/Runtime/MassEntity/Source/MassEntityTestSuite/Public/MassEntityTestTypes.h
+++ b/Engine/Plugins/Runtime/MassEntity/Source/MassEntityTestSuite/Public/MassEntityTestTypes.h
@@ -3,7 +3,7 @@
 #pragma once
 
 #include "MassProcessor.h"
-#include "LWComponentTypes.h"
+#include "MassEntityTypes.h"
 #include "MassEntitySubsystem.h"
 #include "AITestsCommon.h"
 #include "MassEntityTestTypes.generated.h"
diff --git a/Engine/Plugins/Runtime/MassGameplay/Source/MassActors/Private/MassAgentComponent.cpp b/Engine/Plugins/Runtime/MassGameplay/Source/MassActors/Private/MassAgentComponent.cpp
--- a/Engine/Plugins/Runtime/MassGameplay/Source/MassActors/Private/MassAgentComponent.cpp
+++ b/Engine/Plugins/Runtime/MassGameplay/Source/MassActors/Private/MassAgentComponent.cpp
@@ -2,7 +2,7 @@
 
 #include "MassAgentComponent.h"
 #include "EngineUtils.h"
-#include "EntityView.h"
+#include "MassEntityView.h"
 #include "MassEntitySubsystem.h"
 #include "MassCommonTypes.h"
 #include "MassAgentSubsystem.h"
diff --git a/Engine/Plugins/Runtime/MassGameplay/Source/MassActors/Private/MassAgentSubsystem.cpp b/Engine/Plugins/Runtime/MassGameplay/Source/MassActors/Private/MassAgentSubsystem.cpp
--- a/Engine/Plugins/Runtime/MassGameplay/Source/MassActors/Private/MassAgentSubsystem.cpp
+++ b/Engine/Plugins/Runtime/MassGameplay/Source/MassActors/Private/MassAgentSubsystem.cpp
@@ -1,7 +1,7 @@
 // Copyright Epic Games, Inc. All Rights Reserved.
 
 #include "MassAgentSubsystem.h"
-#include "LWCCommandBuffer.h"
+#include "MassCommandBuffer.h"
 #include "Engine/World.h"
 #include "MassEntitySubsystem.h"
 #include "MassCommonTypes.h"
@@ -11,7 +11,7 @@
 #include "MassAgentComponent.h"
 #include "MassEntityTemplateRegistry.h"
 #include "VisualLogger/VisualLogger.h"
-#include "EntityView.h"
+#include "MassEntityView.h"
 #include "MassReplicationManager.h"
 #include "Engine/NetDriver.h"
 
diff --git a/Engine/Plugins/Runtime/MassGameplay/Source/MassActors/Private/MassAgentTraits.cpp b/Engine/Plugins/Runtime/MassGameplay/Source/MassActors/Private/MassAgentTraits.cpp
--- a/Engine/Plugins/Runtime/MassGameplay/Source/MassActors/Private/MassAgentTraits.cpp
+++ b/Engine/Plugins/Runtime/MassGameplay/Source/MassActors/Private/MassAgentTraits.cpp
@@ -13,7 +13,7 @@
 #include "GameFramework/Actor.h"
 #include "GameFramework/CharacterMovementComponent.h"
 #include "VisualLogger/VisualLogger.h"
-#include "EntityView.h"
+#include "MassEntityView.h"
 
 namespace FMassAgentTraitsHelper 
 {
diff --git a/Engine/Plugins/Runtime/MassGameplay/Source/MassActors/Public/MassActorHelper.h b/Engine/Plugins/Runtime/MassGameplay/Source/MassActors/Public/MassActorHelper.h
--- a/Engine/Plugins/Runtime/MassGameplay/Source/MassActors/Public/MassActorHelper.h
+++ b/Engine/Plugins/Runtime/MassGameplay/Source/MassActors/Public/MassActorHelper.h
@@ -2,7 +2,7 @@
 
 #pragma once
 
-#include "LWComponentTypes.h"
+#include "MassEntityTypes.h"
 
 class UScriptStruct;
 class AActor;
diff --git a/Engine/Plugins/Runtime/MassGameplay/Source/MassActors/Public/MassActorSubsystem.h b/Engine/Plugins/Runtime/MassGameplay/Source/MassActors/Public/MassActorSubsystem.h
--- a/Engine/Plugins/Runtime/MassGameplay/Source/MassActors/Public/MassActorSubsystem.h
+++ b/Engine/Plugins/Runtime/MassGameplay/Source/MassActors/Public/MassActorSubsystem.h
@@ -2,7 +2,7 @@
 
 #pragma once
 
-#include "LWComponentTypes.h"
+#include "MassEntityTypes.h"
 #include "MassCommonFragments.h"
 #include "MassEntityTemplate.h"
 #include "Misc/MTAccessDetector.h"
diff --git a/Engine/Plugins/Runtime/MassGameplay/Source/MassActors/Public/MassAgentSubsystem.h b/Engine/Plugins/Runtime/MassGameplay/Source/MassActors/Public/MassAgentSubsystem.h
--- a/Engine/Plugins/Runtime/MassGameplay/Source/MassActors/Public/MassAgentSubsystem.h
+++ b/Engine/Plugins/Runtime/MassGameplay/Source/MassActors/Public/MassAgentSubsystem.h
@@ -2,8 +2,8 @@
 
 #pragma once
 
-#include "LWComponentTypes.h"
 #include "MassEntityTypes.h"
+#include "MassProcessingTypes.h"
 #include "MassEntityTemplate.h"
 #include "Subsystems/WorldSubsystem.h"
 #include "MassAgentSubsystem.generated.h"
diff --git a/Engine/Plugins/Runtime/MassGameplay/Source/MassCommon/Public/MassCommonFragments.h b/Engine/Plugins/Runtime/MassGameplay/Source/MassCommon/Public/MassCommonFragments.h
--- a/Engine/Plugins/Runtime/MassGameplay/Source/MassCommon/Public/MassCommonFragments.h
+++ b/Engine/Plugins/Runtime/MassGameplay/Source/MassCommon/Public/MassCommonFragments.h
@@ -2,7 +2,7 @@
 
 #pragma once
 
-#include "LWComponentTypes.h"
+#include "MassEntityTypes.h"
 #include "AI/Navigation/NavigationTypes.h"
 #include "MassCommonFragments.generated.h"
 
diff --git a/Engine/Plugins/Runtime/MassGameplay/Source/MassCommon/Public/MassCommonTypes.h b/Engine/Plugins/Runtime/MassGameplay/Source/MassCommon/Public/MassCommonTypes.h
--- a/Engine/Plugins/Runtime/MassGameplay/Source/MassCommon/Public/MassCommonTypes.h
+++ b/Engine/Plugins/Runtime/MassGameplay/Source/MassCommon/Public/MassCommonTypes.h
@@ -2,8 +2,8 @@
 
 #pragma once
 
-#include "LWComponentTypes.h"
 #include "MassEntityTypes.h"
+#include "MassProcessingTypes.h"
 #include "InstancedStruct.h"
 #include "SequentialID.h"
 #include "MassCommonTypes.generated.h"
diff --git a/Engine/Plugins/Runtime/MassGameplay/Source/MassGameplayDebug/Public/MassGameplayDebugTypes.h b/Engine/Plugins/Runtime/MassGameplay/Source/MassGameplayDebug/Public/MassGameplayDebugTypes.h
--- a/Engine/Plugins/Runtime/MassGameplay/Source/MassGameplayDebug/Public/MassGameplayDebugTypes.h
+++ b/Engine/Plugins/Runtime/MassGameplay/Source/MassGameplayDebug/Public/MassGameplayDebugTypes.h
@@ -8,7 +8,7 @@
 #include "Logging/LogMacros.h"
 #include "Templates/SubclassOf.h"
 #include "Engine/DataTable.h"
-#include "LWComponentTypes.h"
+#include "MassEntityTypes.h"
 #include "MassSpawnerTypes.h"
 #include "MassGameplayDebugTypes.generated.h"
 
diff --git a/Engine/Plugins/Runtime/MassGameplay/Source/MassLOD/Public/MassLODCollectorProcessor.h b/Engine/Plugins/Runtime/MassGameplay/Source/MassLOD/Public/MassLODCollectorProcessor.h
--- a/Engine/Plugins/Runtime/MassGameplay/Source/MassLOD/Public/MassLODCollectorProcessor.h
+++ b/Engine/Plugins/Runtime/MassGameplay/Source/MassLOD/Public/MassLODCollectorProcessor.h
@@ -2,7 +2,7 @@
 
 #pragma once
 
-#include "LWComponentTypes.h"
+#include "MassEntityTypes.h"
 #include "MassLODTypes.h"
 #include "MassLODManager.h"
 #include "MassLODCollector.h"
diff --git a/Engine/Plugins/Runtime/MassGameplay/Source/MassLOD/Public/MassLODTickRateController.h b/Engine/Plugins/Runtime/MassGameplay/Source/MassLOD/Public/MassLODTickRateController.h
--- a/Engine/Plugins/Runtime/MassGameplay/Source/MassLOD/Public/MassLODTickRateController.h
+++ b/Engine/Plugins/Runtime/MassGameplay/Source/MassLOD/Public/MassLODTickRateController.h
@@ -4,7 +4,7 @@
 
 #include "MassLODLogic.h"
 #include "MassLODUtils.h"
-#include "LWCCommandBuffer.h"
+#include "MassCommandBuffer.h"
 
 /**
  * Helper struct to control LOD tick rate for each agent, 
diff --git a/Engine/Plugins/Runtime/MassGameplay/Source/MassLOD/Public/MassLODTypes.h b/Engine/Plugins/Runtime/MassGameplay/Source/MassLOD/Public/MassLODTypes.h
--- a/Engine/Plugins/Runtime/MassGameplay/Source/MassLOD/Public/MassLODTypes.h
+++ b/Engine/Plugins/Runtime/MassGameplay/Source/MassLOD/Public/MassLODTypes.h
@@ -2,7 +2,7 @@
 
 #pragma once
 
-#include "LWComponentTypes.h"
+#include "MassEntityTypes.h"
 #include "Containers/StaticArray.h"
 #include "MassEntitySubsystem.h"
 
diff --git a/Engine/Plugins/Runtime/MassGameplay/Source/MassLOD/Public/MassLODUtils.h b/Engine/Plugins/Runtime/MassGameplay/Source/MassLOD/Public/MassLODUtils.h
--- a/Engine/Plugins/Runtime/MassGameplay/Source/MassLOD/Public/MassLODUtils.h
+++ b/Engine/Plugins/Runtime/MassGameplay/Source/MassLOD/Public/MassLODUtils.h
@@ -2,7 +2,7 @@
 
 #pragma once
 
-#include "LWComponentTypes.h"
+#include "MassEntityTypes.h"
 #include "MassEntitySubsystem.h"
 #include "MassCommonTypes.h"
 #include "MassLODTypes.h"
diff --git a/Engine/Plugins/Runtime/MassGameplay/Source/MassLOD/Public/MassRepresentationLODCollectorProcessor.h b/Engine/Plugins/Runtime/MassGameplay/Source/MassLOD/Public/MassRepresentationLODCollectorProcessor.h
--- a/Engine/Plugins/Runtime/MassGameplay/Source/MassLOD/Public/MassRepresentationLODCollectorProcessor.h
+++ b/Engine/Plugins/Runtime/MassGameplay/Source/MassLOD/Public/MassRepresentationLODCollectorProcessor.h
@@ -5,7 +5,7 @@
 #include "MassLODManager.h"
 #include "MassLODCollector.h"
 #include "MassLODCalculator.h"
-#include "LWComponentTypes.h"
+#include "MassEntityTypes.h"
 
 #include "MassRepresentationLODCollectorProcessor.generated.h"
 
diff --git a/Engine/Plugins/Runtime/MassGameplay/Source/MassLOD/Public/MassSimulationLOD.h b/Engine/Plugins/Runtime/MassGameplay/Source/MassLOD/Public/MassSimulationLOD.h
--- a/Engine/Plugins/Runtime/MassGameplay/Source/MassLOD/Public/MassSimulationLOD.h
+++ b/Engine/Plugins/Runtime/MassGameplay/Source/MassLOD/Public/MassSimulationLOD.h
@@ -7,7 +7,7 @@
 #include "MassLODCalculator.h"
 #include "MassLODTickRateController.h"
 #include "MassLODLogic.h"
-#include "LWComponentTypes.h"
+#include "MassEntityTypes.h"
 
 #include "MassSimulationLOD.generated.h"
 
diff --git a/Engine/Plugins/Runtime/MassGameplay/Source/MassMovement/Public/MassMovementTypes.h b/Engine/Plugins/Runtime/MassGameplay/Source/MassMovement/Public/MassMovementTypes.h
--- a/Engine/Plugins/Runtime/MassGameplay/Source/MassMovement/Public/MassMovementTypes.h
+++ b/Engine/Plugins/Runtime/MassGameplay/Source/MassMovement/Public/MassMovementTypes.h
@@ -2,7 +2,7 @@
 
 #pragma once
 
-#include "LWComponentTypes.h"
+#include "MassEntityTypes.h"
 #include "HierarchicalHashGrid2D.h"
 
 namespace UE::Mass::ProcessorGroupNames
diff --git a/Engine/Plugins/Runtime/MassGameplay/Source/MassReplication/Private/MassReplicationTransformHandlers.cpp b/Engine/Plugins/Runtime/MassGameplay/Source/MassReplication/Private/MassReplicationTransformHandlers.cpp
--- a/Engine/Plugins/Runtime/MassGameplay/Source/MassReplication/Private/MassReplicationTransformHandlers.cpp
+++ b/Engine/Plugins/Runtime/MassGameplay/Source/MassReplication/Private/MassReplicationTransformHandlers.cpp
@@ -1,7 +1,7 @@
 // Copyright Epic Games, Inc. All Rights Reserved.
 
 #include "MassReplicationTransformHandlers.h"
-#include "EntityQuery.h"
+#include "MassEntityQuery.h"
 #include "MassCommonFragments.h"
 #include "AIHelpers.h"
 
diff --git a/Engine/Plugins/Runtime/MassGameplay/Source/MassReplication/Public/MassClientBubbleHandler.h b/Engine/Plugins/Runtime/MassGameplay/Source/MassReplication/Public/MassClientBubbleHandler.h
--- a/Engine/Plugins/Runtime/MassGameplay/Source/MassReplication/Public/MassClientBubbleHandler.h
+++ b/Engine/Plugins/Runtime/MassGameplay/Source/MassReplication/Public/MassClientBubbleHandler.h
@@ -12,7 +12,7 @@
 #include "MassReplicationManager.h"
 #include "Engine/World.h"
 #include "MassEntityTemplate.h"
-#include "EntityView.h"
+#include "MassEntityView.h"
 
 #include "MassClientBubbleHandler.generated.h"
 
diff --git a/Engine/Plugins/Runtime/MassGameplay/Source/MassReplication/Public/MassReplicationManager.h b/Engine/Plugins/Runtime/MassGameplay/Source/MassReplication/Public/MassReplicationManager.h
--- a/Engine/Plugins/Runtime/MassGameplay/Source/MassReplication/Public/MassReplicationManager.h
+++ b/Engine/Plugins/Runtime/MassGameplay/Source/MassReplication/Public/MassReplicationManager.h
@@ -3,7 +3,7 @@
 #pragma once
 
 #include "MassCommonTypes.h"
-#include "LWComponentTypes.h"
+#include "MassEntityTypes.h"
 #include "MassReplicationTypes.h"
 #include "Subsystems/WorldSubsystem.h"
 
diff --git a/Engine/Plugins/Runtime/MassGameplay/Source/MassReplication/Public/MassReplicationTransformHandlers.h b/Engine/Plugins/Runtime/MassGameplay/Source/MassReplication/Public/MassReplicationTransformHandlers.h
--- a/Engine/Plugins/Runtime/MassGameplay/Source/MassReplication/Public/MassReplicationTransformHandlers.h
+++ b/Engine/Plugins/Runtime/MassGameplay/Source/MassReplication/Public/MassReplicationTransformHandlers.h
@@ -4,7 +4,7 @@
 
 #include "MassMovementTypes.h"
 #include "MassReplicationTypes.h"
-#include "EntityView.h"
+#include "MassEntityView.h"
 #include "MassClientBubbleHandler.h"
 #include "MassCommonFragments.h"
 #include "AIHelpers.h"
diff --git a/Engine/Plugins/Runtime/MassGameplay/Source/MassRepresentation/Private/MassRepresentationProcessor.cpp b/Engine/Plugins/Runtime/MassGameplay/Source/MassRepresentation/Private/MassRepresentationProcessor.cpp
--- a/Engine/Plugins/Runtime/MassGameplay/Source/MassRepresentation/Private/MassRepresentationProcessor.cpp
+++ b/Engine/Plugins/Runtime/MassGameplay/Source/MassRepresentation/Private/MassRepresentationProcessor.cpp
@@ -5,9 +5,9 @@
 #include "MassRepresentationUtils.h"
 #include "MassActorSubsystem.h"
 #include "MassCommonFragments.h"
-#include "LWCCommandBuffer.h"
+#include "MassCommandBuffer.h"
 #include "MassEntitySubsystem.h"
-#include "EntityView.h"
+#include "MassEntityView.h"
 #include "Engine/World.h"
 
 //----------------------------------------------------------------------//
diff --git a/Engine/Plugins/Runtime/MassGameplay/Source/MassRepresentation/Private/MassRepresentationSubsystem.cpp b/Engine/Plugins/Runtime/MassGameplay/Source/MassRepresentation/Private/MassRepresentationSubsystem.cpp
--- a/Engine/Plugins/Runtime/MassGameplay/Source/MassRepresentation/Private/MassRepresentationSubsystem.cpp
+++ b/Engine/Plugins/Runtime/MassGameplay/Source/MassRepresentation/Private/MassRepresentationSubsystem.cpp
@@ -14,7 +14,7 @@
 #include "MassEntitySubsystem.h"
 #include "MassRepresentationFragments.h"
 #include "MassRepresentationProcessor.h"
-#include "EntityView.h"
+#include "MassEntityView.h"
 #include "WorldPartition/WorldPartitionSubsystem.h"
 #include "WorldPartition/WorldPartitionRuntimeCell.h"
 
diff --git a/Engine/Plugins/Runtime/MassGameplay/Source/MassRepresentation/Public/MassRepresentationFragments.h b/Engine/Plugins/Runtime/MassGameplay/Source/MassRepresentation/Public/MassRepresentationFragments.h
--- a/Engine/Plugins/Runtime/MassGameplay/Source/MassRepresentation/Public/MassRepresentationFragments.h
+++ b/Engine/Plugins/Runtime/MassGameplay/Source/MassRepresentation/Public/MassRepresentationFragments.h
@@ -3,7 +3,7 @@
 #pragma once
 
 #include "MassLODTypes.h"
-#include "LWComponentTypes.h"
+#include "MassEntityTypes.h"
 #include "MassEntitySubsystem.h"
 #include "MassActorSpawnerSubsystem.h"
 #include "MassRepresentationTypes.h"
diff --git a/Engine/Plugins/Runtime/MassGameplay/Source/MassRepresentation/Public/MassRepresentationUtils.h b/Engine/Plugins/Runtime/MassGameplay/Source/MassRepresentation/Public/MassRepresentationUtils.h
--- a/Engine/Plugins/Runtime/MassGameplay/Source/MassRepresentation/Public/MassRepresentationUtils.h
+++ b/Engine/Plugins/Runtime/MassGameplay/Source/MassRepresentation/Public/MassRepresentationUtils.h
@@ -2,7 +2,7 @@
 
 #pragma once
 
-#include "LWComponentTypes.h"
+#include "MassEntityTypes.h"
 #include "MassEntitySubsystem.h"
 #include "MassCommonTypes.h"
 #include "MassRepresentationTypes.h"
diff --git a/Engine/Plugins/Runtime/MassGameplay/Source/MassSignals/Private/MassSignalProcessorBase.cpp b/Engine/Plugins/Runtime/MassGameplay/Source/MassSignals/Private/MassSignalProcessorBase.cpp
--- a/Engine/Plugins/Runtime/MassGameplay/Source/MassSignals/Private/MassSignalProcessorBase.cpp
+++ b/Engine/Plugins/Runtime/MassGameplay/Source/MassSignals/Private/MassSignalProcessorBase.cpp
@@ -2,7 +2,7 @@
 
 #include "MassSignalProcessorBase.h"
 #include "MassSignalSubsystem.h"
-#include "ArchetypeTypes.h"
+#include "MassArchetypeTypes.h"
 #include "Engine/World.h"
 
 void UMassSignalProcessorBase::Initialize(UObject& Owner)
diff --git a/Engine/Plugins/Runtime/MassGameplay/Source/MassSignals/Private/MassSignalSubsystem.cpp b/Engine/Plugins/Runtime/MassGameplay/Source/MassSignals/Private/MassSignalSubsystem.cpp
--- a/Engine/Plugins/Runtime/MassGameplay/Source/MassSignals/Private/MassSignalSubsystem.cpp
+++ b/Engine/Plugins/Runtime/MassGameplay/Source/MassSignals/Private/MassSignalSubsystem.cpp
@@ -1,7 +1,7 @@
 // Copyright Epic Games, Inc. All Rights Reserved.
 
 #include "MassSignalSubsystem.h"
-#include "LWCCommandBuffer.h"
+#include "MassCommandBuffer.h"
 #include "MassSignalTypes.h"
 #include "VisualLogger/VisualLogger.h"
 #include "ProfilingDebugging/CsvProfiler.h"
diff --git a/Engine/Plugins/Runtime/MassGameplay/Source/MassSimulation/Public/MassSimulationLocalCoordinator.h b/Engine/Plugins/Runtime/MassGameplay/Source/MassSimulation/Public/MassSimulationLocalCoordinator.h
--- a/Engine/Plugins/Runtime/MassGameplay/Source/MassSimulation/Public/MassSimulationLocalCoordinator.h
+++ b/Engine/Plugins/Runtime/MassGameplay/Source/MassSimulation/Public/MassSimulationLocalCoordinator.h
@@ -10,7 +10,7 @@
 #include "Engine/World.h"
 #include "GameFramework/Actor.h"
 #include "Engine/DataTable.h"
-#include "MassEntityTypes.h"
+#include "MassProcessingTypes.h"
 #include "MassCommonTypes.h"
 #include "MassSimulationLocalCoordinator.generated.h"
 
diff --git a/Engine/Plugins/Runtime/MassGameplay/Source/MassSimulation/Public/MassSimulationSubsystem.h b/Engine/Plugins/Runtime/MassGameplay/Source/MassSimulation/Public/MassSimulationSubsystem.h
--- a/Engine/Plugins/Runtime/MassGameplay/Source/MassSimulation/Public/MassSimulationSubsystem.h
+++ b/Engine/Plugins/Runtime/MassGameplay/Source/MassSimulation/Public/MassSimulationSubsystem.h
@@ -3,7 +3,7 @@
 #pragma once
 
 #include "Delegates/IDelegateInstance.h"
-#include "MassEntityTypes.h"
+#include "MassProcessingTypes.h"
 #include "MassProcessingPhase.h"
 #include "Subsystems/WorldSubsystem.h"
 #include "MassSimulationSubsystem.generated.h"
diff --git a/Engine/Plugins/Runtime/MassGameplay/Source/MassSmartObjects/Private/MassSmartObjectHandler.cpp b/Engine/Plugins/Runtime/MassGameplay/Source/MassSmartObjects/Private/MassSmartObjectHandler.cpp
--- a/Engine/Plugins/Runtime/MassGameplay/Source/MassSmartObjects/Private/MassSmartObjectHandler.cpp
+++ b/Engine/Plugins/Runtime/MassGameplay/Source/MassSmartObjects/Private/MassSmartObjectHandler.cpp
@@ -5,7 +5,7 @@
 #include "MassSmartObjectBehaviorConfig.h"
 #include "SmartObjectSubsystem.h"
 #include "VisualLogger/VisualLogger.h"
-#include "LWCCommandBuffer.h"
+#include "MassCommandBuffer.h"
 
 namespace UE::Mass::SmartObject
 {
diff --git a/Engine/Plugins/Runtime/MassGameplay/Source/MassSmartObjects/Private/MassSmartObjectProcessor.cpp b/Engine/Plugins/Runtime/MassGameplay/Source/MassSmartObjects/Private/MassSmartObjectProcessor.cpp
--- a/Engine/Plugins/Runtime/MassGameplay/Source/MassSmartObjects/Private/MassSmartObjectProcessor.cpp
+++ b/Engine/Plugins/Runtime/MassGameplay/Source/MassSmartObjects/Private/MassSmartObjectProcessor.cpp
@@ -1,7 +1,7 @@
 // Copyright Epic Games, Inc. All Rights Reserved.
 
 #include "MassSmartObjectProcessor.h"
-#include "LWCCommandBuffer.h"
+#include "MassCommandBuffer.h"
 #include "MassCommonTypes.h"
 #include "MassSignalSubsystem.h"
 #include "MassSmartObjectBehaviorConfig.h"
diff --git a/Engine/Plugins/Runtime/MassGameplay/Source/MassSmartObjects/Public/MassSmartObjectRequest.h b/Engine/Plugins/Runtime/MassGameplay/Source/MassSmartObjects/Public/MassSmartObjectRequest.h
--- a/Engine/Plugins/Runtime/MassGameplay/Source/MassSmartObjects/Public/MassSmartObjectRequest.h
+++ b/Engine/Plugins/Runtime/MassGameplay/Source/MassSmartObjects/Public/MassSmartObjectRequest.h
@@ -3,7 +3,7 @@
 #pragma once
 
 #include "SmartObjectTypes.h"
-#include "LWComponentTypes.h"
+#include "MassEntityTypes.h"
 #include "ZoneGraphTypes.h"
 #include "Containers/StaticArray.h"
 #include "MassSmartObjectRequest.generated.h"
diff --git a/Engine/Plugins/Runtime/MassGameplay/Source/MassSmartObjects/Public/MassSmartObjectTypes.h b/Engine/Plugins/Runtime/MassGameplay/Source/MassSmartObjects/Public/MassSmartObjectTypes.h
--- a/Engine/Plugins/Runtime/MassGameplay/Source/MassSmartObjects/Public/MassSmartObjectTypes.h
+++ b/Engine/Plugins/Runtime/MassGameplay/Source/MassSmartObjects/Public/MassSmartObjectTypes.h
@@ -2,7 +2,7 @@
 
 #pragma once
 
-#include "LWComponentTypes.h"
+#include "MassEntityTypes.h"
 #include "MassSmartObjectTypes.generated.h"
 
 namespace UE::Mass::Signals
diff --git a/Engine/Plugins/Runtime/MassGameplay/Source/MassSpawner/Public/MassEntityTemplate.h b/Engine/Plugins/Runtime/MassGameplay/Source/MassSpawner/Public/MassEntityTemplate.h
--- a/Engine/Plugins/Runtime/MassGameplay/Source/MassSpawner/Public/MassEntityTemplate.h
+++ b/Engine/Plugins/Runtime/MassGameplay/Source/MassSpawner/Public/MassEntityTemplate.h
@@ -3,7 +3,7 @@
 #pragma once
 
 #include "CoreMinimal.h"
-#include "MassEntityTypes.h"
+#include "MassProcessingTypes.h"
 #include "MassEntitySubsystem.h"
 #include "MassCommonTypes.h"
 #include "MassTranslator.h"
diff --git a/Engine/Plugins/Runtime/MassGameplay/Source/MassSpawner/Public/MassTranslator.h b/Engine/Plugins/Runtime/MassGameplay/Source/MassSpawner/Public/MassTranslator.h
--- a/Engine/Plugins/Runtime/MassGameplay/Source/MassSpawner/Public/MassTranslator.h
+++ b/Engine/Plugins/Runtime/MassGameplay/Source/MassSpawner/Public/MassTranslator.h
@@ -4,7 +4,7 @@
 
 #include "MassProcessor.h"
 #include "MassCommonTypes.h"
-#include "LWComponentTypes.h"
+#include "MassEntityTypes.h"
 #include "MassTranslator.generated.h"
 
 
