Legacy Mass namespaces refactoring:

UE::AggregateTicking -> UE::Mass::Core
UE::Pipe -> UE::Mass
UE::MassDebug -> UE::Mass::Debug

diff --git a/Engine/Plugins/AI/MassAI/Source/MassAIBehavior/Private/MassLookAtProcessors.cpp b/Engine/Plugins/AI/MassAI/Source/MassAIBehavior/Private/MassLookAtProcessors.cpp
--- a/Engine/Plugins/AI/MassAI/Source/MassAIBehavior/Private/MassLookAtProcessors.cpp
+++ b/Engine/Plugins/AI/MassAI/Source/MassAIBehavior/Private/MassLookAtProcessors.cpp
@@ -145,7 +145,7 @@ void UMassLookAtProcessor::Execute(UMassEntitySubsystem& EntitySubsystem, FMassE
 				const FMassEntityHandle Entity = Context.GetEntity(i);
 	#if WITH_MASSGAMEPLAY_DEBUG
 				FColor EntityColor = FColor::White;
-				bDisplayDebug = UE::MassDebug::IsDebuggingEntity(Entity, &EntityColor);
+				bDisplayDebug = UE::Mass::Debug::IsDebuggingEntity(Entity, &EntityColor);
 	#endif // WITH_MASSGAMEPLAY_DEBUG
 
 				// Update gaze target when current cycle is finished.
diff --git a/Engine/Plugins/AI/MassAI/Source/MassAIBehavior/Private/Tasks/MassZoneGraphFindEscapeTarget.cpp b/Engine/Plugins/AI/MassAI/Source/MassAIBehavior/Private/Tasks/MassZoneGraphFindEscapeTarget.cpp
--- a/Engine/Plugins/AI/MassAI/Source/MassAIBehavior/Private/Tasks/MassZoneGraphFindEscapeTarget.cpp
+++ b/Engine/Plugins/AI/MassAI/Source/MassAIBehavior/Private/Tasks/MassZoneGraphFindEscapeTarget.cpp
@@ -23,7 +23,7 @@ EStateTreeRunStatus FMassZoneGraphFindEscapeTarget::EnterState(FStateTreeExecuti
 	EscapeTargetLocation = nullptr;
 
 #if WITH_MASSGAMEPLAY_DEBUG
-	bDisplayDebug = UE::MassDebug::IsDebuggingEntity(MassContext.GetEntity());
+	bDisplayDebug = UE::Mass::Debug::IsDebuggingEntity(MassContext.GetEntity());
 #endif // WITH_MASSGAMEPLAY_DEBUG
 
 	const FMassZoneGraphLaneLocationFragment& LaneLocation = Context.GetExternalItem(LocationHandle).Get<FMassZoneGraphLaneLocationFragment>();
diff --git a/Engine/Plugins/AI/MassAI/Source/MassAIBehavior/Private/Tasks/MassZoneGraphPathFollowTask.cpp b/Engine/Plugins/AI/MassAI/Source/MassAIBehavior/Private/Tasks/MassZoneGraphPathFollowTask.cpp
--- a/Engine/Plugins/AI/MassAI/Source/MassAIBehavior/Private/Tasks/MassZoneGraphPathFollowTask.cpp
+++ b/Engine/Plugins/AI/MassAI/Source/MassAIBehavior/Private/Tasks/MassZoneGraphPathFollowTask.cpp
@@ -34,7 +34,7 @@ bool FMassZoneGraphPathFollowTask::RequestPath(FMassStateTreeExecutionContext& C
 
 	bool bDisplayDebug = false;
 #if WITH_MASSGAMEPLAY_DEBUG
-	bDisplayDebug = UE::MassDebug::IsDebuggingEntity(Context.GetEntity());
+	bDisplayDebug = UE::Mass::Debug::IsDebuggingEntity(Context.GetEntity());
 #endif // WITH_MASSGAMEPLAY_DEBUG
 
 	if (RequestedTargetLocation.LaneHandle != LaneLocation.LaneHandle)
@@ -88,7 +88,7 @@ EStateTreeRunStatus FMassZoneGraphPathFollowTask::EnterState(FStateTreeExecution
 
 	bool bDisplayDebug = false;
 #if WITH_MASSGAMEPLAY_DEBUG
-	bDisplayDebug = UE::MassDebug::IsDebuggingEntity(MassContext.GetEntity());
+	bDisplayDebug = UE::Mass::Debug::IsDebuggingEntity(MassContext.GetEntity());
 #endif // WITH_MASSGAMEPLAY_DEBUG
 	if (bDisplayDebug)
 	{
@@ -130,7 +130,7 @@ EStateTreeRunStatus FMassZoneGraphPathFollowTask::Tick(FStateTreeExecutionContex
 
 	bool bDisplayDebug = false;
 #if WITH_MASSGAMEPLAY_DEBUG
-	bDisplayDebug = UE::MassDebug::IsDebuggingEntity(MassContext.GetEntity());
+	bDisplayDebug = UE::Mass::Debug::IsDebuggingEntity(MassContext.GetEntity());
 #endif // WITH_MASSGAMEPLAY_DEBUG
 	if (bDisplayDebug)
 	{
diff --git a/Engine/Plugins/AI/MassAI/Source/MassAIDebug/Private/MassDebugStateTreeProcessor.cpp b/Engine/Plugins/AI/MassAI/Source/MassAIDebug/Private/MassDebugStateTreeProcessor.cpp
--- a/Engine/Plugins/AI/MassAI/Source/MassAIDebug/Private/MassDebugStateTreeProcessor.cpp
+++ b/Engine/Plugins/AI/MassAI/Source/MassAIDebug/Private/MassDebugStateTreeProcessor.cpp
@@ -44,7 +44,7 @@ void UMassDebugStateTreeProcessor::Execute(UMassEntitySubsystem& EntitySubsystem
 		return;
 	}
 
-	if (!Debugger->GetSelectedEntity().IsSet() && !UE::MassDebug::HasDebugEntities())
+	if (!Debugger->GetSelectedEntity().IsSet() && !UE::Mass::Debug::HasDebugEntities())
 	{
 		return;
 	}
@@ -80,7 +80,7 @@ void UMassDebugStateTreeProcessor::Execute(UMassEntitySubsystem& EntitySubsystem
 				}
 					
 				FColor EntityColor = FColor::White;
-				const bool bDisplayDebug = UE::MassDebug::IsDebuggingEntity(Entity, &EntityColor);
+				const bool bDisplayDebug = UE::Mass::Debug::IsDebuggingEntity(Entity, &EntityColor);
 				if (bDisplayDebug)
 				{
 					const FDataFragment_Transform& Transform = TransformList[i];
diff --git a/Engine/Plugins/AI/MassAI/Source/MassAIMovement/Private/MassAvoidanceProcessors.cpp b/Engine/Plugins/AI/MassAI/Source/MassAIMovement/Private/MassAvoidanceProcessors.cpp
--- a/Engine/Plugins/AI/MassAI/Source/MassAIMovement/Private/MassAvoidanceProcessors.cpp
+++ b/Engine/Plugins/AI/MassAI/Source/MassAIMovement/Private/MassAvoidanceProcessors.cpp
@@ -228,7 +228,7 @@ namespace UE::MassAvoidance
 	static bool DebugIsSelected(const FMassEntityHandle Entity)
 	{
 		FColor Color;
-		return UE::MassDebug::IsDebuggingEntity(Entity, &Color);
+		return UE::Mass::Debug::IsDebuggingEntity(Entity, &Color);
 	}
 
 	static void DebugDrawLine(const FDebugContext& Context, const FVector& Start, const FVector& End, const FColor& Color, const float Thickness = 0.f, const bool bPersistent = false)
diff --git a/Engine/Plugins/AI/MassAI/Source/MassAIMovement/Private/MassZoneGraphMovementProcessors.cpp b/Engine/Plugins/AI/MassAI/Source/MassAIMovement/Private/MassZoneGraphMovementProcessors.cpp
--- a/Engine/Plugins/AI/MassAI/Source/MassAIMovement/Private/MassZoneGraphMovementProcessors.cpp
+++ b/Engine/Plugins/AI/MassAI/Source/MassAIMovement/Private/MassZoneGraphMovementProcessors.cpp
@@ -224,7 +224,7 @@ void UMassZoneGraphPathFollowProcessor::Execute(UMassEntitySubsystem& EntitySubs
 
 			bool bDisplayDebug = false;
 #if WITH_MASSGAMEPLAY_DEBUG && UNSAFE_FOR_MT // this will result in bDisplayDebug == false and disabling of all the vlogs below
-			bDisplayDebug = UE::MassDebug::IsDebuggingEntity(Entity);
+			bDisplayDebug = UE::Mass::Debug::IsDebuggingEntity(Entity);
 			if (bDisplayDebug)
 			{
 				UE_VLOG(this, LogMassNavigation, Log, TEXT("Entity [%s] Updating path following"), *Entity.DebugGetDescription());
@@ -411,7 +411,7 @@ void UMassZoneGraphPathFollowProcessor::Execute(UMassEntitySubsystem& EntitySubs
 #if WITH_MASSGAMEPLAY_DEBUG && UNSAFE_FOR_MT
 				if (bDisplayDebug)
 				{
-					const FColor EntityColor = UE::MassDebug::GetEntityDebugColor(Entity);
+					const FColor EntityColor = UE::Mass::Debug::GetEntityDebugColor(Entity);
 
 					const FVector ZOffset(0,0,25);
 					const FColor LightEntityColor = UE::MassMovement::Debug::MixColors(EntityColor, FColor::White);
@@ -709,7 +709,7 @@ void UMassZoneGraphSteeringProcessor::Execute(UMassEntitySubsystem& EntitySubsys
 
 #if WITH_MASSGAMEPLAY_DEBUG && UNSAFE_FOR_MT
 			FColor EntityColor = FColor::White;
-			const bool bDisplayDebug = UE::MassDebug::IsDebuggingEntity(Entity, &EntityColor);
+			const bool bDisplayDebug = UE::Mass::Debug::IsDebuggingEntity(Entity, &EntityColor);
 			if (bDisplayDebug)
 			{
 				const FVector ZOffset(0,0,25);
diff --git a/Engine/Plugins/AI/MassAI/Source/MassAIMovement/Private/MassZoneGraphMovementUtils.cpp b/Engine/Plugins/AI/MassAI/Source/MassAIMovement/Private/MassZoneGraphMovementUtils.cpp
--- a/Engine/Plugins/AI/MassAI/Source/MassAIMovement/Private/MassZoneGraphMovementUtils.cpp
+++ b/Engine/Plugins/AI/MassAI/Source/MassAIMovement/Private/MassZoneGraphMovementUtils.cpp
@@ -51,7 +51,7 @@ namespace UE::MassMovement
 			MoveTarget.IntentAtGoal = ShortPath.EndOfPathIntent;
 			MoveTarget.DistanceToGoal = (ShortPath.NumPoints > 0) ? ShortPath.Points[ShortPath.NumPoints - 1].DistanceAlongLane.Get() : 0.0f;
 #if WITH_MASSGAMEPLAY_DEBUG
-			UE_CVLOG(UE::MassDebug::IsDebuggingEntity(Entity),
+			UE_CVLOG(UE::Mass::Debug::IsDebuggingEntity(Entity),
 				Requester,
 				LogMassNavigation,
 				Log,
diff --git a/Engine/Plugins/AI/MassAI/Source/MassAIMovement/Private/MovementProcessor.cpp b/Engine/Plugins/AI/MassAI/Source/MassAIMovement/Private/MovementProcessor.cpp
--- a/Engine/Plugins/AI/MassAI/Source/MassAIMovement/Private/MovementProcessor.cpp
+++ b/Engine/Plugins/AI/MassAI/Source/MassAIMovement/Private/MovementProcessor.cpp
@@ -61,7 +61,7 @@ namespace UE::MassMovement
  	static bool DebugIsSelected(const FMassEntityHandle Entity)
  	{
  #if WITH_MASSGAMEPLAY_DEBUG
- 		return UE::MassDebug::IsDebuggingEntity(Entity);
+ 		return UE::Mass::Debug::IsDebuggingEntity(Entity);
  #else
  	return false;
  #endif // WITH_MASSGAMEPLAY_DEBUG
diff --git a/Engine/Plugins/AI/MassCrowd/Source/MassCrowd/Private/Tasks/MassZoneGraphFindWanderTarget.cpp b/Engine/Plugins/AI/MassCrowd/Source/MassCrowd/Private/Tasks/MassZoneGraphFindWanderTarget.cpp
--- a/Engine/Plugins/AI/MassCrowd/Source/MassCrowd/Private/Tasks/MassZoneGraphFindWanderTarget.cpp
+++ b/Engine/Plugins/AI/MassCrowd/Source/MassCrowd/Private/Tasks/MassZoneGraphFindWanderTarget.cpp
@@ -28,7 +28,7 @@ EStateTreeRunStatus FMassZoneGraphFindWanderTarget::EnterState(FStateTreeExecuti
 
 	bool bDisplayDebug = false;
 #if WITH_MASSGAMEPLAY_DEBUG
-	bDisplayDebug = UE::MassDebug::IsDebuggingEntity(MassContext.GetEntity());
+	bDisplayDebug = UE::Mass::Debug::IsDebuggingEntity(MassContext.GetEntity());
 #endif // WITH_MASSGAMEPLAY_DEBUG
 
 	if (!LaneLocation.LaneHandle.IsValid())
diff --git a/Engine/Plugins/Runtime/MassEntity/Source/MassEntity/Private/MassArchetypeData.cpp b/Engine/Plugins/Runtime/MassEntity/Source/MassEntity/Private/MassArchetypeData.cpp
--- a/Engine/Plugins/Runtime/MassEntity/Source/MassEntity/Private/MassArchetypeData.cpp
+++ b/Engine/Plugins/Runtime/MassEntity/Source/MassEntity/Private/MassArchetypeData.cpp
@@ -7,10 +7,10 @@
 
 //////////////////////////////////////////////////////////////////////
 // FMassArchetypeData
-namespace UE { namespace AggregateTicking 
+namespace UE::Mass::Core 
 {
 	constexpr static bool bBitwiseRelocateComponents = true;
-}} // UE::AggregateTicking
+}// UE::Mass::Core
 
 void FMassArchetypeData::ForEachComponentType(TFunction< void(const UScriptStruct* /*Component*/)> Function) const
 {
@@ -218,7 +218,7 @@ void FMassArchetypeData::RemoveEntityInternal(const int32 AbsoluteIndex, const b
 	Chunk.RemoveInstance();
 	const int32 IndexToSwapFrom = Chunk.GetNumInstances();
 
-	checkf(bDestroyComponents || UE::AggregateTicking::bBitwiseRelocateComponents, TEXT("We allow not to destroy components only in bit wise relocation mode."));
+	checkf(bDestroyComponents || UE::Mass::Core::bBitwiseRelocateComponents, TEXT("We allow not to destroy components only in bit wise relocation mode."));
 
 	// Remove and swap the last entry in the chunk to the location of the removed item (if it's not the same as the dying entry)
 	if (IndexToSwapFrom != IndexWithinChunk)
@@ -228,7 +228,7 @@ void FMassArchetypeData::RemoveEntityInternal(const int32 AbsoluteIndex, const b
 			void* DyingComponentPtr = ComponentConfig.GetComponentData(Chunk.GetRawMemory(), IndexWithinChunk);
 			void* MovingComponentPtr = ComponentConfig.GetComponentData(Chunk.GetRawMemory(), IndexToSwapFrom);
 
-			if (UE::AggregateTicking::bBitwiseRelocateComponents)
+			if (UE::Mass::Core::bBitwiseRelocateComponents)
 			{
 				// Destroy component data
 				if (bDestroyComponents)
@@ -257,7 +257,7 @@ void FMassArchetypeData::RemoveEntityInternal(const int32 AbsoluteIndex, const b
 		Chunk.GetEntityArrayElementRef(EntityListOffsetWithinChunk, IndexWithinChunk) = EntityBeingSwapped;
 		EntityMap.FindChecked(EntityBeingSwapped.Index) = AbsoluteIndex;
 	}
-	else if (!UE::AggregateTicking::bBitwiseRelocateComponents || bDestroyComponents)
+	else if (!UE::Mass::Core::bBitwiseRelocateComponents || bDestroyComponents)
 	{
 		for (const FMassArchetypeFragmentConfig& ComponentConfig : ComponentConfigs)
 		{
@@ -309,7 +309,7 @@ void FMassArchetypeData::BatchDestroyEntityChunks(const FArchetypeChunkCollectio
 				void* DyingComponentPtr = ComponentConfig.GetComponentData(Chunk.GetRawMemory(), SubchunkInfo.SubchunkStart);
 				void* MovingComponentPtr = ComponentConfig.GetComponentData(Chunk.GetRawMemory(), SwapStartIndex);
 
-				if (UE::AggregateTicking::bBitwiseRelocateComponents)
+				if (UE::Mass::Core::bBitwiseRelocateComponents)
 				{
 					// Destroy the components we'll replace by the following copy
 					ComponentConfig.ComponentType->DestroyStruct(DyingComponentPtr, NumberToMove);
@@ -405,7 +405,7 @@ void FMassArchetypeData::SetComponentsData(const FMassEntityHandle Entity, TArra
 		check(ComponentType);
 		const int32 ComponentIndex = ComponentIndexMap.FindChecked(ComponentType);
 		void* ComponentMemory = GetComponentData(ComponentIndex, InternalIndex);
-		// No UE::AggregateTicking::bBitwiseRelocateComponents, this isn't a move component
+		// No UE::Mass::Core::bBitwiseRelocateComponents, this isn't a move component
 		ComponentType->CopyScriptStruct(ComponentMemory, Instance.GetMemory());
 	}
 }
@@ -425,7 +425,7 @@ void FMassArchetypeData::SetComponentData(const FArchetypeChunkCollection& Chunk
 		uint8* ComponentMemory = (uint8*)ComponentConfigs[ComponentIndex].GetComponentData(Chunks[ChunkIterator->ChunkIndex].GetRawMemory(), ChunkIterator->SubchunkStart);
 		for (int i = ChunkIterator->Length; i; --i, ComponentMemory += ComponentTypeSize)
 		{
-			// No UE::AggregateTicking::bBitwiseRelocateComponents, this isn't a move of a component
+			// No UE::Mass::Core::bBitwiseRelocateComponents, this isn't a move of a component
 			ComponentType->CopyScriptStruct(ComponentMemory, ComponentSourceMemory);
 		}
 	}
@@ -451,7 +451,7 @@ void FMassArchetypeData::MoveEntityToAnotherArchetype(const FMassEntityHandle En
 	const int32 IndexWithinChunk = AbsoluteIndex % NumEntitiesPerChunk;
 	FMassArchetypeChunk& Chunk = Chunks[ChunkIndex];
 
-	constexpr bool bInitializeComponentsDuringCreation = !UE::AggregateTicking::bBitwiseRelocateComponents;
+	constexpr bool bInitializeComponentsDuringCreation = !UE::Mass::Core::bBitwiseRelocateComponents;
 	const int32 NewAbsoluteIndex = NewArchetype.AddEntityInternal(Entity, bInitializeComponentsDuringCreation);
 	const int32 NewChunkIndex = NewAbsoluteIndex / NewArchetype.NumEntitiesPerChunk;
 	const int32 NewIndexWithinChunk = NewAbsoluteIndex % NewArchetype.NumEntitiesPerChunk;
@@ -469,7 +469,7 @@ void FMassArchetypeData::MoveEntityToAnotherArchetype(const FMassEntityHandle En
 		if (OldComponentIndex)
 		{
 			const void* Src = ComponentConfigs[*OldComponentIndex].GetComponentData(Chunk.GetRawMemory(), IndexWithinChunk);
-			if (UE::AggregateTicking::bBitwiseRelocateComponents)
+			if (UE::Mass::Core::bBitwiseRelocateComponents)
 			{
 				FMemory::Memcpy(Dst, Src, NewComponentConfig.ComponentType->GetStructureSize());
 			}
@@ -487,7 +487,7 @@ void FMassArchetypeData::MoveEntityToAnotherArchetype(const FMassEntityHandle En
 		}
 	}
 
-	constexpr bool bDestroyComponents = !UE::AggregateTicking::bBitwiseRelocateComponents;
+	constexpr bool bDestroyComponents = !UE::Mass::Core::bBitwiseRelocateComponents;
 	RemoveEntityInternal(AbsoluteIndex, bDestroyComponents);
 }
 
@@ -609,7 +609,7 @@ void FMassArchetypeData::CompactEntities(const double TimeAllowed)
 			void* FromComponentPtr = ComponentConfig.GetComponentData(ChunkToEmpty->GetRawMemory(), FromIndex);
 			void* ToComponentPtr = ComponentConfig.GetComponentData(ChunkToFill->GetRawMemory(), ToIndex);
 
-			if (UE::AggregateTicking::bBitwiseRelocateComponents)
+			if (UE::Mass::Core::bBitwiseRelocateComponents)
 			{
 				// Move all entries
 				FMemory::Memcpy(ToComponentPtr, FromComponentPtr, ComponentConfig.ComponentType->GetStructureSize() * NumberOfEntitiesToMove);
diff --git a/Engine/Plugins/Runtime/MassEntity/Source/MassEntity/Private/MassEntityDebug.cpp b/Engine/Plugins/Runtime/MassEntity/Source/MassEntity/Private/MassEntityDebug.cpp
--- a/Engine/Plugins/Runtime/MassEntity/Source/MassEntity/Private/MassEntityDebug.cpp
+++ b/Engine/Plugins/Runtime/MassEntity/Source/MassEntity/Private/MassEntityDebug.cpp
@@ -8,9 +8,25 @@
 
 DEFINE_ENUM_TO_STRING(EMassProcessingPhase);
 
-namespace UE::Pipe::Debug
+namespace UE::Mass::Debug
 {
 
+FString DebugGetComponentAccessString(EMassFragmentAccess Access)
+{
+#if WITH_MASSENTITY_DEBUG
+	switch (Access)
+	{
+	case EMassFragmentAccess::None:	return TEXT("--");
+	case EMassFragmentAccess::ReadOnly:	return TEXT("RO");
+	case EMassFragmentAccess::ReadWrite:	return TEXT("RW");
+	default:
+		ensureMsgf(false, TEXT("Missing string conversion for EMassFragmentAccess=%d"), Access);
+		break;
+	}
+#endif // WITH_MASSENTITY_DEBUG
+	return TEXT("Missing string conversion");
+}
+	
 void DebugOutputDescription(TConstArrayView<UMassProcessor*> Processors, FOutputDevice& Ar)
 {
 #if WITH_MASSENTITY_DEBUG
@@ -129,5 +145,5 @@ FAutoConsoleCommandWithWorldArgsAndOutputDevice LogFragmentSizes(
 	);
 #endif // WITH_MASSENTITY_DEBUG && WITH_MASSENTITY_DEBUG
 
-} // namespace UE::Pipe::Debug
+} // namespace UE::Mass::Debug
 
diff --git a/Engine/Plugins/Runtime/MassEntity/Source/MassEntity/Private/MassEntityQuery.cpp b/Engine/Plugins/Runtime/MassEntity/Source/MassEntity/Private/MassEntityQuery.cpp
--- a/Engine/Plugins/Runtime/MassEntity/Source/MassEntity/Private/MassEntityQuery.cpp
+++ b/Engine/Plugins/Runtime/MassEntity/Source/MassEntity/Private/MassEntityQuery.cpp
@@ -1,6 +1,7 @@
 // Copyright Epic Games, Inc. All Rights Reserved.
 
 #include "MassEntityQuery.h"
+#include "MassEntityDebug.h"
 #include "MassEntitySubsystem.h"
 #include "MassArchetypeData.h"
 #include "MassCommandBuffer.h"
@@ -301,7 +302,7 @@ FString FMassFragmentRequirement::DebugGetDescription() const
 {
 #if WITH_MASSENTITY_DEBUG
 	return FString::Printf(TEXT("%s%s[%s]"), IsOptional() ? TEXT("?") : (Presence == EMassFragmentPresence::None ? TEXT("-") : TEXT("+"))
-		, *GetNameSafe(StructType), *UE::AggregateTicking::DebugGetComponentAccessString(AccessMode));
+		, *GetNameSafe(StructType), *UE::Mass::Debug::DebugGetComponentAccessString(AccessMode));
 #else
 	return {};
 #endif
diff --git a/Engine/Plugins/Runtime/MassEntity/Source/MassEntity/Private/MassEntitySubsystem.cpp b/Engine/Plugins/Runtime/MassEntity/Source/MassEntity/Private/MassEntitySubsystem.cpp
--- a/Engine/Plugins/Runtime/MassEntity/Source/MassEntity/Private/MassEntitySubsystem.cpp
+++ b/Engine/Plugins/Runtime/MassEntity/Source/MassEntity/Private/MassEntitySubsystem.cpp
@@ -11,26 +11,6 @@ const FMassEntityHandle UMassEntitySubsystem::InvalidEntity;
 
 DEFINE_LOG_CATEGORY(LogAggregateTicking);
 
-namespace UE { namespace AggregateTicking {
-
-FString DebugGetComponentAccessString(EMassFragmentAccess Access)
-{
-#if WITH_MASSENTITY_DEBUG
-	switch (Access)
-	{
-	case EMassFragmentAccess::None:	return TEXT("--");
-	case EMassFragmentAccess::ReadOnly:	return TEXT("RO");
-	case EMassFragmentAccess::ReadWrite:	return TEXT("RW");
-	default:
-		ensureMsgf(false, TEXT("Missing string conversion for EMassFragmentAccess=%d"), Access);
-		break;
-	}
-#endif // WITH_MASSENTITY_DEBUG
-	return TEXT("Missing string conversion");
-}
-
-}} // UE::AggregateTicking
-
 //@TODO: Everything still alive leaks at shutdown
 //@TODO: No re-entrance safety while running a system (e.g., preventing someone from adding/removing entities or altering archetypes, etc...)
 //@TODO: Do we allow GCable types?  If so, need to implement AddReferencedObjects
diff --git a/Engine/Plugins/Runtime/MassEntity/Source/MassEntity/Private/MassEntityUtils.cpp b/Engine/Plugins/Runtime/MassEntity/Source/MassEntity/Private/MassEntityUtils.cpp
--- a/Engine/Plugins/Runtime/MassEntity/Source/MassEntity/Private/MassEntityUtils.cpp
+++ b/Engine/Plugins/Runtime/MassEntity/Source/MassEntity/Private/MassEntityUtils.cpp
@@ -4,7 +4,7 @@
 #include "Engine/EngineBaseTypes.h"
 #include "Engine/World.h"
 
-namespace UE { namespace Pipe { namespace Utils
+namespace UE::Mass::Utils
 {
 
 EProcessorExecutionFlags GetProcessorExecutionFlagsForWold(const UWorld& World)
@@ -31,4 +31,4 @@ EProcessorExecutionFlags GetProcessorExecutionFlagsForWold(const UWorld& World)
 	return ExecutionFlags;
 }
 
-} } }// namespace UE::Pipe::Utils
\ No newline at end of file
+} // namespace UE::Mass::Utils
\ No newline at end of file
diff --git a/Engine/Plugins/Runtime/MassEntity/Source/MassEntity/Private/MassExecutor.cpp b/Engine/Plugins/Runtime/MassEntity/Source/MassEntity/Private/MassExecutor.cpp
--- a/Engine/Plugins/Runtime/MassEntity/Source/MassEntity/Private/MassExecutor.cpp
+++ b/Engine/Plugins/Runtime/MassEntity/Source/MassEntity/Private/MassExecutor.cpp
@@ -6,7 +6,7 @@
 #include "MassCommandBuffer.h"
 #include "ProfilingDebugging/CpuProfilerTrace.h"
 
-namespace UE::Pipe::Executor
+namespace UE::Mass::Executor
 {
 void Run(FMassRuntimePipeline& RuntimePipeline, FMassProcessingContext& PipeContext)
 {
@@ -26,7 +26,7 @@ void RunSparse(FMassRuntimePipeline& RuntimePipeline, FMassProcessingContext& Pi
 	if (!ensure(PipeContext.EntitySubsystem) ||
 		!ensure(RuntimePipeline.Processors.Find(nullptr) == INDEX_NONE) ||
 		RuntimePipeline.Processors.Num() == 0 ||
-		!ensureMsgf(Archetype.IsValid(), TEXT("The Archetype passed in to UE::Pipe::Executor::RunSparse is invalid")))
+		!ensureMsgf(Archetype.IsValid(), TEXT("The Archetype passed in to UE::Mass::Executor::RunSparse is invalid")))
 	{
 		return;
 	}
@@ -42,7 +42,7 @@ void RunSparse(FMassRuntimePipeline& RuntimePipeline, FMassProcessingContext& Pi
 	if (!ensure(PipeContext.EntitySubsystem) ||
 		!ensure(RuntimePipeline.Processors.Find(nullptr) == INDEX_NONE) ||
 		RuntimePipeline.Processors.Num() == 0 ||
-		!ensureMsgf(ChunkCollection.GetArchetype().IsValid(), TEXT("The Archetype of ChunkCollection passed in to UE::Pipe::Executor::RunSparse is invalid")))
+		!ensureMsgf(ChunkCollection.GetArchetype().IsValid(), TEXT("The Archetype of ChunkCollection passed in to UE::Mass::Executor::RunSparse is invalid")))
 	{
 		return;
 	}
@@ -187,4 +187,4 @@ FGraphEventRef TriggerParallelTasks(UMassProcessor& Processor, FMassProcessingCo
 	return CompletionEvent;
 }
 
-} // namespace UE::Pipe::Executor
+} // namespace UE::Mass::Executor
diff --git a/Engine/Plugins/Runtime/MassEntity/Source/MassEntity/Private/MassProcessingPhase.cpp b/Engine/Plugins/Runtime/MassEntity/Source/MassEntity/Private/MassProcessingPhase.cpp
--- a/Engine/Plugins/Runtime/MassEntity/Source/MassEntity/Private/MassProcessingPhase.cpp
+++ b/Engine/Plugins/Runtime/MassEntity/Source/MassEntity/Private/MassProcessingPhase.cpp
@@ -53,7 +53,7 @@ void FMassProcessingPhase::ExecuteTick(float DeltaTime, ELevelTick TickType, ENa
 	if (bRunInParallelMode)
 	{
 		bIsDuringPipeProcessing = true;
-		const FGraphEventRef PipelineCompletionEvent = UE::Pipe::Executor::TriggerParallelTasks(*PhaseProcessor, Context, [this, DeltaTime]()
+		const FGraphEventRef PipelineCompletionEvent = UE::Mass::Executor::TriggerParallelTasks(*PhaseProcessor, Context, [this, DeltaTime]()
 			{
 				OnParallelExecutionDone(DeltaTime);
 			});
@@ -70,7 +70,7 @@ void FMassProcessingPhase::ExecuteTick(float DeltaTime, ELevelTick TickType, ENa
 	else
 	{
 		TGuardValue<bool> PipeRunningGuard(bIsDuringPipeProcessing, true);
-		UE::Pipe::Executor::Run(*PhaseProcessor, Context);
+		UE::Mass::Executor::Run(*PhaseProcessor, Context);
 
 		OnPhaseEnd.Broadcast(DeltaTime);
 		Manager->OnPhaseEnd(*this);
diff --git a/Engine/Plugins/Runtime/MassEntity/Source/MassEntity/Private/MassProcessingTypes.cpp b/Engine/Plugins/Runtime/MassEntity/Source/MassEntity/Private/MassProcessingTypes.cpp
--- a/Engine/Plugins/Runtime/MassEntity/Source/MassEntity/Private/MassProcessingTypes.cpp
+++ b/Engine/Plugins/Runtime/MassEntity/Source/MassEntity/Private/MassProcessingTypes.cpp
@@ -82,7 +82,7 @@ void FMassRuntimePipeline::InitializeFromClassArray(TConstArrayView<TSubclassOf<
 	Reset();
 
 	const UWorld* World = InOwner.GetWorld();
-	const EProcessorExecutionFlags WorldExecutionFlags = World ? UE::Pipe::Utils::GetProcessorExecutionFlagsForWold(*World) : EProcessorExecutionFlags::All;
+	const EProcessorExecutionFlags WorldExecutionFlags = World ? UE::Mass::Utils::GetProcessorExecutionFlagsForWold(*World) : EProcessorExecutionFlags::All;
 
 	for (const TSubclassOf<UMassProcessor>& ProcessorClass : InProcessorClasses)
 	{
@@ -113,7 +113,7 @@ bool FMassRuntimePipeline::HasProcessorOfExactClass(TSubclassOf<UMassProcessor>
 void FMassRuntimePipeline::AppendUniqueRuntimeProcessorCopies(TConstArrayView<const UMassProcessor*> InProcessors, UObject& InOwner)
 {
 	const UWorld* World = InOwner.GetWorld();
-	const EProcessorExecutionFlags WorldExecutionFlags = World ? UE::Pipe::Utils::GetProcessorExecutionFlagsForWold(*World) : EProcessorExecutionFlags::All;
+	const EProcessorExecutionFlags WorldExecutionFlags = World ? UE::Mass::Utils::GetProcessorExecutionFlagsForWold(*World) : EProcessorExecutionFlags::All;
 	const int32 StartingCount = Processors.Num();
 		
 	for (const UMassProcessor* Proc : InProcessors)
@@ -152,7 +152,7 @@ void FMassRuntimePipeline::AppendUniqueRuntimeProcessorCopies(TConstArrayView<co
 void FMassRuntimePipeline::AppendOrOverrideRuntimeProcessorCopies(TConstArrayView<const UMassProcessor*> InProcessors, UObject& InOwner)
 {
 	const UWorld* World = InOwner.GetWorld();
-	const EProcessorExecutionFlags WorldExecutionFlags = World ? UE::Pipe::Utils::GetProcessorExecutionFlagsForWold(*World) : EProcessorExecutionFlags::All;
+	const EProcessorExecutionFlags WorldExecutionFlags = World ? UE::Mass::Utils::GetProcessorExecutionFlagsForWold(*World) : EProcessorExecutionFlags::All;
 	
 	for (const UMassProcessor* Proc : InProcessors)
 	{
@@ -211,7 +211,7 @@ UMassCompositeProcessor* FMassRuntimePipeline::FindTopLevelGroupByName(FName Gro
 
 void FMassRuntimePipeline::DebugOutputDescription(FOutputDevice& Ar) const
 {
-	UE::Pipe::Debug::DebugOutputDescription(Processors, Ar);
+	UE::Mass::Debug::DebugOutputDescription(Processors, Ar);
 }
 
 uint32 GetTypeHash(const FMassRuntimePipeline& Instance)
diff --git a/Engine/Plugins/Runtime/MassEntity/Source/MassEntity/Public/MassEntityDebug.h b/Engine/Plugins/Runtime/MassEntity/Source/MassEntity/Public/MassEntityDebug.h
--- a/Engine/Plugins/Runtime/MassEntity/Source/MassEntity/Public/MassEntityDebug.h
+++ b/Engine/Plugins/Runtime/MassEntity/Source/MassEntity/Public/MassEntityDebug.h
@@ -1,17 +1,19 @@
 // Copyright Epic Games, Inc. All Rights Reserved.
 
 #pragma once
+
 #include "VisualLogger/VisualLoggerTypes.h"
 #include "MassProcessingTypes.h"
+#include "MassEntityQuery.h"
 
 class FOutputDevice;
 class UMassProcessor;
 
 MASSENTITY_API DECLARE_ENUM_TO_STRING(EMassProcessingPhase);
 
-namespace UE { namespace Pipe { namespace Debug {
-
-MASSENTITY_API extern void DebugOutputDescription(TConstArrayView<UMassProcessor*> Processors, FOutputDevice& Ar);
-
-}}} // namespace UE::Pipe::Debug
+namespace UE::Mass::Debug
+{
+	FString DebugGetComponentAccessString(EMassFragmentAccess Access);
+	MASSENTITY_API extern void DebugOutputDescription(TConstArrayView<UMassProcessor*> Processors, FOutputDevice& Ar);
+} // namespace UE::Mass::Debug
 
diff --git a/Engine/Plugins/Runtime/MassEntity/Source/MassEntity/Public/MassEntitySubsystem.h b/Engine/Plugins/Runtime/MassEntity/Source/MassEntity/Public/MassEntitySubsystem.h
--- a/Engine/Plugins/Runtime/MassEntity/Source/MassEntity/Public/MassEntitySubsystem.h
+++ b/Engine/Plugins/Runtime/MassEntity/Source/MassEntity/Public/MassEntitySubsystem.h
@@ -21,10 +21,6 @@ struct FMassArchetypeChunk;
 class FOutputDevice;
 enum class EMassFragmentAccess : uint8;
 
-namespace UE { namespace AggregateTicking {
-FString DebugGetComponentAccessString(EMassFragmentAccess Access);
-}}
-
 //@TODO: Comment this guy
 UCLASS()
 class MASSENTITY_API UMassEntitySubsystem : public UWorldSubsystem
diff --git a/Engine/Plugins/Runtime/MassEntity/Source/MassEntity/Public/MassEntityUtils.h b/Engine/Plugins/Runtime/MassEntity/Source/MassEntity/Public/MassEntityUtils.h
--- a/Engine/Plugins/Runtime/MassEntity/Source/MassEntity/Public/MassEntityUtils.h
+++ b/Engine/Plugins/Runtime/MassEntity/Source/MassEntity/Public/MassEntityUtils.h
@@ -4,10 +4,10 @@
 
 #include "MassProcessingTypes.h"
 
-namespace UE { namespace Pipe { namespace Utils
+namespace UE::Mass::Utils
 {
 
 /** returns the current execution mode for the processors calculated from the world network mode */
 MASSENTITY_API extern EProcessorExecutionFlags GetProcessorExecutionFlagsForWold(const UWorld& World);
 
-} } }// namespace UE::Pipe::Utils
\ No newline at end of file
+} // namespace UE::Mass::Utils
\ No newline at end of file
diff --git a/Engine/Plugins/Runtime/MassEntity/Source/MassEntity/Public/MassExecutor.h b/Engine/Plugins/Runtime/MassEntity/Source/MassEntity/Public/MassExecutor.h
--- a/Engine/Plugins/Runtime/MassEntity/Source/MassEntity/Public/MassExecutor.h
+++ b/Engine/Plugins/Runtime/MassEntity/Source/MassEntity/Public/MassExecutor.h
@@ -14,7 +14,7 @@ struct FArchetypeChunkCollection;
 class UMassProcessor;
 
 
-namespace UE::Pipe::Executor
+namespace UE::Mass::Executor
 {
 	/** Executes processors in a given RuntimePipeline */
 	MASSENTITY_API void Run(FMassRuntimePipeline& RuntimePipeline, FMassProcessingContext& PipeContext);
diff --git a/Engine/Plugins/Runtime/MassEntity/Source/MassEntityTestSuite/Private/MassCompositeProcessorTest.cpp b/Engine/Plugins/Runtime/MassEntity/Source/MassEntityTestSuite/Private/MassCompositeProcessorTest.cpp
--- a/Engine/Plugins/Runtime/MassEntity/Source/MassEntityTestSuite/Private/MassCompositeProcessorTest.cpp
+++ b/Engine/Plugins/Runtime/MassEntity/Source/MassEntityTestSuite/Private/MassCompositeProcessorTest.cpp
@@ -26,7 +26,7 @@ struct FCompositeProcessorTest_Empty : FEntityTestBase
 		check(CompositeProcessor);
 		FMassProcessingContext ProcessingContext(*EntitySubsystem, /*DeltaSeconds=*/0.f);
 		// it should just run, no warnings
-		UE::Pipe::Executor::Run(*CompositeProcessor, ProcessingContext);
+		UE::Mass::Executor::Run(*CompositeProcessor, ProcessingContext);
 		return true;
 	}
 };
@@ -59,7 +59,7 @@ struct FCompositeProcessorTest_MultipleSubProcessors : FEntityTestBase
 		}
 
 		FMassProcessingContext ProcessingContext(*EntitySubsystem, /*DeltaSeconds=*/0.f);
-		UE::Pipe::Executor::Run(*CompositeProcessor, ProcessingContext);
+		UE::Mass::Executor::Run(*CompositeProcessor, ProcessingContext);
 		AITEST_EQUAL("All of the child processors should get run", Result, 111);
 		return true;
 	}
diff --git a/Engine/Plugins/Runtime/MassEntity/Source/MassEntityTestSuite/Private/MassExecutionTest.cpp b/Engine/Plugins/Runtime/MassEntity/Source/MassEntityTestSuite/Private/MassExecutionTest.cpp
--- a/Engine/Plugins/Runtime/MassEntity/Source/MassEntityTestSuite/Private/MassExecutionTest.cpp
+++ b/Engine/Plugins/Runtime/MassEntity/Source/MassEntityTestSuite/Private/MassExecutionTest.cpp
@@ -38,7 +38,7 @@ struct FExecution_EmptyArray : FExecutionTestBase
 		const float DeltaSeconds = 0.f;
 		FMassProcessingContext PipeContext(*EntitySubsystem, DeltaSeconds);
 		// no test performed, let's just see if it results in errors/warnings
-		UE::Pipe::Executor::RunProcessorsView(TArrayView<UMassProcessor*>(), PipeContext);
+		UE::Mass::Executor::RunProcessorsView(TArrayView<UMassProcessor*>(), PipeContext);
 		return true;
 	}
 };
@@ -54,7 +54,7 @@ struct FExecution_EmptyPipeline : FExecutionTestBase
 		FMassProcessingContext PipeContext(*EntitySubsystem, DeltaSeconds);
 		FMassRuntimePipeline Pipeline;
 		// no test performed, let's just see if it results in errors/warnings
-		UE::Pipe::Executor::Run(Pipeline, PipeContext);
+		UE::Mass::Executor::Run(Pipeline, PipeContext);
 		return true;
 	}
 };
@@ -72,9 +72,9 @@ struct FExecution_InvalidPipeContext : FExecutionTestBase
 		AITEST_NULL("FMassProcessingContext\'s default constructor is expected to set FMassProcessingContext.EntitySubsystem to null", PipeContext.EntitySubsystem);
 		
 		GetTestRunner().AddExpectedError(TEXT("PipeContext.EntitySubsystem is null"), EAutomationExpectedErrorFlags::Contains, 1);		
-		// note that using RunProcessorsView is to bypass reasonable tests UE::Pipe::Executor::Run(Pipeline,...) does that are 
+		// note that using RunProcessorsView is to bypass reasonable tests UE::Mass::Executor::Run(Pipeline,...) does that are 
 		// reported via ensures which are not handled by the automation framework
-		UE::Pipe::Executor::RunProcessorsView(TArrayView<UMassProcessor*>(), PipeContext);
+		UE::Mass::Executor::RunProcessorsView(TArrayView<UMassProcessor*>(), PipeContext);
 		return true;
 	}
 };
@@ -94,9 +94,9 @@ struct FExecution_SingleNullProcessor : FExecutionTestBase
 		
 
 		GetTestRunner().AddExpectedError(TEXT("Processors contains nullptr"), EAutomationExpectedErrorFlags::Contains, 1);
-		// note that using RunProcessorsView is to bypass reasonable tests UE::Pipe::Executor::Run(Pipeline,...) does that are 
+		// note that using RunProcessorsView is to bypass reasonable tests UE::Mass::Executor::Run(Pipeline,...) does that are 
 		// reported via ensures which are not handled by the automation framework
-		UE::Pipe::Executor::RunProcessorsView(Processors, PipeContext);
+		UE::Mass::Executor::RunProcessorsView(Processors, PipeContext);
 		return true;
 	}
 };
@@ -114,7 +114,7 @@ struct FExecution_SingleValidProcessor : FExecutionTestBase
 		check(Processor);
 
 		// nothing should break. The actual result of processing is getting tested in MassProcessorTests.cpp
-		UE::Pipe::Executor::Run(*Processor, PipeContext);
+		UE::Mass::Executor::Run(*Processor, PipeContext);
 		return true;
 	}
 };
@@ -134,9 +134,9 @@ struct FExecution_MultipleNullProcessors : FExecutionTestBase
 		Processors.Add(nullptr);
 
 		GetTestRunner().AddExpectedError(TEXT("Processors contains nullptr"), EAutomationExpectedErrorFlags::Contains, 1);
-		// note that using RunProcessorsView is to bypass reasonable tests UE::Pipe::Executor::Run(Pipeline,...) does that are 
+		// note that using RunProcessorsView is to bypass reasonable tests UE::Mass::Executor::Run(Pipeline,...) does that are 
 		// reported via ensures which are not handled by the automation framework
-		UE::Pipe::Executor::RunProcessorsView(Processors, PipeContext);
+		UE::Mass::Executor::RunProcessorsView(Processors, PipeContext);
 		return true;
 	}
 };
@@ -164,7 +164,7 @@ struct FExecution_Sparse : FEntityTestBase
 		FArchetypeChunkCollection ChunkCollection(FloatsArchetype);
 		// nothing should break. The actual result of processing is getting tested in MassProcessorTests.cpp
 		
-		UE::Pipe::Executor::RunSparse(Pipeline, PipeContext, ChunkCollection);
+		UE::Mass::Executor::RunSparse(Pipeline, PipeContext, ChunkCollection);
 
 		return true;
 	}
diff --git a/Engine/Plugins/Runtime/MassEntity/Source/MassEntityTestSuite/Private/MassProcessorTest.cpp b/Engine/Plugins/Runtime/MassEntity/Source/MassEntityTestSuite/Private/MassProcessorTest.cpp
--- a/Engine/Plugins/Runtime/MassEntity/Source/MassEntityTestSuite/Private/MassProcessorTest.cpp
+++ b/Engine/Plugins/Runtime/MassEntity/Source/MassEntityTestSuite/Private/MassProcessorTest.cpp
@@ -30,7 +30,7 @@ int32 SimpleProcessorRun(UMassEntitySubsystem& EntitySubsystem)
 	};
 
 	FMassProcessingContext PipeContext(EntitySubsystem, /*DeltaSeconds=*/0.f);
-	UE::Pipe::Executor::Run(*Processor, PipeContext);
+	UE::Mass::Executor::Run(*Processor, PipeContext);
 
 	return EntityProcessedCount;
 }
diff --git a/Engine/Plugins/Runtime/MassEntity/Source/MassEntityTestSuite/Private/MassQueryTest.cpp b/Engine/Plugins/Runtime/MassEntity/Source/MassEntityTestSuite/Private/MassQueryTest.cpp
--- a/Engine/Plugins/Runtime/MassEntity/Source/MassEntityTestSuite/Private/MassQueryTest.cpp
+++ b/Engine/Plugins/Runtime/MassEntity/Source/MassEntityTestSuite/Private/MassQueryTest.cpp
@@ -92,7 +92,7 @@ struct FQueryTest_FragmentViewBinding : FEntityTestBase
 		};
 
 		FMassProcessingContext PipeContext(*EntitySubsystem, /*DeltaSeconds=*/0.f);
-		UE::Pipe::Executor::Run(*Processor, PipeContext);
+		UE::Mass::Executor::Run(*Processor, PipeContext);
 
 		AITEST_EQUAL("Fragment value should have changed to the expected value", TestedFragment.Value, 13.f);
 
diff --git a/Engine/Plugins/Runtime/MassGameplay/Source/MassCommon/Private/MassCommonTypes.cpp b/Engine/Plugins/Runtime/MassGameplay/Source/MassCommon/Private/MassCommonTypes.cpp
--- a/Engine/Plugins/Runtime/MassGameplay/Source/MassCommon/Private/MassCommonTypes.cpp
+++ b/Engine/Plugins/Runtime/MassGameplay/Source/MassCommon/Private/MassCommonTypes.cpp
@@ -5,7 +5,9 @@
 #include "Math/ColorList.h"
 
 #if WITH_MASSGAMEPLAY_DEBUG
-namespace UE { namespace MassDebug {
+
+namespace UE::Mass::Debug
+{
 
 // First Id of a range of lightweight entity for which we want to activate debug information
 int32 DebugEntityBegin = INDEX_NONE;
@@ -100,7 +102,7 @@ FColor GetEntityDebugColor(FMassEntityHandle Entity)
 	return EntityIdx != INDEX_NONE ? GColorList.GetFColorByIndex(EntityIdx % GColorList.GetColorsNum()) : FColor::Black;
 }
 
-}} // namespace UE::MassDebug
+} // namespace UE::Mass::Debug
 #endif // WITH_MASSGAMEPLAY_DEBUG
 
 //----------------------------------------------------------------------//
diff --git a/Engine/Plugins/Runtime/MassGameplay/Source/MassCommon/Public/MassCommonTypes.h b/Engine/Plugins/Runtime/MassGameplay/Source/MassCommon/Public/MassCommonTypes.h
--- a/Engine/Plugins/Runtime/MassGameplay/Source/MassCommon/Public/MassCommonTypes.h
+++ b/Engine/Plugins/Runtime/MassGameplay/Source/MassCommon/Public/MassCommonTypes.h
@@ -16,12 +16,12 @@ class UMaterialInterface;
 MASSCOMMON_API DECLARE_LOG_CATEGORY_EXTERN(LogMass, Warning, All);
 
 #if WITH_MASSGAMEPLAY_DEBUG
-namespace UE::MassDebug
+namespace UE::Mass::Debug
 {
 	MASSCOMMON_API extern bool HasDebugEntities();
 	MASSCOMMON_API extern bool IsDebuggingEntity(FMassEntityHandle Entity, FColor* OutEntityColor = nullptr);
 	MASSCOMMON_API extern FColor GetEntityDebugColor(FMassEntityHandle Entity);
-} // namespace UE::MassDebug
+} // namespace UE::Mass::Debug
 #endif // WITH_MASSGAMEPLAY_DEBUG
 
 namespace UE::Mass::ProcessorGroupNames
diff --git a/Engine/Plugins/Runtime/MassGameplay/Source/MassSmartObjects/Private/MassSmartObjectHandler.cpp b/Engine/Plugins/Runtime/MassGameplay/Source/MassSmartObjects/Private/MassSmartObjectHandler.cpp
--- a/Engine/Plugins/Runtime/MassGameplay/Source/MassSmartObjects/Private/MassSmartObjectHandler.cpp
+++ b/Engine/Plugins/Runtime/MassGameplay/Source/MassSmartObjects/Private/MassSmartObjectHandler.cpp
@@ -128,7 +128,7 @@ EMassSmartObjectClaimResult FMassSmartObjectHandler::ClaimCandidate(const FMassE
 		if (ClaimSmartObject(Entity, User, Candidate.ID))
 		{
 #if WITH_MASSGAMEPLAY_DEBUG
-			UE_CVLOG(UE::MassDebug::IsDebuggingEntity(Entity),
+			UE_CVLOG(UE::Mass::Debug::IsDebuggingEntity(Entity),
 				&SmartObjectSubsystem,
 				LogSmartObject,
 				Log,
@@ -165,7 +165,7 @@ bool FMassSmartObjectHandler::ClaimSmartObject(const FMassEntityHandle Entity, F
 	}
 
 #if WITH_MASSGAMEPLAY_DEBUG
-	UE_CVLOG(UE::MassDebug::IsDebuggingEntity(Entity),
+	UE_CVLOG(UE::Mass::Debug::IsDebuggingEntity(Entity),
 		&SmartObjectSubsystem,
 		LogSmartObject,
 		Log,
@@ -184,7 +184,7 @@ bool FMassSmartObjectHandler::UseSmartObject(
 	const FDataFragment_Transform& Transform) const
 {
 #if WITH_MASSGAMEPLAY_DEBUG
-	UE_CVLOG(UE::MassDebug::IsDebuggingEntity(Entity),
+	UE_CVLOG(UE::Mass::Debug::IsDebuggingEntity(Entity),
 		&SmartObjectSubsystem,
 		LogSmartObject,
 		Log,
@@ -210,7 +210,7 @@ bool FMassSmartObjectHandler::UseSmartObject(
 void FMassSmartObjectHandler::ReleaseSmartObject(const FMassEntityHandle Entity, FDataFragment_SmartObjectUser& User, const EMassSmartObjectInteractionStatus NewStatus) const
 {
 #if WITH_MASSGAMEPLAY_DEBUG
-	UE_CVLOG(UE::MassDebug::IsDebuggingEntity(Entity),
+	UE_CVLOG(UE::Mass::Debug::IsDebuggingEntity(Entity),
 		&SmartObjectSubsystem,
 		LogSmartObject,
 		Log,
diff --git a/Engine/Plugins/Runtime/MassGameplay/Source/MassSmartObjects/Private/MassSmartObjectProcessor.cpp b/Engine/Plugins/Runtime/MassGameplay/Source/MassSmartObjects/Private/MassSmartObjectProcessor.cpp
--- a/Engine/Plugins/Runtime/MassGameplay/Source/MassSmartObjects/Private/MassSmartObjectProcessor.cpp
+++ b/Engine/Plugins/Runtime/MassGameplay/Source/MassSmartObjects/Private/MassSmartObjectProcessor.cpp
@@ -115,7 +115,7 @@ void UMassProcessor_SmartObjectCandidatesFinder::Execute(UMassEntitySubsystem& E
 			FColor DebugColor(FColor::White);
 
 #if WITH_MASSGAMEPLAY_DEBUG
-			bDisplayDebug = UE::MassDebug::IsDebuggingEntity(Entity, &DebugColor);
+			bDisplayDebug = UE::Mass::Debug::IsDebuggingEntity(Entity, &DebugColor);
 #endif // WITH_MASSGAMEPLAY_DEBUG
 
 			BeginRequestProcessing(Entity, Context, Result);
@@ -177,7 +177,7 @@ void UMassProcessor_SmartObjectCandidatesFinder::Execute(UMassEntitySubsystem& E
 				bool bDisplayDebug = false;
 #if WITH_MASSGAMEPLAY_DEBUG
 				FColor DebugColor(FColor::White);
-				bDisplayDebug = UE::MassDebug::IsDebuggingEntity(Entity, &DebugColor);
+				bDisplayDebug = UE::Mass::Debug::IsDebuggingEntity(Entity, &DebugColor);
 #endif // WITH_MASSGAMEPLAY_DEBUG
 
 				BeginRequestProcessing(Entity, Context, Result);
@@ -308,7 +308,7 @@ void UMassProcessor_SmartObjectTimedBehavior::Execute(UMassEntitySubsystem& Enti
 #if WITH_MASSGAMEPLAY_DEBUG
 			const FMassEntityHandle Entity = Context.GetEntity(i);
 			FColor DebugColor(FColor::White);
-			const bool bIsDebuggingEntity = UE::MassDebug::IsDebuggingEntity(Entity, &DebugColor);
+			const bool bIsDebuggingEntity = UE::Mass::Debug::IsDebuggingEntity(Entity, &DebugColor);
 			if (bIsDebuggingEntity)
 			{
 				UE_CVLOG(bMustRelease, SmartObjectSubsystem, LogSmartObject, Log, TEXT("[%s] stops using [%s]"), *Entity.DebugGetDescription(), *SOUser.GetClaimHandle().Describe());
diff --git a/Engine/Plugins/Runtime/MassGameplay/Source/MassSpawner/Private/MassSpawnerSubsystem.cpp b/Engine/Plugins/Runtime/MassGameplay/Source/MassSpawner/Private/MassSpawnerSubsystem.cpp
--- a/Engine/Plugins/Runtime/MassGameplay/Source/MassSpawner/Private/MassSpawnerSubsystem.cpp
+++ b/Engine/Plugins/Runtime/MassGameplay/Source/MassSpawner/Private/MassSpawnerSubsystem.cpp
@@ -50,7 +50,7 @@ namespace UE::Mass::SpawnerSubsystem
 
 			FMassProcessingContext PipeContext(EntitySystem, /*TimeDelta=*/0.0f);
 
-			UE::Pipe::Executor::RunSparse(DestructionPipeline, PipeContext, Chunks);
+			UE::Mass::Executor::RunSparse(DestructionPipeline, PipeContext, Chunks);
 		}
 
 		return Destructors.Num() > 0;
@@ -145,7 +145,7 @@ bool UMassSpawnerSubsystem::SpawnEntities(const FMassEntityTemplate& EntityTempl
 			FMassEntityTemplate& MutableEntityTemplate = const_cast<FMassEntityTemplate&>(EntityTemplate);
 
 			FMassProcessingContext PipeContext(*EntitySystem, /*TimeDelta=*/0.0f);
-			UE::Pipe::Executor::RunSparse(MutableEntityTemplate.GetMutableInitializationPipeline(), PipeContext, ChunkCollection);
+			UE::Mass::Executor::RunSparse(MutableEntityTemplate.GetMutableInitializationPipeline(), PipeContext, ChunkCollection);
 		}
 		return true;
 	}
@@ -260,7 +260,7 @@ void UMassSpawnerSubsystem::DestroyEntities(const FMassEntityTemplateID Template
 			{
 				if (Chunks.GetArchetype().IsValid())
 				{
-					UE::Pipe::Executor::RunSparse(MutableEntityTemplate.GetMutableDeinitializationPipeline(), PipeContext, Chunks);
+					UE::Mass::Executor::RunSparse(MutableEntityTemplate.GetMutableDeinitializationPipeline(), PipeContext, Chunks);
 				}
 			}
 		}
@@ -311,7 +311,7 @@ void UMassSpawnerSubsystem::DoSpawning(const FMassEntityTemplate& EntityTemplate
 
 		FMassProcessingContext PipeContext(*EntitySystem, /*TimeDelta=*/0.0f);
 		PipeContext.AuxData = AuxData;
-		UE::Pipe::Executor::RunSparse(MutableEntityTemplate.GetMutableInitializationPipeline(), PipeContext, ChunkCollection);
+		UE::Mass::Executor::RunSparse(MutableEntityTemplate.GetMutableInitializationPipeline(), PipeContext, ChunkCollection);
 	}
 
 	OutEntities.Append(MoveTemp(SpawnedEntities));
