More Mass naming cleanup:
LogPipe, LogAggregateTicking, AggregateTicking, "Pipe", automation tests, PipeContext -> ProcessingContext, PipeEditor

diff --git a/Engine/Plugins/Runtime/MassEntity/Source/MassEntity/Private/MassEntityModule.cpp b/Engine/Plugins/Runtime/MassEntity/Source/MassEntity/Private/MassEntityModule.cpp
--- a/Engine/Plugins/Runtime/MassEntity/Source/MassEntity/Private/MassEntityModule.cpp
+++ b/Engine/Plugins/Runtime/MassEntity/Source/MassEntity/Private/MassEntityModule.cpp
@@ -8,9 +8,9 @@
 #include "Logging/MessageLog.h"
 #endif // WITH_UNREAL_DEVELOPER_TOOLS
 
-#define LOCTEXT_NAMESPACE "Pipe"
+#define LOCTEXT_NAMESPACE "Mass"
 
-class FPipeModuleModule : public IPipeModule
+class FMassEntityModuleModule : public IMassEntityModule
 {
 	/** IModuleInterface implementation */
 	virtual void StartupModule() override;
@@ -21,10 +21,10 @@ class FPipeModuleModule : public IPipeModule
 #endif // WITH_UNREAL_DEVELOPER_TOOLS
 };
 
-IMPLEMENT_MODULE(FPipeModuleModule, Pipe)
+IMPLEMENT_MODULE(FMassEntityModuleModule, MassEntity)
 
 
-void FPipeModuleModule::StartupModule()
+void FMassEntityModuleModule::StartupModule()
 {
 	// This code will execute after your module is loaded into memory (but after global variables are initialized, of course.)
 
@@ -34,15 +34,15 @@ void FPipeModuleModule::StartupModule()
 		FMessageLogInitializationOptions InitOptions;
 		InitOptions.bShowPages = true;
 		InitOptions.bShowFilters = true;
-		MessageLogModule.RegisterLogListing("Pipe", LOCTEXT("Pipe", "Pipe"), InitOptions);
+		MessageLogModule.RegisterLogListing("MassEntity", LOCTEXT("MassEntity", "MassEntity"), InitOptions);
 
-		OnWorldCleanupHandle = FWorldDelegates::OnWorldCleanup.AddStatic(&FPipeModuleModule::OnWorldCleanup);
+		OnWorldCleanupHandle = FWorldDelegates::OnWorldCleanup.AddStatic(&FMassEntityModuleModule::OnWorldCleanup);
 	}
 #endif // WITH_UNREAL_DEVELOPER_TOOLS
 }
 
 
-void FPipeModuleModule::ShutdownModule()
+void FMassEntityModuleModule::ShutdownModule()
 {
 	// This function may be called during shutdown to clean up your module.  For modules that support dynamic reloading,
 	// we call this function before unloading the module.
@@ -52,10 +52,10 @@ void FPipeModuleModule::ShutdownModule()
 }
 
 #if WITH_UNREAL_DEVELOPER_TOOLS
-void FPipeModuleModule::OnWorldCleanup(UWorld* /*World*/, bool /*bSessionEnded*/, bool /*bCleanupResources*/)
+void FMassEntityModuleModule::OnWorldCleanup(UWorld* /*World*/, bool /*bSessionEnded*/, bool /*bCleanupResources*/)
 {
 	// clearing out messages from the world being cleaned up
-	FMessageLog("Pipe").NewPage(FText::FromString(TEXT("Pipe")));
+	FMessageLog("MassEntity").NewPage(FText::FromString(TEXT("MassEntity")));
 }
 #endif // WITH_UNREAL_DEVELOPER_TOOLS
 
diff --git a/Engine/Plugins/Runtime/MassEntity/Source/MassEntity/Private/MassEntityQuery.cpp b/Engine/Plugins/Runtime/MassEntity/Source/MassEntity/Private/MassEntityQuery.cpp
--- a/Engine/Plugins/Runtime/MassEntity/Source/MassEntity/Private/MassEntityQuery.cpp
+++ b/Engine/Plugins/Runtime/MassEntity/Source/MassEntity/Private/MassEntityQuery.cpp
@@ -83,7 +83,7 @@ void FMassEntityQuery::CacheArchetypes(UMassEntitySubsystem& InEntitySubsystem)
 		}
 		else
 		{
-			UE_VLOG_UELOG(&InEntitySubsystem, LogAggregateTicking, Error, TEXT("FMassEntityQuery::CacheArchetypes: requirements not valid: %s"), *DebugGetDescription());
+			UE_VLOG_UELOG(&InEntitySubsystem, LogMass, Error, TEXT("FMassEntityQuery::CacheArchetypes: requirements not valid: %s"), *DebugGetDescription());
 		}
 	}
 }
@@ -139,7 +139,7 @@ void FMassEntityQuery::ForEachEntityChunk(UMassEntitySubsystem& EntitySubsystem,
 		if (DoesArchetypeMatchRequirements(ExecutionContext.GetChunkCollection().GetArchetype()) == false)
 		{
 			// mz@todo add a unit test for this message
-			UE_VLOG_UELOG(&EntitySubsystem, LogAggregateTicking, Error, TEXT("Attempted to execute FMassEntityQuery with an incompatible Archetype"));
+			UE_VLOG_UELOG(&EntitySubsystem, LogMass, Error, TEXT("Attempted to execute FMassEntityQuery with an incompatible Archetype"));
 			return;
 		}
 		ExecutionContext.SetRequirements(Requirements, ChunkRequirements);
@@ -168,7 +168,7 @@ void FMassEntityQuery::ForEachEntityChunk(UMassEntitySubsystem& EntitySubsystem,
 
 #if WITH_MASSENTITY_DEBUG
 	// Not using VLOG to be thread safe
-	UE_CLOG(!ExecutionContext.DebugGetExecutionDesc().IsEmpty(), LogAggregateTicking, VeryVerbose,
+	UE_CLOG(!ExecutionContext.DebugGetExecutionDesc().IsEmpty(), LogMass, VeryVerbose,
 		TEXT("%s: %d entities sent for processing"), *ExecutionContext.DebugGetExecutionDesc(), NumEntitiesToProcess);
 #endif
 
@@ -200,7 +200,7 @@ void FMassEntityQuery::ParallelForEachEntityChunk(UMassEntitySubsystem& EntitySu
 		if (DoesArchetypeMatchRequirements(ArchetypeHandle) == false)
 		{
 			// mz@todo add a unit test for this message
-			UE_VLOG_UELOG(&EntitySubsystem, LogAggregateTicking, Error, TEXT("Attempted to execute FMassEntityQuery with an incompatible Archetype"));
+			UE_VLOG_UELOG(&EntitySubsystem, LogMass, Error, TEXT("Attempted to execute FMassEntityQuery with an incompatible Archetype"));
 			return;
 		}
 
diff --git a/Engine/Plugins/Runtime/MassEntity/Source/MassEntity/Private/MassEntitySettings.cpp b/Engine/Plugins/Runtime/MassEntity/Source/MassEntity/Private/MassEntitySettings.cpp
--- a/Engine/Plugins/Runtime/MassEntity/Source/MassEntity/Private/MassEntitySettings.cpp
+++ b/Engine/Plugins/Runtime/MassEntity/Source/MassEntity/Private/MassEntitySettings.cpp
@@ -107,15 +107,15 @@ void UMassSettings::AddToActiveProcessorsList(TSubclassOf<UMassProcessor> Proces
 	{
 		if (ProcessorClass == UMassCompositeProcessor::StaticClass())
 		{
-			UE_VLOG_UELOG(this, LogPipe, Log, TEXT("%s adding PipeCompositeProcessor to the global processor list is unsupported"), ANSI_TO_TCHAR(__FUNCTION__));
+			UE_VLOG_UELOG(this, LogMass, Log, TEXT("%s adding MassCompositeProcessor to the global processor list is unsupported"), ANSI_TO_TCHAR(__FUNCTION__));
 		}
 		else if (ProcessorClass->HasAnyClassFlags(CLASS_Abstract))
 		{
-			UE_VLOG_UELOG(this, LogPipe, Log, TEXT("%s unable to add %s due to it being an abstract class"), ANSI_TO_TCHAR(__FUNCTION__), *ProcessorClass->GetName());
+			UE_VLOG_UELOG(this, LogMass, Log, TEXT("%s unable to add %s due to it being an abstract class"), ANSI_TO_TCHAR(__FUNCTION__), *ProcessorClass->GetName());
 		}
 		else if (ProcessorCDOs.Find(ProcessorCDO) != INDEX_NONE)
 		{
-			UE_VLOG_UELOG(this, LogPipe, Log, TEXT("%s already in global processor list"), *ProcessorCDO->GetName());
+			UE_VLOG_UELOG(this, LogMass, Log, TEXT("%s already in global processor list"), *ProcessorCDO->GetName());
 		}
 		else 
 		{
@@ -135,9 +135,6 @@ const FMassProcessingPhaseConfig* UMassSettings::GetProcessingPhasesConfig()
 #if WITH_EDITOR
 void UMassSettings::PostEditChangeProperty(FPropertyChangedEvent& PropertyChangedEvent)
 {
-	// @todo replacing a valid entry in the array is not handled right now. 
-
-	static const FName PipeName = TEXT("Pipe");
 	static const FName ProcessorCDOsName = GET_MEMBER_NAME_CHECKED(UMassSettings, ProcessorCDOs);
 
 	Super::PostEditChangeProperty(PropertyChangedEvent);
diff --git a/Engine/Plugins/Runtime/MassEntity/Source/MassEntity/Private/MassEntitySubsystem.cpp b/Engine/Plugins/Runtime/MassEntity/Source/MassEntity/Private/MassEntitySubsystem.cpp
--- a/Engine/Plugins/Runtime/MassEntity/Source/MassEntity/Private/MassEntitySubsystem.cpp
+++ b/Engine/Plugins/Runtime/MassEntity/Source/MassEntity/Private/MassEntitySubsystem.cpp
@@ -9,8 +9,6 @@
 
 const FMassEntityHandle UMassEntitySubsystem::InvalidEntity;
 
-DEFINE_LOG_CATEGORY(LogAggregateTicking);
-
 //@TODO: Everything still alive leaks at shutdown
 //@TODO: No re-entrance safety while running a system (e.g., preventing someone from adding/removing entities or altering archetypes, etc...)
 //@TODO: Do we allow GCable types?  If so, need to implement AddReferencedObjects
@@ -103,7 +101,7 @@ FArchetypeHandle UMassEntitySubsystem::CreateArchetype(TConstArrayView<const USc
 		}
 		else
 		{
-			UE_LOG(LogAggregateTicking, Warning, TEXT("%s: %s is not a valid component nor tag type. Ignoring.")
+			UE_LOG(LogMass, Warning, TEXT("%s: %s is not a valid component nor tag type. Ignoring.")
 				, ANSI_TO_TCHAR(__FUNCTION__), *GetNameSafe(Type));
 		}
 	}
@@ -758,7 +756,7 @@ void UMassEntitySubsystem::GetValidArchetypes(const FMassEntityQuery& Query, TAr
 			const FMassTagBitSet UnsatisfiedTags = Query.GetRequiredAllTags() - Archetype.GetTagBitSet();
 			FStringOutputDevice Description;
 			UnsatisfiedTags.DebugGetStringDesc(Description);
-			UE_LOG(LogAggregateTicking, VeryVerbose, TEXT("Archetype did not match due to missing tags: %s")
+			UE_LOG(LogMass, VeryVerbose, TEXT("Archetype did not match due to missing tags: %s")
 				, *Description);
 #endif // WITH_MASSENTITY_DEBUG
 			continue;
@@ -771,7 +769,7 @@ void UMassEntitySubsystem::GetValidArchetypes(const FMassEntityQuery& Query, TAr
 			const FMassTagBitSet UnwantedTags = Query.GetRequiredAllTags().GetOverlap(Archetype.GetTagBitSet());
 			FStringOutputDevice Description;
 			UnwantedTags.DebugGetStringDesc(Description);
-			UE_LOG(LogAggregateTicking, VeryVerbose, TEXT("Archetype has tags required absent: %s")
+			UE_LOG(LogMass, VeryVerbose, TEXT("Archetype has tags required absent: %s")
 				, *Description);
 #endif // WITH_MASSENTITY_DEBUG
 			continue;
@@ -783,7 +781,7 @@ void UMassEntitySubsystem::GetValidArchetypes(const FMassEntityQuery& Query, TAr
 #if WITH_MASSENTITY_DEBUG
 			FStringOutputDevice Description;
 			Query.GetRequiredAnyTags().DebugGetStringDesc(Description);
-			UE_LOG(LogAggregateTicking, VeryVerbose, TEXT("Archetype did not match due to missing \'any\' tags: %s")
+			UE_LOG(LogMass, VeryVerbose, TEXT("Archetype did not match due to missing \'any\' tags: %s")
 				, *Description);
 #endif // WITH_MASSENTITY_DEBUG
 			continue;
@@ -796,7 +794,7 @@ void UMassEntitySubsystem::GetValidArchetypes(const FMassEntityQuery& Query, TAr
 			const FMassFragmentBitSet UnsatisfiedComponents = Query.GetRequiredAllComponents() - Archetype.GetComponentBitSet();
 			FStringOutputDevice Description;
 			UnsatisfiedComponents.DebugGetStringDesc(Description);
-			UE_LOG(LogAggregateTicking, VeryVerbose, TEXT("Archetype did not match due to missing Components: %s")
+			UE_LOG(LogMass, VeryVerbose, TEXT("Archetype did not match due to missing Components: %s")
 				, *Description);
 #endif // WITH_MASSENTITY_DEBUG
 			continue;
@@ -809,7 +807,7 @@ void UMassEntitySubsystem::GetValidArchetypes(const FMassEntityQuery& Query, TAr
 			const FMassFragmentBitSet UnwantedComponents = Query.GetRequiredAllComponents().GetOverlap(Archetype.GetComponentBitSet());
 			FStringOutputDevice Description;
 			UnwantedComponents.DebugGetStringDesc(Description);
-			UE_LOG(LogAggregateTicking, VeryVerbose, TEXT("Archetype has Components required absent: %s")
+			UE_LOG(LogMass, VeryVerbose, TEXT("Archetype has Components required absent: %s")
 				, *Description);
 #endif // WITH_MASSENTITY_DEBUG
 			continue;
@@ -821,7 +819,7 @@ void UMassEntitySubsystem::GetValidArchetypes(const FMassEntityQuery& Query, TAr
 #if WITH_MASSENTITY_DEBUG
 			FStringOutputDevice Description;
 			Query.GetRequiredAnyComponents().DebugGetStringDesc(Description);
-			UE_LOG(LogAggregateTicking, VeryVerbose, TEXT("Archetype did not match due to missing \'any\' components: %s")
+			UE_LOG(LogMass, VeryVerbose, TEXT("Archetype did not match due to missing \'any\' components: %s")
 				, *Description);
 #endif // WITH_MASSENTITY_DEBUG
 			continue;
@@ -834,7 +832,7 @@ void UMassEntitySubsystem::GetValidArchetypes(const FMassEntityQuery& Query, TAr
 			const FMassChunkFragmentBitSet UnsatisfiedComponents = Query.GetRequiredAllChunkComponents() - Archetype.GetChunkComponentBitSet();
 			FStringOutputDevice Description;
 			UnsatisfiedComponents.DebugGetStringDesc(Description);
-			UE_LOG(LogAggregateTicking, VeryVerbose, TEXT("Archetype did not match due to missing Chunk Components: %s")
+			UE_LOG(LogMass, VeryVerbose, TEXT("Archetype did not match due to missing Chunk Components: %s")
 				, *Description);
 #endif // WITH_MASSENTITY_DEBUG
 			continue;
@@ -847,7 +845,7 @@ void UMassEntitySubsystem::GetValidArchetypes(const FMassEntityQuery& Query, TAr
 			const FMassChunkFragmentBitSet UnwantedComponents = Query.GetRequiredNoneChunkComponents().GetOverlap(Archetype.GetChunkComponentBitSet());
 			FStringOutputDevice Description;
 			UnwantedComponents.DebugGetStringDesc(Description);
-			UE_LOG(LogAggregateTicking, VeryVerbose, TEXT("Archetype has Chunk Components required absent: %s")
+			UE_LOG(LogMass, VeryVerbose, TEXT("Archetype has Chunk Components required absent: %s")
 				, *Description);
 #endif // WITH_MASSENTITY_DEBUG
 			continue;
diff --git a/Engine/Plugins/Runtime/MassEntity/Source/MassEntity/Private/MassExecutor.cpp b/Engine/Plugins/Runtime/MassEntity/Source/MassEntity/Private/MassExecutor.cpp
--- a/Engine/Plugins/Runtime/MassEntity/Source/MassEntity/Private/MassExecutor.cpp
+++ b/Engine/Plugins/Runtime/MassEntity/Source/MassEntity/Private/MassExecutor.cpp
@@ -8,22 +8,22 @@
 
 namespace UE::Mass::Executor
 {
-void Run(FMassRuntimePipeline& RuntimePipeline, FMassProcessingContext& PipeContext)
+void Run(FMassRuntimePipeline& RuntimePipeline, FMassProcessingContext& ProcessingContext)
 {
-	if (!ensure(PipeContext.EntitySubsystem) || 
-		!ensure(PipeContext.DeltaSeconds >= 0.f) ||
+	if (!ensure(ProcessingContext.EntitySubsystem) || 
+		!ensure(ProcessingContext.DeltaSeconds >= 0.f) ||
 		!ensure(RuntimePipeline.Processors.Find(nullptr) == INDEX_NONE))
 	{
 		return;
 	}
 
 	TRACE_CPUPROFILER_EVENT_SCOPE_STR("PipeExecutor Run Pipeline")
-	RunProcessorsView(RuntimePipeline.Processors, PipeContext);
+	RunProcessorsView(RuntimePipeline.Processors, ProcessingContext);
 }
 
-void RunSparse(FMassRuntimePipeline& RuntimePipeline, FMassProcessingContext& PipeContext, FArchetypeHandle Archetype, TConstArrayView<FMassEntityHandle> Entities)
+void RunSparse(FMassRuntimePipeline& RuntimePipeline, FMassProcessingContext& ProcessingContext, FArchetypeHandle Archetype, TConstArrayView<FMassEntityHandle> Entities)
 {
-	if (!ensure(PipeContext.EntitySubsystem) ||
+	if (!ensure(ProcessingContext.EntitySubsystem) ||
 		!ensure(RuntimePipeline.Processors.Find(nullptr) == INDEX_NONE) ||
 		RuntimePipeline.Processors.Num() == 0 ||
 		!ensureMsgf(Archetype.IsValid(), TEXT("The Archetype passed in to UE::Mass::Executor::RunSparse is invalid")))
@@ -34,12 +34,12 @@ void RunSparse(FMassRuntimePipeline& RuntimePipeline, FMassProcessingContext& Pi
 	TRACE_CPUPROFILER_EVENT_SCOPE_STR("PipeExecutor RunSparseEntities");
 
 	const FArchetypeChunkCollection ChunkCollection(Archetype, Entities);
-	RunProcessorsView(RuntimePipeline.Processors, PipeContext, &ChunkCollection);
+	RunProcessorsView(RuntimePipeline.Processors, ProcessingContext, &ChunkCollection);
 }
 
-void RunSparse(FMassRuntimePipeline& RuntimePipeline, FMassProcessingContext& PipeContext, const FArchetypeChunkCollection& ChunkCollection)
+void RunSparse(FMassRuntimePipeline& RuntimePipeline, FMassProcessingContext& ProcessingContext, const FArchetypeChunkCollection& ChunkCollection)
 {
-	if (!ensure(PipeContext.EntitySubsystem) ||
+	if (!ensure(ProcessingContext.EntitySubsystem) ||
 		!ensure(RuntimePipeline.Processors.Find(nullptr) == INDEX_NONE) ||
 		RuntimePipeline.Processors.Num() == 0 ||
 		!ensureMsgf(ChunkCollection.GetArchetype().IsValid(), TEXT("The Archetype of ChunkCollection passed in to UE::Mass::Executor::RunSparse is invalid")))
@@ -49,12 +49,12 @@ void RunSparse(FMassRuntimePipeline& RuntimePipeline, FMassProcessingContext& Pi
 
 	TRACE_CPUPROFILER_EVENT_SCOPE_STR("PipeExecutor RunSparse");
 
-	RunProcessorsView(RuntimePipeline.Processors, PipeContext, &ChunkCollection);
+	RunProcessorsView(RuntimePipeline.Processors, ProcessingContext, &ChunkCollection);
 }
 
-void Run(UMassProcessor& Processor, FMassProcessingContext& PipeContext)
+void Run(UMassProcessor& Processor, FMassProcessingContext& ProcessingContext)
 {
-	if (!ensure(PipeContext.EntitySubsystem) || !ensure(PipeContext.DeltaSeconds >= 0.f))
+	if (!ensure(ProcessingContext.EntitySubsystem) || !ensure(ProcessingContext.DeltaSeconds >= 0.f))
 	{
 		return;
 	}
@@ -62,29 +62,29 @@ void Run(UMassProcessor& Processor, FMassProcessingContext& PipeContext)
 	TRACE_CPUPROFILER_EVENT_SCOPE_STR("PipeExecutor Run")
 
 	UMassProcessor* ProcPtr = &Processor;
-	RunProcessorsView(MakeArrayView(&ProcPtr, 1), PipeContext);
+	RunProcessorsView(MakeArrayView(&ProcPtr, 1), ProcessingContext);
 }
 
-void RunProcessorsView(TArrayView<UMassProcessor*> Processors, FMassProcessingContext& PipeContext, const FArchetypeChunkCollection* ChunkCollection)
+void RunProcessorsView(TArrayView<UMassProcessor*> Processors, FMassProcessingContext& ProcessingContext, const FArchetypeChunkCollection* ChunkCollection)
 {
 	TRACE_CPUPROFILER_EVENT_SCOPE(RunProcessorsView);
 
-	if (PipeContext.EntitySubsystem == nullptr)
+	if (ProcessingContext.EntitySubsystem == nullptr)
 	{
-		UE_LOG(LogPipe, Error, TEXT("%s PipeContext.EntitySubsystem is null. Baling out."), ANSI_TO_TCHAR(__FUNCTION__));
+		UE_LOG(LogMass, Error, TEXT("%s ProcessingContext.EntitySubsystem is null. Baling out."), ANSI_TO_TCHAR(__FUNCTION__));
 		return;
 	}
 #if WITH_MASSENTITY_DEBUG
 	if (Processors.Find(nullptr) != INDEX_NONE)
 	{
-		UE_LOG(LogPipe, Error, TEXT("%s input Processors contains nullptr. Baling out."), ANSI_TO_TCHAR(__FUNCTION__));
+		UE_LOG(LogMass, Error, TEXT("%s input Processors contains nullptr. Baling out."), ANSI_TO_TCHAR(__FUNCTION__));
 		return;
 	}
 #endif // WITH_MASSENTITY_DEBUG
 
 	TRACE_CPUPROFILER_EVENT_SCOPE_STR("PipeExecutor RunProcessorsView")
 
-	FMassExecutionContext ExecutionContext(PipeContext.DeltaSeconds);
+	FMassExecutionContext ExecutionContext(ProcessingContext.DeltaSeconds);
 	if (ChunkCollection)
 	{
 		ExecutionContext.SetChunkCollection(*ChunkCollection);
@@ -92,16 +92,16 @@ void RunProcessorsView(TArrayView<UMassProcessor*> Processors, FMassProcessingCo
 	// manually creating a new command buffer to let the default one still be used by code unaware of pipe processing
 	ExecutionContext.SetDeferredCommandBuffer(MakeShareable(new FMassCommandBuffer()));
 	ExecutionContext.SetFlushDeferredCommands(false);
-	ExecutionContext.SetAuxData(PipeContext.AuxData);
+	ExecutionContext.SetAuxData(ProcessingContext.AuxData);
 
 	{
 		TRACE_CPUPROFILER_EVENT_SCOPE_STR("Execute Processors")
 		
-		UMassEntitySubsystem::FScopedProcessing ProcessingScope = PipeContext.EntitySubsystem->NewProcessingScope();
+		UMassEntitySubsystem::FScopedProcessing ProcessingScope = ProcessingContext.EntitySubsystem->NewProcessingScope();
 
 		for (UMassProcessor* Proc : Processors)
 		{
-			Proc->CallExecute(*PipeContext.EntitySubsystem, ExecutionContext);
+			Proc->CallExecute(*ProcessingContext.EntitySubsystem, ExecutionContext);
 		}
 	}
 	
@@ -110,8 +110,8 @@ void RunProcessorsView(TArrayView<UMassProcessor*> Processors, FMassProcessingCo
 		
 		ExecutionContext.SetFlushDeferredCommands(true);
 		// append the commands added from other, non-processor sources (like MassAgentSubsystem)
-		ExecutionContext.Defer().MoveAppend(PipeContext.EntitySubsystem->Defer());
-		ExecutionContext.FlushDeferred(*PipeContext.EntitySubsystem);
+		ExecutionContext.Defer().MoveAppend(ProcessingContext.EntitySubsystem->Defer());
+		ExecutionContext.FlushDeferred(*ProcessingContext.EntitySubsystem);
 	}
 }
 
@@ -140,7 +140,7 @@ struct FPipeExecutorDoneTask
 			ExecutionContext.Defer().MoveAppend(EntitySubsystem.Defer());
 		}
 
-		UE_LOG(LogPipe, Log, TEXT("PipeExecutor %s tasks DONE"), *DebugName);
+		UE_LOG(LogMass, Log, TEXT("PipeExecutor %s tasks DONE"), *DebugName);
 		ExecutionContext.SetFlushDeferredCommands(true);
 		ExecutionContext.FlushDeferred(EntitySubsystem);
 
@@ -153,35 +153,35 @@ private:
 	FString DebugName;
 };
 
-FGraphEventRef TriggerParallelTasks(UMassProcessor& Processor, FMassProcessingContext& PipeContext, TFunction<void()> OnDoneNotification)
+FGraphEventRef TriggerParallelTasks(UMassProcessor& Processor, FMassProcessingContext& ProcessingContext, TFunction<void()> OnDoneNotification)
 {
 	TRACE_CPUPROFILER_EVENT_SCOPE(RunProcessorsView);
 
-	if (PipeContext.EntitySubsystem == nullptr)
+	if (ProcessingContext.EntitySubsystem == nullptr)
 	{
-		UE_LOG(LogPipe, Error, TEXT("%s PipeContext.EntitySubsystem is null. Baling out."), ANSI_TO_TCHAR(__FUNCTION__));
+		UE_LOG(LogMass, Error, TEXT("%s ProcessingContext.EntitySubsystem is null. Baling out."), ANSI_TO_TCHAR(__FUNCTION__));
 		return FGraphEventRef();
 	}
 
 	TRACE_CPUPROFILER_EVENT_SCOPE_STR("PipeExecutor RunParallel")
 
 	// not going through UMassEntitySubsystem::CreateExecutionContext on purpose - we do need a separate command buffer
-	FMassExecutionContext ExecutionContext(PipeContext.DeltaSeconds);
+	FMassExecutionContext ExecutionContext(ProcessingContext.DeltaSeconds);
 	ExecutionContext.SetDeferredCommandBuffer(MakeShareable(new FMassCommandBuffer()));
 	ExecutionContext.SetFlushDeferredCommands(false);
-	ExecutionContext.SetAuxData(PipeContext.AuxData);
+	ExecutionContext.SetAuxData(ProcessingContext.AuxData);
 
 	FGraphEventRef CompletionEvent;
 	{
 		TRACE_CPUPROFILER_EVENT_SCOPE_STR("Dispatch Processors")
-		CompletionEvent = Processor.DispatchProcessorTasks(*PipeContext.EntitySubsystem, ExecutionContext, {});
+		CompletionEvent = Processor.DispatchProcessorTasks(*ProcessingContext.EntitySubsystem, ExecutionContext, {});
 	}
 
 	if (CompletionEvent.IsValid())
 	{
 		const FGraphEventArray Prerequisites = { CompletionEvent };
 		CompletionEvent = TGraphTask<FPipeExecutorDoneTask>::CreateTask(&Prerequisites)
-			.ConstructAndDispatchWhenReady(ExecutionContext, *PipeContext.EntitySubsystem, OnDoneNotification, Processor.GetName());
+			.ConstructAndDispatchWhenReady(ExecutionContext, *ProcessingContext.EntitySubsystem, OnDoneNotification, Processor.GetName());
 	}
 
 	return CompletionEvent;
diff --git a/Engine/Plugins/Runtime/MassEntity/Source/MassEntity/Private/MassProcessingPhase.cpp b/Engine/Plugins/Runtime/MassEntity/Source/MassEntity/Private/MassProcessingPhase.cpp
--- a/Engine/Plugins/Runtime/MassEntity/Source/MassEntity/Private/MassProcessingPhase.cpp
+++ b/Engine/Plugins/Runtime/MassEntity/Source/MassEntity/Private/MassProcessingPhase.cpp
@@ -155,7 +155,7 @@ void UMassProcessingPhaseManager::InitializePhases(UObject& InProcessorOwner)
 
 #if WITH_MASSENTITY_DEBUG
 	// print it all out to vislog
-	UE_VLOG_UELOG(this, LogPipe, Verbose, TEXT("Phases initialization done. Current composition:"));
+	UE_VLOG_UELOG(this, LogMass, Verbose, TEXT("Phases initialization done. Current composition:"));
 
 	FStringOutputDevice OutDescription;
 	for (int i = 0; i < int(EMassProcessingPhase::MAX); ++i)
@@ -163,7 +163,7 @@ void UMassProcessingPhaseManager::InitializePhases(UObject& InProcessorOwner)
 		if (ProcessingPhases[i].PhaseProcessor)
 		{
 			ProcessingPhases[i].PhaseProcessor->DebugOutputDescription(OutDescription);
-			UE_VLOG_UELOG(this, LogPipe, Verbose, TEXT("--- %s"), *OutDescription);
+			UE_VLOG_UELOG(this, LogMass, Verbose, TEXT("--- %s"), *OutDescription);
 			OutDescription.Reset();
 		}
 	}	
@@ -180,7 +180,7 @@ void UMassProcessingPhaseManager::Start(UWorld& World)
 	}
 	else
 	{
-		UE_VLOG_UELOG(this, LogPipe, Error, TEXT("Called %s while missing the EntitySubsystem"), ANSI_TO_TCHAR(__FUNCTION__));
+		UE_VLOG_UELOG(this, LogMass, Error, TEXT("Called %s while missing the EntitySubsystem"), ANSI_TO_TCHAR(__FUNCTION__));
 	}
 }
 
@@ -201,7 +201,7 @@ void UMassProcessingPhaseManager::EnableTickFunctions(const UWorld& World)
 		Phase.RegisterTickFunction(World.PersistentLevel);
 		Phase.SetTickFunctionEnable(true);
 	}
-	UE_VLOG_UELOG(this, LogPipe, Log, TEXT("UPipeTickManager %s.%s has been started")
+	UE_VLOG_UELOG(this, LogMass, Log, TEXT("UPipeTickManager %s.%s has been started")
 		, *GetNameSafe(GetOuter()), *GetName());
 }
 
@@ -214,7 +214,7 @@ void UMassProcessingPhaseManager::Stop()
 		Phase.SetTickFunctionEnable(false);
 	}
 
-	UE_VLOG_UELOG(this, LogPipe, Log, TEXT("UPipeTickManager %s.%s has been stopped")
+	UE_VLOG_UELOG(this, LogMass, Log, TEXT("UPipeTickManager %s.%s has been stopped")
 		, *GetNameSafe(GetOuter()), *GetName());
 }
 
diff --git a/Engine/Plugins/Runtime/MassEntity/Source/MassEntity/Private/MassProcessingTypes.cpp b/Engine/Plugins/Runtime/MassEntity/Source/MassEntity/Private/MassProcessingTypes.cpp
--- a/Engine/Plugins/Runtime/MassEntity/Source/MassEntity/Private/MassProcessingTypes.cpp
+++ b/Engine/Plugins/Runtime/MassEntity/Source/MassEntity/Private/MassProcessingTypes.cpp
@@ -8,7 +8,7 @@
 #include "VisualLogger/VisualLogger.h"
 #include "MassEntityDebug.h"
 
-DEFINE_LOG_CATEGORY(LogPipe);
+DEFINE_LOG_CATEGORY(LogMass);
 
 //----------------------------------------------------------------------//
 //  FMassProcessingContext
@@ -58,7 +58,7 @@ void FMassRuntimePipeline::InitializeFromSchematics(TConstArrayView<TSoftObjectP
 		}
 		else
 		{
-			UE_LOG(LogPipe, Error, TEXT("Unable to resolve MassSchematic %s while creating FMassRuntimePipeline"), *Schematic.GetLongPackageName());
+			UE_LOG(LogMass, Error, TEXT("Unable to resolve MassSchematic %s while creating FMassRuntimePipeline"), *Schematic.GetLongPackageName());
 		}
 	}
 
@@ -96,7 +96,7 @@ void FMassRuntimePipeline::InitializeFromClassArray(TConstArrayView<TSubclassOf<
 			}
 			else
 			{
-				UE_CVLOG(CDO, &InOwner, LogPipe, Log, TEXT("Skipping %s due to ExecutionFlags"), *CDO->GetName());
+				UE_CVLOG(CDO, &InOwner, LogMass, Log, TEXT("Skipping %s due to ExecutionFlags"), *CDO->GetName());
 			}
 		}
 	}
@@ -130,11 +130,11 @@ void FMassRuntimePipeline::AppendUniqueRuntimeProcessorCopies(TConstArrayView<co
 		{
 			if (Proc->ShouldExecute(WorldExecutionFlags) == false)
 			{
-				UE_VLOG(&InOwner, LogPipe, Log, TEXT("Skipping %s due to ExecutionFlags"), *Proc->GetName());
+				UE_VLOG(&InOwner, LogMass, Log, TEXT("Skipping %s due to ExecutionFlags"), *Proc->GetName());
 			}
 			else if (Proc->AllowDuplicates() == false)
 			{
-				UE_VLOG(&InOwner, LogPipe, Log, TEXT("Skipping %s due to it being a duplicate"), *Proc->GetName());
+				UE_VLOG(&InOwner, LogMass, Log, TEXT("Skipping %s due to it being a duplicate"), *Proc->GetName());
 			}
 		}
 #endif // WITH_MASSENTITY_DEBUG
@@ -186,7 +186,7 @@ void FMassRuntimePipeline::AppendOrOverrideRuntimeProcessorCopies(TConstArrayVie
 		}
 		else
 		{
-			UE_CVLOG(Proc, &InOwner, LogPipe, Log, TEXT("Skipping %s due to ExecutionFlags"), *Proc->GetName());
+			UE_CVLOG(Proc, &InOwner, LogMass, Log, TEXT("Skipping %s due to ExecutionFlags"), *Proc->GetName());
 		}
 	}
 }
diff --git a/Engine/Plugins/Runtime/MassEntity/Source/MassEntity/Private/MassProcessor.cpp b/Engine/Plugins/Runtime/MassEntity/Source/MassEntity/Private/MassProcessor.cpp
--- a/Engine/Plugins/Runtime/MassEntity/Source/MassEntity/Private/MassProcessor.cpp
+++ b/Engine/Plugins/Runtime/MassEntity/Source/MassEntity/Private/MassProcessor.cpp
@@ -30,7 +30,7 @@ namespace UE::Mass::Debug
 
 // change to && 1 to enable more detailed processing tasks logging
 #if WITH_MASSENTITY_DEBUG && 0
-#define PROCESSOR_LOG(Fmt, ...) UE_LOG(LogPipe, Log, Fmt, ##__VA_ARGS__)
+#define PROCESSOR_LOG(Fmt, ...) UE_LOG(LogMass, Log, Fmt, ##__VA_ARGS__)
 #else // WITH_MASSENTITY_DEBUG
 #define PROCESSOR_LOG(...) 
 #endif // WITH_MASSENTITY_DEBUG
diff --git a/Engine/Plugins/Runtime/MassEntity/Source/MassEntity/Private/MassProcessorDependencySolver.cpp b/Engine/Plugins/Runtime/MassEntity/Source/MassEntity/Private/MassProcessorDependencySolver.cpp
--- a/Engine/Plugins/Runtime/MassEntity/Source/MassEntity/Private/MassProcessorDependencySolver.cpp
+++ b/Engine/Plugins/Runtime/MassEntity/Source/MassEntity/Private/MassProcessorDependencySolver.cpp
@@ -8,7 +8,7 @@
 #include "Misc/Paths.h"
 #include "Engine/World.h"
 
-#define LOCTEXT_NAMESPACE "Pipe"
+#define LOCTEXT_NAMESPACE "Mass"
 
 //----------------------------------------------------------------------//
 //  FProcessorDependencySolver::FNode
@@ -143,7 +143,7 @@ void FProcessorDependencySolver::AddNode(FName InGroupName, UMassProcessor& Proc
 		}
 		else
 		{
-			UE_LOG(LogPipe, Warning, TEXT("%s Processor %s already registered. Duplicates are not supported.")
+			UE_LOG(LogMass, Warning, TEXT("%s Processor %s already registered. Duplicates are not supported.")
 				, ANSI_TO_TCHAR(__FUNCTION__), *ProcName.ToString());
 		}
 	}
@@ -182,7 +182,7 @@ void FProcessorDependencySolver::AddNode(FName InGroupName, UMassProcessor& Proc
 		}
 		else
 		{
-			UE_LOG(LogPipe, Warning, TEXT("%s Processor %s already registered. Duplicates are not supported.")
+			UE_LOG(LogMass, Warning, TEXT("%s Processor %s already registered. Duplicates are not supported.")
 				, ANSI_TO_TCHAR(__FUNCTION__), *ProcName.ToString());
 		}
 	}
@@ -231,7 +231,7 @@ void FProcessorDependencySolver::BuildDependencies(FNode& RootNode)
 				}
 				else
 				{
-					UE_LOG(LogPipe, Warning, TEXT("%s a node (%s) tried to add a circular dependency on itself")
+					UE_LOG(LogMass, Warning, TEXT("%s a node (%s) tried to add a circular dependency on itself")
 						, ANSI_TO_TCHAR(__FUNCTION__), *Node.Name.ToString());
 				}
 			}
@@ -253,7 +253,7 @@ void FProcessorDependencySolver::BuildDependencies(FNode& RootNode)
 				}
 				else
 				{
-					UE_LOG(LogPipe, Warning, TEXT("%s a node (%s) tried to add a circular dependency on itself")
+					UE_LOG(LogMass, Warning, TEXT("%s a node (%s) tried to add a circular dependency on itself")
 						, ANSI_TO_TCHAR(__FUNCTION__), *Node.Name.ToString());
 				}
 			}
@@ -272,12 +272,12 @@ void FProcessorDependencySolver::LogNode(const FNode& RootNode, const FNode* Par
 		}
 	}
 
-	UE_LOG(LogPipe, Log, TEXT("%*s%s %s%s"), Indent, TEXT(""), *RootNode.Name.ToString()
+	UE_LOG(LogMass, Log, TEXT("%*s%s %s%s"), Indent, TEXT(""), *RootNode.Name.ToString()
 		, Dependencies.Len() > 0 ? TEXT("after: ") : TEXT(""), *Dependencies);
 
 	if (RootNode.Processor == nullptr)
 	{
-		UE_LOG(LogPipe, Log, TEXT("%*s%s before:%s after:%s"), Indent, TEXT(""), *RootNode.Name.ToString()
+		UE_LOG(LogMass, Log, TEXT("%*s%s before:%s after:%s"), Indent, TEXT(""), *RootNode.Name.ToString()
 			, *NameViewToString(RootNode.ExecuteBefore)
 			, *NameViewToString(RootNode.ExecuteAfter));
 
@@ -288,7 +288,7 @@ void FProcessorDependencySolver::LogNode(const FNode& RootNode, const FNode* Par
 	}
 	else
 	{
-		UE_LOG(LogPipe, Log, TEXT("%*s%s before:%s after:%s"), Indent, TEXT(""), *RootNode.Name.ToString()
+		UE_LOG(LogMass, Log, TEXT("%*s%s before:%s after:%s"), Indent, TEXT(""), *RootNode.Name.ToString()
 			, *NameViewToString(RootNode.Processor->GetExecutionOrder().ExecuteBefore)
 			, *NameViewToString(RootNode.Processor->GetExecutionOrder().ExecuteAfter));
 	}
@@ -758,13 +758,13 @@ void FProcessorDependencySolver::Solve(FNode& RootNode, TConstArrayView<const FN
 		{
 			bAnyCyclesDetected = true;
 
-			UE_LOG(LogPipe, Error, TEXT("Detected processing dependency cycle:"));
+			UE_LOG(LogMass, Error, TEXT("Detected processing dependency cycle:"));
 			for (const int32 Index : IndicesRemaining)
 			{
 				UMassProcessor* Processor = RootNode.SubNodes[Index].Processor;
 				if (Processor)
 				{
-					UE_LOG(LogPipe, Error, TEXT("\t%s, group: %s, before: %s, after %s")
+					UE_LOG(LogMass, Error, TEXT("\t%s, group: %s, before: %s, after %s")
 						, *Processor->GetName()
 						, *Processor->GetExecutionOrder().ExecuteInGroup.ToString()
 						, *NameViewToString(Processor->GetExecutionOrder().ExecuteBefore)
@@ -773,10 +773,10 @@ void FProcessorDependencySolver::Solve(FNode& RootNode, TConstArrayView<const FN
 				else
 				{
 					// group
-					UE_LOG(LogPipe, Error, TEXT("\tGroup %s"), *RootNode.SubNodes[Index].Name.ToString());
+					UE_LOG(LogMass, Error, TEXT("\tGroup %s"), *RootNode.SubNodes[Index].Name.ToString());
 				}
 			}
-			UE_LOG(LogPipe, Error, TEXT("Cutting the chain at an arbitrary location."));
+			UE_LOG(LogMass, Error, TEXT("Cutting the chain at an arbitrary location."));
 
 			// remove first dependency
 			// note that if we're in a cycle handling scenario every node does have some dependencies left
@@ -815,18 +815,18 @@ void FProcessorDependencySolver::Solve(FNode& RootNode, TConstArrayView<const FN
 
 void FProcessorDependencySolver::ResolveDependencies(TArray<FProcessorDependencySolver::FOrderInfo>& OutResult, TConstArrayView<const FName> PriorityNodes)
 {
-	FScopedCategoryAndVerbosityOverride LogOverride(TEXT("LogPipe"), ELogVerbosity::Log);
+	FScopedCategoryAndVerbosityOverride LogOverride(TEXT("LogMass"), ELogVerbosity::Log);
 
 	bAnyCyclesDetected = false;
 
-	UE_LOG(LogPipe, Log, TEXT("Gathering dependencies data:"));
+	UE_LOG(LogMass, Log, TEXT("Gathering dependencies data:"));
 
 	// gather the processors information first
 	for (UMassProcessor* Processor : Processors)
 	{
 		if (Processor == nullptr)
 		{
-			UE_LOG(LogPipe, Warning, TEXT("%s nullptr found in Processors collection being processed"), ANSI_TO_TCHAR(__FUNCTION__));
+			UE_LOG(LogMass, Warning, TEXT("%s nullptr found in Processors collection being processed"), ANSI_TO_TCHAR(__FUNCTION__));
 			continue;
 		}
 		AddNode(Processor->GetExecutionOrder().ExecuteInGroup, *Processor);
@@ -842,13 +842,13 @@ void FProcessorDependencySolver::ResolveDependencies(TArray<FProcessorDependency
 		const UWorld* World = Processors.IsEmpty() ? nullptr : Processors[0]->GetWorld();
 	    for (const FName& DependencyName : GroupRootNode.ExecuteBefore)
 	    {
-		    UE_LOG(LogPipe, Warning, TEXT("(%s) %s found an unresolved execute before dependency (%s)"),
+		    UE_LOG(LogMass, Warning, TEXT("(%s) %s found an unresolved execute before dependency (%s)"),
 		    	World != nullptr ? *ToString(World->GetNetMode()) : TEXT("unknown"), ANSI_TO_TCHAR(__FUNCTION__), *DependencyName.ToString());
 	    }
     
 	    for (const FName& DependencyName : GroupRootNode.ExecuteAfter)
 	    {
-		    UE_LOG(LogPipe, Warning, TEXT("(%s) %s found an unresolved execute after dependency (%s)"),
+		    UE_LOG(LogMass, Warning, TEXT("(%s) %s found an unresolved execute after dependency (%s)"),
 		    	World != nullptr ? *ToString(World->GetNetMode()) : TEXT("unknown"), ANSI_TO_TCHAR(__FUNCTION__), *DependencyName.ToString());
 	    }
 	}
@@ -867,7 +867,7 @@ void FProcessorDependencySolver::ResolveDependencies(TArray<FProcessorDependency
 		}
 		else
 		{
-			UE_LOG(LogPipe, Error, TEXT("%s Unable to dump dependency graph into filename %s"), ANSI_TO_TCHAR(__FUNCTION__), *DependencyGraphFileName);
+			UE_LOG(LogMass, Error, TEXT("%s Unable to dump dependency graph into filename %s"), ANSI_TO_TCHAR(__FUNCTION__), *DependencyGraphFileName);
 		}
 	}
 
@@ -877,7 +877,7 @@ void FProcessorDependencySolver::ResolveDependencies(TArray<FProcessorDependency
 #if WITH_UNREAL_DEVELOPER_TOOLS
 	if (bAnyCyclesDetected || bAnyUnresolvedDependencies)
 	{
-		FMessageLog EditorErrors("LogPipe");
+		FMessageLog EditorErrors("LogMass");
 		if(bAnyCyclesDetected)
 		{
 		EditorErrors.Error(LOCTEXT("ProcessorDependenciesCycle", "Processor dependencies cycle found!"));
@@ -892,10 +892,10 @@ void FProcessorDependencySolver::ResolveDependencies(TArray<FProcessorDependency
 	}
 #endif // WITH_UNREAL_DEVELOPER_TOOLS
 
-	UE_LOG(LogPipe, Log, TEXT("Dependency order:"));
+	UE_LOG(LogMass, Log, TEXT("Dependency order:"));
 	for (const FProcessorDependencySolver::FOrderInfo& Info : OutResult)
 	{
-		UE_LOG(LogPipe, Log, TEXT("\t%s"), *Info.Name.ToString());
+		UE_LOG(LogMass, Log, TEXT("\t%s"), *Info.Name.ToString());
 	}
 }
 
diff --git a/Engine/Plugins/Runtime/MassEntity/Source/MassEntity/Public/IMassEntityModule.h b/Engine/Plugins/Runtime/MassEntity/Source/MassEntity/Public/IMassEntityModule.h
--- a/Engine/Plugins/Runtime/MassEntity/Source/MassEntity/Public/IMassEntityModule.h
+++ b/Engine/Plugins/Runtime/MassEntity/Source/MassEntity/Public/IMassEntityModule.h
@@ -11,7 +11,7 @@
  * The public interface to this module.  In most cases, this interface is only public to sibling modules 
  * within this plugin.
  */
-class IPipeModule : public IModuleInterface
+class IMassEntityModule : public IModuleInterface
 {
 
 public:
@@ -22,9 +22,9 @@ public:
 	 *
 	 * @return Returns singleton instance, loading the module on demand if needed
 	 */
-	static inline IPipeModule& Get()
+	static inline IMassEntityModule& Get()
 	{
-		return FModuleManager::LoadModuleChecked<IPipeModule>( "Pipe" );
+		return FModuleManager::LoadModuleChecked<IMassEntityModule>( "MassEntity" );
 	}
 
 	/**
@@ -34,7 +34,7 @@ public:
 	 */
 	static inline bool IsAvailable()
 	{
-		return FModuleManager::Get().IsModuleLoaded( "Pipe" );
+		return FModuleManager::Get().IsModuleLoaded( "MassEntity" );
 	}
 };
 
diff --git a/Engine/Plugins/Runtime/MassEntity/Source/MassEntity/Public/MassEntitySettings.h b/Engine/Plugins/Runtime/MassEntity/Source/MassEntity/Public/MassEntitySettings.h
--- a/Engine/Plugins/Runtime/MassEntity/Source/MassEntity/Public/MassEntitySettings.h
+++ b/Engine/Plugins/Runtime/MassEntity/Source/MassEntity/Public/MassEntitySettings.h
@@ -46,7 +46,7 @@ struct FMassProcessingPhaseConfig
 /**
  * Implements the settings for MassSimulation
  */
-UCLASS(config = Game, defaultconfig, DisplayName = "Pipe")
+UCLASS(config = Game, defaultconfig, DisplayName = "Mass")
 class MASSENTITY_API UMassSettings : public UDeveloperSettings
 {
 	GENERATED_BODY()
diff --git a/Engine/Plugins/Runtime/MassEntity/Source/MassEntity/Public/MassEntitySubsystem.h b/Engine/Plugins/Runtime/MassEntity/Source/MassEntity/Public/MassEntitySubsystem.h
--- a/Engine/Plugins/Runtime/MassEntity/Source/MassEntity/Public/MassEntitySubsystem.h
+++ b/Engine/Plugins/Runtime/MassEntity/Source/MassEntity/Public/MassEntitySubsystem.h
@@ -9,7 +9,6 @@
 #include "MassEntityQuery.h"
 #include "MassEntitySubsystem.generated.h"
 
-MASSENTITY_API DECLARE_LOG_CATEGORY_EXTERN(LogAggregateTicking, Warning, All);
 
 class UMassEntitySubsystem;
 struct FMassEntityQuery;
diff --git a/Engine/Plugins/Runtime/MassEntity/Source/MassEntity/Public/MassEntityTypes.h b/Engine/Plugins/Runtime/MassEntity/Source/MassEntity/Public/MassEntityTypes.h
--- a/Engine/Plugins/Runtime/MassEntity/Source/MassEntity/Public/MassEntityTypes.h
+++ b/Engine/Plugins/Runtime/MassEntity/Source/MassEntity/Public/MassEntityTypes.h
@@ -5,6 +5,9 @@
 #include "ScriptStructTypeBitSet.h"
 #include "MassEntityTypes.generated.h"
 
+
+MASSENTITY_API DECLARE_LOG_CATEGORY_EXTERN(LogMass, Warning, All);
+
 // This is the base class for all lightweight components
 USTRUCT()
 struct FMassFragment
@@ -41,10 +44,10 @@ struct FMassEntityHandle
 
 	FMassEntityHandle() {}
 	
-	UPROPERTY(VisibleAnywhere, Category = "AggregateTicking|Debug", Transient)
+	UPROPERTY(VisibleAnywhere, Category = "Mass|Debug", Transient)
 	int32 Index = 0;
 	
-	UPROPERTY(VisibleAnywhere, Category = "AggregateTicking|Debug", Transient)
+	UPROPERTY(VisibleAnywhere, Category = "Mass|Debug", Transient)
 	int32 SerialNumber = 0;
 
 	bool operator==(const FMassEntityHandle Other) const
diff --git a/Engine/Plugins/Runtime/MassEntity/Source/MassEntity/Public/MassExecutor.h b/Engine/Plugins/Runtime/MassEntity/Source/MassEntity/Public/MassExecutor.h
--- a/Engine/Plugins/Runtime/MassEntity/Source/MassEntity/Public/MassExecutor.h
+++ b/Engine/Plugins/Runtime/MassEntity/Source/MassEntity/Public/MassExecutor.h
@@ -17,26 +17,26 @@ class UMassProcessor;
 namespace UE::Mass::Executor
 {
 	/** Executes processors in a given RuntimePipeline */
-	MASSENTITY_API void Run(FMassRuntimePipeline& RuntimePipeline, FMassProcessingContext& PipeContext);
+	MASSENTITY_API void Run(FMassRuntimePipeline& RuntimePipeline, FMassProcessingContext& ProcessingContext);
 
 	/** Executes given Processor. Used mainly for triggering calculations via MassCompositeProcessors, e.g processing phases */
-	MASSENTITY_API void Run(UMassProcessor& Processor, FMassProcessingContext& PipeContext);
+	MASSENTITY_API void Run(UMassProcessor& Processor, FMassProcessingContext& ProcessingContext);
 
-	/** Similar to the Run function, but instead of using all the entities hosted by PipeContext.EntitySubsystem 
+	/** Similar to the Run function, but instead of using all the entities hosted by ProcessingContext.EntitySubsystem 
 	 *  it is processing only the entities given by EntityID via the Entities input parameter. 
 	 *  Note that all the entities need to be of Archetype archetype. 
 	 *  Under the hood the function converts Archetype-Entities pair to FArchetypeChunkCollection and calls the other flavor of RunSparse
 	 */
-	MASSENTITY_API void RunSparse(FMassRuntimePipeline& RuntimePipeline, FMassProcessingContext& PipeContext, FArchetypeHandle Archetype, TConstArrayView<FMassEntityHandle> Entities);
+	MASSENTITY_API void RunSparse(FMassRuntimePipeline& RuntimePipeline, FMassProcessingContext& ProcessingContext, FArchetypeHandle Archetype, TConstArrayView<FMassEntityHandle> Entities);
 
-	/** Similar to the Run function, but instead of using all the entities hosted by PipeContext.EntitySubsystem 
+	/** Similar to the Run function, but instead of using all the entities hosted by ProcessingContext.EntitySubsystem 
 	 *  it is processing only the entities given by SparseEntities input parameter.
 	 *  @todo rename
 	 */
-	MASSENTITY_API void RunSparse(FMassRuntimePipeline& RuntimePipeline, FMassProcessingContext& PipeContext, const FArchetypeChunkCollection& ChunkCollection);
+	MASSENTITY_API void RunSparse(FMassRuntimePipeline& RuntimePipeline, FMassProcessingContext& ProcessingContext, const FArchetypeChunkCollection& ChunkCollection);
 
 	/** Executes given Processors array view. This function gets called under the hood by the rest of Run* functions */
-	MASSENTITY_API void RunProcessorsView(TArrayView<UMassProcessor*> Processors, FMassProcessingContext& PipeContext, const FArchetypeChunkCollection* ChunkCollection = nullptr);
+	MASSENTITY_API void RunProcessorsView(TArrayView<UMassProcessor*> Processors, FMassProcessingContext& ProcessingContext, const FArchetypeChunkCollection* ChunkCollection = nullptr);
 
 	/** 
 	 *  Triggers tasks executing Processor (and potentially it's children) and returns the task graph event representing 
@@ -44,5 +44,5 @@ namespace UE::Mass::Executor
 	 *  @param OnDoneNotification will be called after all the processors are done, just after flushing the command buffer.
 	 *    Note that OnDoneNotification will be executed on GameThread.
 	 */
-	MASSENTITY_API FGraphEventRef TriggerParallelTasks(UMassProcessor& Processor, FMassProcessingContext& PipeContext, TFunction<void()> OnDoneNotification);
+	MASSENTITY_API FGraphEventRef TriggerParallelTasks(UMassProcessor& Processor, FMassProcessingContext& ProcessingContext, TFunction<void()> OnDoneNotification);
 };
diff --git a/Engine/Plugins/Runtime/MassEntity/Source/MassEntity/Public/MassProcessingTypes.h b/Engine/Plugins/Runtime/MassEntity/Source/MassEntity/Public/MassProcessingTypes.h
--- a/Engine/Plugins/Runtime/MassEntity/Source/MassEntity/Public/MassProcessingTypes.h
+++ b/Engine/Plugins/Runtime/MassEntity/Source/MassEntity/Public/MassProcessingTypes.h
@@ -1,4 +1,4 @@
-﻿// Copyright Epic Games, Inc. All Rights Reserved.
+// Copyright Epic Games, Inc. All Rights Reserved.
 
 #pragma once
 
@@ -8,8 +8,6 @@
 
 #define WITH_MASSENTITY_DEBUG (!(UE_BUILD_SHIPPING || UE_BUILD_SHIPPING_WITH_EDITOR || UE_BUILD_TEST) && WITH_STRUCTUTILS_DEBUG && 1)
 
-MASSENTITY_API DECLARE_LOG_CATEGORY_EXTERN(LogPipe, Warning, All);
-
 class UMassEntitySubsystem;
 class UMassProcessor;
 class UMassSchematic;
diff --git a/Engine/Plugins/Runtime/MassEntity/Source/MassEntityEditor/Private/AssetTypeActions_MassSchematic.cpp b/Engine/Plugins/Runtime/MassEntity/Source/MassEntityEditor/Private/AssetTypeActions_MassSchematic.cpp
--- a/Engine/Plugins/Runtime/MassEntity/Source/MassEntityEditor/Private/AssetTypeActions_MassSchematic.cpp
+++ b/Engine/Plugins/Runtime/MassEntity/Source/MassEntityEditor/Private/AssetTypeActions_MassSchematic.cpp
@@ -16,8 +16,8 @@ void FAssetTypeActions_MassSchematic::OpenAssetEditor(const TArray<UObject*>& In
 	{
 		if (UMassSchematic* Pipe = Cast<UMassSchematic>(*ObjIt))
 		{
-			FPipeEditorModule& EditorModule = FModuleManager::LoadModuleChecked<FPipeEditorModule>("PipeEditor");
-			TSharedRef<IPipeEditor> NewEditor = EditorModule.CreatePipeEditor(Mode, EditWithinLevelEditor, Pipe);
+			FMassEntityEditorModule& EditorModule = FModuleManager::LoadModuleChecked<FMassEntityEditorModule>("MassEntityEditor");
+			TSharedRef<IMassEntityEditor> NewEditor = EditorModule.CreateMassEntityEditor(Mode, EditWithinLevelEditor, Pipe);
 		}
 	}
 }
diff --git a/Engine/Plugins/Runtime/MassEntity/Source/MassEntityEditor/Private/MassEditorStyle.cpp b/Engine/Plugins/Runtime/MassEntity/Source/MassEntityEditor/Private/MassEditorStyle.cpp
--- a/Engine/Plugins/Runtime/MassEntity/Source/MassEntityEditor/Private/MassEditorStyle.cpp
+++ b/Engine/Plugins/Runtime/MassEntity/Source/MassEntityEditor/Private/MassEditorStyle.cpp
@@ -14,15 +14,15 @@
 
 #define DEFAULT_FONT(...) FCoreStyle::GetDefaultFontStyle(__VA_ARGS__)
 
-TSharedPtr<FSlateStyleSet> FPipeEditorStyle::StyleSet = nullptr;
+TSharedPtr<FSlateStyleSet> FMassEntityEditorStyle::StyleSet = nullptr;
 
-FString FPipeEditorStyle::InContent(const FString& RelativePath, const ANSICHAR* Extension)
+FString FMassEntityEditorStyle::InContent(const FString& RelativePath, const ANSICHAR* Extension)
 {
-	static FString ContentDir = IPluginManager::Get().FindPlugin(TEXT("PipeEditorModule"))->GetContentDir() / TEXT("Slate");
+	static FString ContentDir = IPluginManager::Get().FindPlugin(TEXT("MassEntityEditorModule"))->GetContentDir() / TEXT("Slate");
 	return (ContentDir / RelativePath) + Extension;
 }
 
-void FPipeEditorStyle::Initialize()
+void FMassEntityEditorStyle::Initialize()
 {
 	if (StyleSet.IsValid())
 	{
@@ -109,7 +109,7 @@ void FPipeEditorStyle::Initialize()
 }
 
 
-void FPipeEditorStyle::Shutdown()
+void FMassEntityEditorStyle::Shutdown()
 {
 	if (StyleSet.IsValid())
 	{
@@ -120,8 +120,8 @@ void FPipeEditorStyle::Shutdown()
 }
 
 
-FName FPipeEditorStyle::GetStyleSetName()
+FName FMassEntityEditorStyle::GetStyleSetName()
 {
-	static FName StyleName("PipeEditorStyle");
+	static FName StyleName("MassEntityEditorStyle");
 	return StyleName;
 }
diff --git a/Engine/Plugins/Runtime/MassEntity/Source/MassEntityEditor/Private/MassEntityEditor.cpp b/Engine/Plugins/Runtime/MassEntity/Source/MassEntityEditor/Private/MassEntityEditor.cpp
--- a/Engine/Plugins/Runtime/MassEntity/Source/MassEntityEditor/Private/MassEntityEditor.cpp
+++ b/Engine/Plugins/Runtime/MassEntity/Source/MassEntityEditor/Private/MassEntityEditor.cpp
@@ -15,12 +15,12 @@
 #include "Widgets/Docking/SDockTab.h"
  
 
-#define LOCTEXT_NAMESPACE "PipeEditor"
+#define LOCTEXT_NAMESPACE "MassEntityEditor"
 
-const FName PipeEditorAppName(TEXT("PipeEditorApp"));
-const FName FPipeEditor::AssetDetailsTabId(TEXT("PipeEditor_AssetDetails"));
+const FName MassEntityEditorAppName(TEXT("MassEntityEditorApp"));
+const FName FMassEntityEditor::AssetDetailsTabId(TEXT("MassEntityEditor_AssetDetails"));
 
-void FPipeEditor::AddReferencedObjects(FReferenceCollector& Collector)
+void FMassEntityEditor::AddReferencedObjects(FReferenceCollector& Collector)
 {
 	if (MassSchematic != nullptr)
 	{
@@ -28,26 +28,26 @@ void FPipeEditor::AddReferencedObjects(FReferenceCollector& Collector)
 	}
 }
 
-void FPipeEditor::RegisterTabSpawners(const TSharedRef<class FTabManager>& InTabManager)
+void FMassEntityEditor::RegisterTabSpawners(const TSharedRef<class FTabManager>& InTabManager)
 {
-	WorkspaceMenuCategory = InTabManager->AddLocalWorkspaceMenuCategory(LOCTEXT("WorkspaceMenu_PipeEditor", "Pipe Editor"));
+	WorkspaceMenuCategory = InTabManager->AddLocalWorkspaceMenuCategory(LOCTEXT("WorkspaceMenu_MassEntityEditor", "Pipe Editor"));
 	TSharedRef<FWorkspaceItem> WorkspaceMenuCategoryRef = WorkspaceMenuCategory.ToSharedRef();
 
 	FAssetEditorToolkit::RegisterTabSpawners(InTabManager);
 
-	InTabManager->RegisterTabSpawner(AssetDetailsTabId, FOnSpawnTab::CreateSP(this, &FPipeEditor::SpawnTab_AssetDetails))
-		.SetDisplayName(NSLOCTEXT("PipeEditor", "AssetDetailsTab", "Asset Details"))
+	InTabManager->RegisterTabSpawner(AssetDetailsTabId, FOnSpawnTab::CreateSP(this, &FMassEntityEditor::SpawnTab_AssetDetails))
+		.SetDisplayName(NSLOCTEXT("MassEntityEditor", "AssetDetailsTab", "Asset Details"))
 		.SetGroup(WorkspaceMenuCategoryRef)
 		.SetIcon(FSlateIcon(FEditorStyle::GetStyleSetName(), "LevelEditor.Tabs.Details"));
 }
 
-void FPipeEditor::UnregisterTabSpawners(const TSharedRef<class FTabManager>& InTabManager)
+void FMassEntityEditor::UnregisterTabSpawners(const TSharedRef<class FTabManager>& InTabManager)
 {
 	FAssetEditorToolkit::UnregisterTabSpawners(InTabManager);
 	InTabManager->UnregisterTabSpawner(AssetDetailsTabId);
 }
 
-void FPipeEditor::InitEditor(const EToolkitMode::Type Mode, const TSharedPtr< class IToolkitHost >& InitToolkitHost, UMassSchematic& InMassSchematic)
+void FMassEntityEditor::InitEditor(const EToolkitMode::Type Mode, const TSharedPtr< class IToolkitHost >& InitToolkitHost, UMassSchematic& InMassSchematic)
 {
 	MassSchematic = &InMassSchematic;
 
@@ -70,30 +70,30 @@ void FPipeEditor::InitEditor(const EToolkitMode::Type Mode, const TSharedPtr< cl
 
 	const bool bCreateDefaultStandaloneMenu = true;
 	const bool bCreateDefaultToolbar = true;
-	FAssetEditorToolkit::InitAssetEditor(Mode, InitToolkitHost, PipeEditorAppName, StandaloneDefaultLayout, bCreateDefaultStandaloneMenu, bCreateDefaultToolbar, MassSchematic);
+	FAssetEditorToolkit::InitAssetEditor(Mode, InitToolkitHost, MassEntityEditorAppName, StandaloneDefaultLayout, bCreateDefaultStandaloneMenu, bCreateDefaultToolbar, MassSchematic);
 }
 
-FName FPipeEditor::GetToolkitFName() const
+FName FMassEntityEditor::GetToolkitFName() const
 {
-	return FName("PipeEditor");
+	return FName("MassEntityEditor");
 }
 
-FText FPipeEditor::GetBaseToolkitName() const
+FText FMassEntityEditor::GetBaseToolkitName() const
 {
-	return NSLOCTEXT("PipeEditor", "AppLabel", "Pipe");
+	return NSLOCTEXT("MassEntityEditor", "AppLabel", "Mass");
 }
 
-FString FPipeEditor::GetWorldCentricTabPrefix() const
+FString FMassEntityEditor::GetWorldCentricTabPrefix() const
 {
-	return NSLOCTEXT("PipeEditor", "WorldCentricTabPrefix", "Pipe").ToString();
+	return NSLOCTEXT("MassEntityEditor", "WorldCentricTabPrefix", "Mass").ToString();
 }
 
-FLinearColor FPipeEditor::GetWorldCentricTabColorScale() const
+FLinearColor FMassEntityEditor::GetWorldCentricTabColorScale() const
 {
 	return FLinearColor( 0.0f, 0.0f, 0.2f, 0.5f );
 }
 
-TSharedRef<SDockTab> FPipeEditor::SpawnTab_AssetDetails(const FSpawnTabArgs& Args)
+TSharedRef<SDockTab> FMassEntityEditor::SpawnTab_AssetDetails(const FSpawnTabArgs& Args)
 {
 	check(Args.GetTabId() == AssetDetailsTabId);
 
@@ -103,10 +103,10 @@ TSharedRef<SDockTab> FPipeEditor::SpawnTab_AssetDetails(const FSpawnTabArgs& Arg
 
 	AssetDetailsView = PropertyEditorModule.CreateDetailView(DetailsViewArgs);
 	AssetDetailsView->SetObject(MassSchematic);
-	AssetDetailsView->OnFinishedChangingProperties().AddSP(this, &FPipeEditor::OnAssetFinishedChangingProperties);
+	AssetDetailsView->OnFinishedChangingProperties().AddSP(this, &FMassEntityEditor::OnAssetFinishedChangingProperties);
 
 	TSharedRef<SDockTab> SpawnedTab = SNew(SDockTab)
-		.Label(NSLOCTEXT("PipeEditor", "AssetDetailsTab", "Pipe"))
+		.Label(NSLOCTEXT("MassEntityEditor", "AssetDetailsTab", "Mass"))
 		[
 			AssetDetailsView.ToSharedRef()
 		];
@@ -114,7 +114,7 @@ TSharedRef<SDockTab> FPipeEditor::SpawnTab_AssetDetails(const FSpawnTabArgs& Arg
 	return SpawnedTab;
 }
 
-void FPipeEditor::SaveAsset_Execute()
+void FMassEntityEditor::SaveAsset_Execute()
 {
 	// @todo pre-save validation
 
@@ -122,10 +122,10 @@ void FPipeEditor::SaveAsset_Execute()
 	FAssetEditorToolkit::SaveAsset_Execute();
 }
 
-void FPipeEditor::OnAssetFinishedChangingProperties(const FPropertyChangedEvent& PropertyChangedEvent)
+void FMassEntityEditor::OnAssetFinishedChangingProperties(const FPropertyChangedEvent& PropertyChangedEvent)
 {
-	FPipeEditorModule& PipeEditorModule = FModuleManager::GetModuleChecked<FPipeEditorModule>("PipeEditor");
-	PipeEditorModule.GetOnAssetPropertiesChanged().Broadcast(MassSchematic, PropertyChangedEvent);
+	FMassEntityEditorModule& MassEntityEditorModule = FModuleManager::GetModuleChecked<FMassEntityEditorModule>("MassEntityEditor");
+	MassEntityEditorModule.GetOnAssetPropertiesChanged().Broadcast(MassSchematic, PropertyChangedEvent);
 }
 
 
diff --git a/Engine/Plugins/Runtime/MassEntity/Source/MassEntityEditor/Private/MassEntityEditor.h b/Engine/Plugins/Runtime/MassEntity/Source/MassEntityEditor/Private/MassEntityEditor.h
--- a/Engine/Plugins/Runtime/MassEntity/Source/MassEntityEditor/Private/MassEntityEditor.h
+++ b/Engine/Plugins/Runtime/MassEntity/Source/MassEntityEditor/Private/MassEntityEditor.h
@@ -11,7 +11,7 @@
 class IDetailsView;
 class UMassSchematic;
 
-class FPipeEditor : public IPipeEditor, public FGCObject//, public FSelfRegisteringEditorUndoClient
+class FMassEntityEditor : public IMassEntityEditor, public FGCObject//, public FSelfRegisteringEditorUndoClient
 {
 public:
 	virtual void RegisterTabSpawners(const TSharedRef<class FTabManager>& TabManager) override;
@@ -36,7 +36,7 @@ public:
 	virtual void AddReferencedObjects(FReferenceCollector& Collector) override;
 	virtual FString GetReferencerName() const override
 	{
-		return TEXT("FPipeEditor");
+		return TEXT("FMassEntityEditor");
 	}
 
 protected:
diff --git a/Engine/Plugins/Runtime/MassEntity/Source/MassEntityEditor/Private/MassEntityEditorModule.cpp b/Engine/Plugins/Runtime/MassEntity/Source/MassEntityEditor/Private/MassEntityEditorModule.cpp
--- a/Engine/Plugins/Runtime/MassEntity/Source/MassEntityEditor/Private/MassEntityEditorModule.cpp
+++ b/Engine/Plugins/Runtime/MassEntity/Source/MassEntityEditor/Private/MassEntityEditorModule.cpp
@@ -12,16 +12,16 @@
 #include "MassProcessor.h"
 
 
-#define LOCTEXT_NAMESPACE "PipeEditor"
+#define LOCTEXT_NAMESPACE "MassEntityEditor"
 
-IMPLEMENT_MODULE(FPipeEditorModule, PipeEditor)
+IMPLEMENT_MODULE(FMassEntityEditorModule, MassEntityEditor)
 
-void FPipeEditorModule::StartupModule()
+void FMassEntityEditorModule::StartupModule()
 {
 	MenuExtensibilityManager = MakeShareable(new FExtensibilityManager);
 	ToolBarExtensibilityManager = MakeShareable(new FExtensibilityManager);
 
-	FPipeEditorStyle::Initialize();
+	FMassEntityEditorStyle::Initialize();
 
 	// Register asset types
 	IAssetTools& AssetTools = FModuleManager::LoadModuleChecked<FAssetToolsModule>("AssetTools").Get();
@@ -36,13 +36,13 @@ void FPipeEditorModule::StartupModule()
 	PropertyModule.NotifyCustomizationModuleChanged();
 }
 
-void FPipeEditorModule::ShutdownModule()
+void FMassEntityEditorModule::ShutdownModule()
 {
 	ProcessorClassCache.Reset();
 	MenuExtensibilityManager.Reset();
 	ToolBarExtensibilityManager.Reset();
 
-	FPipeEditorStyle::Shutdown();
+	FMassEntityEditorStyle::Shutdown();
 
 	// Unregister the data asset type actions
 	if (FModuleManager::Get().IsModuleLoaded("AssetTools"))
@@ -67,7 +67,7 @@ void FPipeEditorModule::ShutdownModule()
 	}
 }
 
-TSharedRef<IPipeEditor> FPipeEditorModule::CreatePipeEditor(const EToolkitMode::Type Mode, const TSharedPtr< class IToolkitHost >& InitToolkitHost, UMassSchematic* MassSchematic)
+TSharedRef<IMassEntityEditor> FMassEntityEditorModule::CreateMassEntityEditor(const EToolkitMode::Type Mode, const TSharedPtr< class IToolkitHost >& InitToolkitHost, UMassSchematic* MassSchematic)
 {
 	if (!ProcessorClassCache.IsValid())
 	{
@@ -75,7 +75,7 @@ TSharedRef<IPipeEditor> FPipeEditorModule::CreatePipeEditor(const EToolkitMode::
 		ProcessorClassCache->UpdateAvailableBlueprintClasses();
 	}
 
-	TSharedRef<FPipeEditor> NewEditor(new FPipeEditor());
+	TSharedRef<FMassEntityEditor> NewEditor(new FMassEntityEditor());
 	if (ensure(MassSchematic))
 	{
 		NewEditor->InitEditor(Mode, InitToolkitHost, *MassSchematic);
diff --git a/Engine/Plugins/Runtime/MassEntity/Source/MassEntityEditor/Public/IMassEntityEditor.h b/Engine/Plugins/Runtime/MassEntity/Source/MassEntityEditor/Public/IMassEntityEditor.h
--- a/Engine/Plugins/Runtime/MassEntity/Source/MassEntityEditor/Public/IMassEntityEditor.h
+++ b/Engine/Plugins/Runtime/MassEntity/Source/MassEntityEditor/Public/IMassEntityEditor.h
@@ -7,7 +7,7 @@
 
 
 /** Pipe Editor public interface */
-class MASSENTITYEDITOR_API IPipeEditor : public FAssetEditorToolkit
+class MASSENTITYEDITOR_API IMassEntityEditor : public FAssetEditorToolkit
 {
 public:
 };
diff --git a/Engine/Plugins/Runtime/MassEntity/Source/MassEntityEditor/Public/MassEditorStyle.h b/Engine/Plugins/Runtime/MassEntity/Source/MassEntityEditor/Public/MassEditorStyle.h
--- a/Engine/Plugins/Runtime/MassEntity/Source/MassEntityEditor/Public/MassEditorStyle.h
+++ b/Engine/Plugins/Runtime/MassEntity/Source/MassEntityEditor/Public/MassEditorStyle.h
@@ -4,7 +4,7 @@
 
 #include "Styling/SlateStyle.h"
 
-class MASSENTITYEDITOR_API FPipeEditorStyle
+class MASSENTITYEDITOR_API FMassEntityEditorStyle
 {
 public:
 	static void Initialize();
diff --git a/Engine/Plugins/Runtime/MassEntity/Source/MassEntityEditor/Public/MassEntityEditorModule.h b/Engine/Plugins/Runtime/MassEntity/Source/MassEntityEditor/Public/MassEntityEditorModule.h
--- a/Engine/Plugins/Runtime/MassEntity/Source/MassEntityEditor/Public/MassEntityEditorModule.h
+++ b/Engine/Plugins/Runtime/MassEntity/Source/MassEntityEditor/Public/MassEntityEditorModule.h
@@ -9,7 +9,7 @@
 #include "Toolkits/AssetEditorToolkit.h"
 
 
-class IPipeEditor;
+class IMassEntityEditor;
 class FAssetTypeActions_Base;
 struct FGraphPanelNodeFactory;
 class UMassSchematic;
@@ -17,7 +17,7 @@ class UMassSchematic;
 /**
 * The public interface to this module
 */
-class MASSENTITYEDITOR_API FPipeEditorModule : public IModuleInterface, public IHasMenuExtensibility, public IHasToolBarExtensibility
+class MASSENTITYEDITOR_API FMassEntityEditorModule : public IModuleInterface, public IHasMenuExtensibility, public IHasToolBarExtensibility
 {
 public:
 	DECLARE_MULTICAST_DELEGATE_TwoParams(FOnAssetPropertiesChanged, class UMassSchematic* /*MassSchematic*/, const FPropertyChangedEvent& /*PropertyChangedEvent*/);
@@ -26,7 +26,7 @@ public:
 	virtual void ShutdownModule() override;
 
 	/** Creates an instance of Pipe editor. Only virtual so that it can be called across the DLL boundary. */
-	virtual TSharedRef<IPipeEditor> CreatePipeEditor(const EToolkitMode::Type Mode, const TSharedPtr< class IToolkitHost >& InitToolkitHost, UMassSchematic* MassSchematic);
+	virtual TSharedRef<IMassEntityEditor> CreateMassEntityEditor(const EToolkitMode::Type Mode, const TSharedPtr< class IToolkitHost >& InitToolkitHost, UMassSchematic* MassSchematic);
 
 	virtual TSharedPtr<FExtensibilityManager> GetMenuExtensibilityManager() override { return MenuExtensibilityManager; }
 	virtual TSharedPtr<FExtensibilityManager> GetToolBarExtensibilityManager() override { return ToolBarExtensibilityManager; }
diff --git a/Engine/Plugins/Runtime/MassEntity/Source/MassEntityTestSuite/Private/MassCompositeProcessorTest.cpp b/Engine/Plugins/Runtime/MassEntity/Source/MassEntityTestSuite/Private/MassCompositeProcessorTest.cpp
--- a/Engine/Plugins/Runtime/MassEntity/Source/MassEntityTestSuite/Private/MassCompositeProcessorTest.cpp
+++ b/Engine/Plugins/Runtime/MassEntity/Source/MassEntityTestSuite/Private/MassCompositeProcessorTest.cpp
@@ -9,7 +9,7 @@
 #include "MassEntityTestTypes.h"
 #include "MassExecutor.h"
 
-#define LOCTEXT_NAMESPACE "PipeTest"
+#define LOCTEXT_NAMESPACE "MassTest"
 
 PRAGMA_DISABLE_OPTIMIZATION
 
diff --git a/Engine/Plugins/Runtime/MassEntity/Source/MassEntityTestSuite/Private/MassDependencyTest.cpp b/Engine/Plugins/Runtime/MassEntity/Source/MassEntityTestSuite/Private/MassDependencyTest.cpp
--- a/Engine/Plugins/Runtime/MassEntity/Source/MassEntityTestSuite/Private/MassDependencyTest.cpp
+++ b/Engine/Plugins/Runtime/MassEntity/Source/MassEntityTestSuite/Private/MassDependencyTest.cpp
@@ -6,7 +6,7 @@
 #include "MassProcessorDependencySolver.h"
 #include "MassEntityTestTypes.h"
 
-#define LOCTEXT_NAMESPACE "PipeTest"
+#define LOCTEXT_NAMESPACE "MassTest"
 
 PRAGMA_DISABLE_OPTIMIZATION
 
@@ -68,7 +68,7 @@ struct FTrivialDependency : FDependencySolverBase
 		return true;
 	}
 };
-IMPLEMENT_AI_INSTANT_TEST(FTrivialDependency, "System.Pipe.Dependencies.Trivial");
+IMPLEMENT_AI_INSTANT_TEST(FTrivialDependency, "System.Mass.Dependencies.Trivial");
 
 
 struct FDeepGroup : FDependencySolverBase
@@ -108,7 +108,7 @@ struct FDeepGroup : FDependencySolverBase
 		return true;
 	}
 };
-IMPLEMENT_AI_INSTANT_TEST(FDeepGroup, "System.Pipe.Dependencies.DeepGroup");
+IMPLEMENT_AI_INSTANT_TEST(FDeepGroup, "System.Mass.Dependencies.DeepGroup");
 
 struct FComplexScenario : FDependencySolverBase
 {
@@ -170,7 +170,7 @@ struct FComplexScenario : FDependencySolverBase
 		return true;
 	}
 };
-IMPLEMENT_AI_INSTANT_TEST(FComplexScenario, "System.Pipe.Dependencies.Complex");
+IMPLEMENT_AI_INSTANT_TEST(FComplexScenario, "System.Mass.Dependencies.Complex");
 
 
 struct FBadInput_Empty : FDependencySolverBase
@@ -182,7 +182,7 @@ struct FBadInput_Empty : FDependencySolverBase
 		return true;
 	}
 };
-IMPLEMENT_AI_INSTANT_TEST(FBadInput_Empty, "System.Pipe.Dependencies.BadInput.Empty");
+IMPLEMENT_AI_INSTANT_TEST(FBadInput_Empty, "System.Mass.Dependencies.BadInput.Empty");
 
 struct FBadInput_Null : FDependencySolverBase
 {
@@ -196,7 +196,7 @@ struct FBadInput_Null : FDependencySolverBase
 		return true;
 	}
 };
-IMPLEMENT_AI_INSTANT_TEST(FBadInput_Null, "System.Pipe.Dependencies.BadInput.SingleNull");
+IMPLEMENT_AI_INSTANT_TEST(FBadInput_Null, "System.Mass.Dependencies.BadInput.SingleNull");
 
 struct FBadInput_MultipleNulls : FDependencySolverBase
 {
@@ -212,7 +212,7 @@ struct FBadInput_MultipleNulls : FDependencySolverBase
 		return true;
 	}
 };
-IMPLEMENT_AI_INSTANT_TEST(FBadInput_MultipleNulls, "System.Pipe.Dependencies.BadInput.MultipleNulls");
+IMPLEMENT_AI_INSTANT_TEST(FBadInput_MultipleNulls, "System.Mass.Dependencies.BadInput.MultipleNulls");
 
 struct FBadInput_MixedNulls : FDependencySolverBase
 {
@@ -229,7 +229,7 @@ struct FBadInput_MixedNulls : FDependencySolverBase
 		return true;
 	}
 };
-IMPLEMENT_AI_INSTANT_TEST(FBadInput_MixedNulls, "System.Pipe.Dependencies.BadInput.NullsMixedIn");
+IMPLEMENT_AI_INSTANT_TEST(FBadInput_MixedNulls, "System.Mass.Dependencies.BadInput.NullsMixedIn");
 
 struct FBadInput_Duplicates : FDependencySolverBase
 {
@@ -245,7 +245,7 @@ struct FBadInput_Duplicates : FDependencySolverBase
 		return true;
 	}
 };
-IMPLEMENT_AI_INSTANT_TEST(FBadInput_Duplicates, "System.Pipe.Dependencies.BadInput.Duplicates");
+IMPLEMENT_AI_INSTANT_TEST(FBadInput_Duplicates, "System.Mass.Dependencies.BadInput.Duplicates");
 
 struct FGroupNamesGeneration : FAITestBase
 {
@@ -275,7 +275,7 @@ struct FGroupNamesGeneration : FAITestBase
 	}
 };
 
-IMPLEMENT_AI_INSTANT_TEST(FGroupNamesGeneration, "System.Pipe.Dependencies.SubgroupNames");
+IMPLEMENT_AI_INSTANT_TEST(FGroupNamesGeneration, "System.Mass.Dependencies.SubgroupNames");
 
 // tests to add:
 // * handling missing dependencies
diff --git a/Engine/Plugins/Runtime/MassEntity/Source/MassEntityTestSuite/Private/MassEntityTest.cpp b/Engine/Plugins/Runtime/MassEntity/Source/MassEntityTestSuite/Private/MassEntityTest.cpp
--- a/Engine/Plugins/Runtime/MassEntity/Source/MassEntityTestSuite/Private/MassEntityTest.cpp
+++ b/Engine/Plugins/Runtime/MassEntity/Source/MassEntityTestSuite/Private/MassEntityTest.cpp
@@ -9,7 +9,7 @@
 #include "MassEntityTestTypes.h"
 #include "MassExecutor.h"
 
-#define LOCTEXT_NAMESPACE "PipeTest"
+#define LOCTEXT_NAMESPACE "MassTest"
 
 PRAGMA_DISABLE_OPTIMIZATION
 
@@ -42,7 +42,7 @@ struct FEntityTest_ArchetypeCreation : FEntityTestBase
 		return true;
 	}
 };
-IMPLEMENT_AI_INSTANT_TEST(FEntityTest_ArchetypeCreation, "System.Pipe.Entity.AchetypesCreation");
+IMPLEMENT_AI_INSTANT_TEST(FEntityTest_ArchetypeCreation, "System.Mass.Entity.AchetypesCreation");
 
 
 struct FEntityTest_ArchetypeEquivalence : FEntityTestBase
@@ -59,7 +59,7 @@ struct FEntityTest_ArchetypeEquivalence : FEntityTestBase
 		return true;
 	}
 };
-IMPLEMENT_AI_INSTANT_TEST(FEntityTest_ArchetypeEquivalence, "System.Pipe.Entity.AchetypeEquivalance");
+IMPLEMENT_AI_INSTANT_TEST(FEntityTest_ArchetypeEquivalence, "System.Mass.Entity.AchetypeEquivalance");
 
 
 struct FEntityTest_MultipleEntitiesCreation : FEntityTestBase
@@ -87,7 +87,7 @@ struct FEntityTest_MultipleEntitiesCreation : FEntityTestBase
 		return true;
 	}
 };
-IMPLEMENT_AI_INSTANT_TEST(FEntityTest_MultipleEntitiesCreation, "System.Pipe.Entity.MultipleEntitiesCreation");
+IMPLEMENT_AI_INSTANT_TEST(FEntityTest_MultipleEntitiesCreation, "System.Mass.Entity.MultipleEntitiesCreation");
 
 
 struct FEntityTest_EntityBatchCreation : FEntityTestBase
@@ -104,7 +104,7 @@ struct FEntityTest_EntityBatchCreation : FEntityTestBase
 		return true;
 	}
 };
-IMPLEMENT_AI_INSTANT_TEST(FEntityTest_EntityBatchCreation, "System.Pipe.Entity.BatchCreation");
+IMPLEMENT_AI_INSTANT_TEST(FEntityTest_EntityBatchCreation, "System.Mass.Entity.BatchCreation");
 
 struct FEntityTest_BatchCreatingSingleEntity : FEntityTestBase
 {
@@ -119,7 +119,7 @@ struct FEntityTest_BatchCreatingSingleEntity : FEntityTestBase
 		return true;
 	}
 };
-IMPLEMENT_AI_INSTANT_TEST(FEntityTest_BatchCreatingSingleEntity, "System.Pipe.Entity.BatchCreatingSingleEntity");
+IMPLEMENT_AI_INSTANT_TEST(FEntityTest_BatchCreatingSingleEntity, "System.Mass.Entity.BatchCreatingSingleEntity");
 
 
 struct FEntityTest_EntityCreation : FEntityTestBase
@@ -136,7 +136,7 @@ struct FEntityTest_EntityCreation : FEntityTestBase
 		return true;
 	}
 };
-IMPLEMENT_AI_INSTANT_TEST(FEntityTest_EntityCreation, "System.Pipe.Entity.EntityCreation");
+IMPLEMENT_AI_INSTANT_TEST(FEntityTest_EntityCreation, "System.Mass.Entity.EntityCreation");
 
 
 struct FEntityTest_EntityCreationFromInstances : FEntityTestBase
@@ -153,7 +153,7 @@ struct FEntityTest_EntityCreationFromInstances : FEntityTestBase
 		return true;
 	}
 };
-IMPLEMENT_AI_INSTANT_TEST(FEntityTest_EntityCreationFromInstances, "System.Pipe.Entity.EntityCreationFromInstances");
+IMPLEMENT_AI_INSTANT_TEST(FEntityTest_EntityCreationFromInstances, "System.Mass.Entity.EntityCreationFromInstances");
 
 #if 0 // this test is compiled out since AddComponentToEntity will fails an ensure if a redundant fragment gets added
 struct FEntityTest_AddingRedundantFragment : FEntityTestBase
@@ -168,7 +168,7 @@ struct FEntityTest_AddingRedundantFragment : FEntityTestBase
 		return true;
 	}
 };
-IMPLEMENT_AI_INSTANT_TEST(FEntityTest_AddingRedundantFragment, "System.Pipe.Entity.AddingRedundantFragment");
+IMPLEMENT_AI_INSTANT_TEST(FEntityTest_AddingRedundantFragment, "System.Mass.Entity.AddingRedundantFragment");
 #endif //0
 
 
@@ -189,7 +189,7 @@ struct FEntityTest_AddingFragmentType : FEntityTestBase
 		return true;
 	}
 };
-IMPLEMENT_AI_INSTANT_TEST(FEntityTest_AddingFragmentType, "System.Pipe.Entity.AddingFragmentType");
+IMPLEMENT_AI_INSTANT_TEST(FEntityTest_AddingFragmentType, "System.Mass.Entity.AddingFragmentType");
 
 struct FEntityTest_AddingFragmentInstance : FEntityTestBase
 {
@@ -209,7 +209,7 @@ struct FEntityTest_AddingFragmentInstance : FEntityTestBase
 		return true;
 	}
 };
-IMPLEMENT_AI_INSTANT_TEST(FEntityTest_AddingFragmentInstance, "System.Pipe.Entity.AddingFragmentType");
+IMPLEMENT_AI_INSTANT_TEST(FEntityTest_AddingFragmentInstance, "System.Mass.Entity.AddingFragmentType");
 
 
 struct FEntityTest_RemovingFragment : FEntityTestBase
@@ -229,7 +229,7 @@ struct FEntityTest_RemovingFragment : FEntityTestBase
 		return true;
 	}
 };
-IMPLEMENT_AI_INSTANT_TEST(FEntityTest_RemovingFragment, "System.Pipe.Entity.RemovingFragment");
+IMPLEMENT_AI_INSTANT_TEST(FEntityTest_RemovingFragment, "System.Mass.Entity.RemovingFragment");
 
 struct FEntityTest_RemovingLastFragment : FEntityTestBase
 {
@@ -246,7 +246,7 @@ struct FEntityTest_RemovingLastFragment : FEntityTestBase
 		return true;
 	}
 };
-IMPLEMENT_AI_INSTANT_TEST(FEntityTest_RemovingLastFragment, "System.Pipe.Entity.RemovingLastFragment");
+IMPLEMENT_AI_INSTANT_TEST(FEntityTest_RemovingLastFragment, "System.Mass.Entity.RemovingLastFragment");
 
 struct FEntityTest_DestroyEntity : FEntityTestBase
 {
@@ -261,7 +261,7 @@ struct FEntityTest_DestroyEntity : FEntityTestBase
 		return true;
 	}
 };
-IMPLEMENT_AI_INSTANT_TEST(FEntityTest_DestroyEntity, "System.Pipe.Entity.DestroyEntity");
+IMPLEMENT_AI_INSTANT_TEST(FEntityTest_DestroyEntity, "System.Mass.Entity.DestroyEntity");
 
 struct FEntityTest_EntityReservationAndBuilding : FEntityTestBase
 {
@@ -284,7 +284,7 @@ struct FEntityTest_EntityReservationAndBuilding : FEntityTestBase
 		return true;
 	}
 };
-IMPLEMENT_AI_INSTANT_TEST(FEntityTest_EntityReservationAndBuilding, "System.Pipe.Entity.EntityReservationAndBuilding");
+IMPLEMENT_AI_INSTANT_TEST(FEntityTest_EntityReservationAndBuilding, "System.Mass.Entity.EntityReservationAndBuilding");
 
 struct FEntityTest_EntityReservationAndBuildingFromInstances : FEntityTestBase
 {
@@ -307,7 +307,7 @@ struct FEntityTest_EntityReservationAndBuildingFromInstances : FEntityTestBase
 		return true;
 	}
 };
-IMPLEMENT_AI_INSTANT_TEST(FEntityTest_EntityReservationAndBuildingFromInstances, "System.Pipe.Entity.EntityReservationAndBuildingFromInstances");
+IMPLEMENT_AI_INSTANT_TEST(FEntityTest_EntityReservationAndBuildingFromInstances, "System.Mass.Entity.EntityReservationAndBuildingFromInstances");
 
 struct FEntityTest_ReleaseEntity : FEntityTestBase
 {
@@ -324,7 +324,7 @@ struct FEntityTest_ReleaseEntity : FEntityTestBase
 		return true;
 	}
 };
-IMPLEMENT_AI_INSTANT_TEST(FEntityTest_ReleaseEntity, "System.Pipe.Entity.ReleaseEntity");
+IMPLEMENT_AI_INSTANT_TEST(FEntityTest_ReleaseEntity, "System.Mass.Entity.ReleaseEntity");
 
 struct FEntityTest_ReserveAPreviouslyBuiltEntity : FEntityTestBase
 {
@@ -353,7 +353,7 @@ struct FEntityTest_ReserveAPreviouslyBuiltEntity : FEntityTestBase
 		return true;
 	}
 };
-IMPLEMENT_AI_INSTANT_TEST(FEntityTest_ReserveAPreviouslyBuiltEntity, "System.Pipe.Entity.ReserveAPreviouslyBuiltEntity");
+IMPLEMENT_AI_INSTANT_TEST(FEntityTest_ReserveAPreviouslyBuiltEntity, "System.Mass.Entity.ReserveAPreviouslyBuiltEntity");
 
 #endif // WITH_MASSENTITY_DEBUG
 
diff --git a/Engine/Plugins/Runtime/MassEntity/Source/MassEntityTestSuite/Private/MassEntityTestSuiteModule.cpp b/Engine/Plugins/Runtime/MassEntity/Source/MassEntityTestSuite/Private/MassEntityTestSuiteModule.cpp
--- a/Engine/Plugins/Runtime/MassEntity/Source/MassEntityTestSuite/Private/MassEntityTestSuiteModule.cpp
+++ b/Engine/Plugins/Runtime/MassEntity/Source/MassEntityTestSuite/Private/MassEntityTestSuiteModule.cpp
@@ -2,12 +2,12 @@
 
 #include "MassEntityTestSuiteModule.h"
 
-#define LOCTEXT_NAMESPACE "PipeTestSuite"
+#define LOCTEXT_NAMESPACE "MassTest"
 
-class FPipeTestSuiteModule : public IPipeTestSuiteModule
+class FMassEntityTestSuiteModule : public IMassEntityTestSuiteModule
 {
 };
 
-IMPLEMENT_MODULE(FPipeTestSuiteModule, PipeTestSuite)
+IMPLEMENT_MODULE(FMassEntityTestSuiteModule, MassEntityTestSuite)
 
 #undef LOCTEXT_NAMESPACE
diff --git a/Engine/Plugins/Runtime/MassEntity/Source/MassEntityTestSuite/Private/MassExecutionTest.cpp b/Engine/Plugins/Runtime/MassEntity/Source/MassEntityTestSuite/Private/MassExecutionTest.cpp
--- a/Engine/Plugins/Runtime/MassEntity/Source/MassEntityTestSuite/Private/MassExecutionTest.cpp
+++ b/Engine/Plugins/Runtime/MassEntity/Source/MassEntityTestSuite/Private/MassExecutionTest.cpp
@@ -7,7 +7,7 @@
 #include "MassEntityTestTypes.h"
 #include "MassExecutor.h"
 
-#define LOCTEXT_NAMESPACE "PipeTest"
+#define LOCTEXT_NAMESPACE "MassTest"
 
 PRAGMA_DISABLE_OPTIMIZATION
 
@@ -27,7 +27,7 @@ struct FExecution_Setup : FExecutionTestBase
 		return true;
 	}
 };
-IMPLEMENT_AI_INSTANT_TEST(FExecution_Setup, "System.Pipe.Execution.Setup");
+IMPLEMENT_AI_INSTANT_TEST(FExecution_Setup, "System.Mass.Execution.Setup");
 
 
 struct FExecution_EmptyArray : FExecutionTestBase
@@ -36,13 +36,13 @@ struct FExecution_EmptyArray : FExecutionTestBase
 	{
 		CA_ASSUME(EntitySubsystem); // if EntitySubsystem null InstantTest won't be called at all
 		const float DeltaSeconds = 0.f;
-		FMassProcessingContext PipeContext(*EntitySubsystem, DeltaSeconds);
+		FMassProcessingContext ProcessingContext(*EntitySubsystem, DeltaSeconds);
 		// no test performed, let's just see if it results in errors/warnings
-		UE::Mass::Executor::RunProcessorsView(TArrayView<UMassProcessor*>(), PipeContext);
+		UE::Mass::Executor::RunProcessorsView(TArrayView<UMassProcessor*>(), ProcessingContext);
 		return true;
 	}
 };
-IMPLEMENT_AI_INSTANT_TEST(FExecution_EmptyArray, "System.Pipe.Execution.EmptyArray");
+IMPLEMENT_AI_INSTANT_TEST(FExecution_EmptyArray, "System.Mass.Execution.EmptyArray");
 
 
 struct FExecution_EmptyPipeline : FExecutionTestBase
@@ -51,34 +51,34 @@ struct FExecution_EmptyPipeline : FExecutionTestBase
 	{
 		CA_ASSUME(EntitySubsystem); // if EntitySubsystem null InstantTest won't be called at all
 		const float DeltaSeconds = 0.f;
-		FMassProcessingContext PipeContext(*EntitySubsystem, DeltaSeconds);
+		FMassProcessingContext ProcessingContext(*EntitySubsystem, DeltaSeconds);
 		FMassRuntimePipeline Pipeline;
 		// no test performed, let's just see if it results in errors/warnings
-		UE::Mass::Executor::Run(Pipeline, PipeContext);
+		UE::Mass::Executor::Run(Pipeline, ProcessingContext);
 		return true;
 	}
 };
-IMPLEMENT_AI_INSTANT_TEST(FExecution_EmptyPipeline, "System.Pipe.Execution.EmptyPipeline");
+IMPLEMENT_AI_INSTANT_TEST(FExecution_EmptyPipeline, "System.Mass.Execution.EmptyPipeline");
 
 
-struct FExecution_InvalidPipeContext : FExecutionTestBase
+struct FExecution_InvalidProcessingContext : FExecutionTestBase
 {
 	virtual bool InstantTest() override
 	{
 		CA_ASSUME(EntitySubsystem);
 		const float DeltaSeconds = 0.f;
-		FMassProcessingContext PipeContext;
+		FMassProcessingContext ProcessingContext;
 		// test assumption
-		AITEST_NULL("FMassProcessingContext\'s default constructor is expected to set FMassProcessingContext.EntitySubsystem to null", PipeContext.EntitySubsystem);
+		AITEST_NULL("FMassProcessingContext\'s default constructor is expected to set FMassProcessingContext.EntitySubsystem to null", ProcessingContext.EntitySubsystem);
 		
-		GetTestRunner().AddExpectedError(TEXT("PipeContext.EntitySubsystem is null"), EAutomationExpectedErrorFlags::Contains, 1);		
+		GetTestRunner().AddExpectedError(TEXT("ProcessingContext.EntitySubsystem is null"), EAutomationExpectedErrorFlags::Contains, 1);		
 		// note that using RunProcessorsView is to bypass reasonable tests UE::Mass::Executor::Run(Pipeline,...) does that are 
 		// reported via ensures which are not handled by the automation framework
-		UE::Mass::Executor::RunProcessorsView(TArrayView<UMassProcessor*>(), PipeContext);
+		UE::Mass::Executor::RunProcessorsView(TArrayView<UMassProcessor*>(), ProcessingContext);
 		return true;
 	}
 };
-IMPLEMENT_AI_INSTANT_TEST(FExecution_InvalidPipeContext, "System.Pipe.Execution.InvalidPipeContext");
+IMPLEMENT_AI_INSTANT_TEST(FExecution_InvalidProcessingContext, "System.Mass.Execution.InvalidProcessingContext");
 
 
 #if WITH_MASSENTITY_DEBUG
@@ -88,7 +88,7 @@ struct FExecution_SingleNullProcessor : FExecutionTestBase
 	{
 		CA_ASSUME(EntitySubsystem);
 		const float DeltaSeconds = 0.f;
-		FMassProcessingContext PipeContext(*EntitySubsystem, DeltaSeconds);
+		FMassProcessingContext ProcessingContext(*EntitySubsystem, DeltaSeconds);
 		TArray<UMassProcessor*> Processors;
 		Processors.Add(nullptr);
 		
@@ -96,11 +96,11 @@ struct FExecution_SingleNullProcessor : FExecutionTestBase
 		GetTestRunner().AddExpectedError(TEXT("Processors contains nullptr"), EAutomationExpectedErrorFlags::Contains, 1);
 		// note that using RunProcessorsView is to bypass reasonable tests UE::Mass::Executor::Run(Pipeline,...) does that are 
 		// reported via ensures which are not handled by the automation framework
-		UE::Mass::Executor::RunProcessorsView(Processors, PipeContext);
+		UE::Mass::Executor::RunProcessorsView(Processors, ProcessingContext);
 		return true;
 	}
 };
-IMPLEMENT_AI_INSTANT_TEST(FExecution_SingleNullProcessor, "System.Pipe.Execution.SingleNullProcessor");
+IMPLEMENT_AI_INSTANT_TEST(FExecution_SingleNullProcessor, "System.Mass.Execution.SingleNullProcessor");
 
 
 struct FExecution_SingleValidProcessor : FExecutionTestBase
@@ -109,16 +109,16 @@ struct FExecution_SingleValidProcessor : FExecutionTestBase
 	{
 		CA_ASSUME(EntitySubsystem);
 		const float DeltaSeconds = 0.f;
-		FMassProcessingContext PipeContext(*EntitySubsystem, DeltaSeconds);
+		FMassProcessingContext ProcessingContext(*EntitySubsystem, DeltaSeconds);
 		UPipeTestProcessorBase* Processor = NewObject<UPipeTestProcessorBase>(EntitySubsystem);
 		check(Processor);
 
 		// nothing should break. The actual result of processing is getting tested in MassProcessorTests.cpp
-		UE::Mass::Executor::Run(*Processor, PipeContext);
+		UE::Mass::Executor::Run(*Processor, ProcessingContext);
 		return true;
 	}
 };
-IMPLEMENT_AI_INSTANT_TEST(FExecution_SingleValidProcessor, "System.Pipe.Execution.SingleValidProcessor");
+IMPLEMENT_AI_INSTANT_TEST(FExecution_SingleValidProcessor, "System.Mass.Execution.SingleValidProcessor");
 
 
 struct FExecution_MultipleNullProcessors : FExecutionTestBase
@@ -127,7 +127,7 @@ struct FExecution_MultipleNullProcessors : FExecutionTestBase
 	{
 		CA_ASSUME(EntitySubsystem);
 		const float DeltaSeconds = 0.f;
-		FMassProcessingContext PipeContext(*EntitySubsystem, DeltaSeconds);
+		FMassProcessingContext ProcessingContext(*EntitySubsystem, DeltaSeconds);
 		TArray<UMassProcessor*> Processors;
 		Processors.Add(nullptr);
 		Processors.Add(nullptr);
@@ -136,11 +136,11 @@ struct FExecution_MultipleNullProcessors : FExecutionTestBase
 		GetTestRunner().AddExpectedError(TEXT("Processors contains nullptr"), EAutomationExpectedErrorFlags::Contains, 1);
 		// note that using RunProcessorsView is to bypass reasonable tests UE::Mass::Executor::Run(Pipeline,...) does that are 
 		// reported via ensures which are not handled by the automation framework
-		UE::Mass::Executor::RunProcessorsView(Processors, PipeContext);
+		UE::Mass::Executor::RunProcessorsView(Processors, ProcessingContext);
 		return true;
 	}
 };
-IMPLEMENT_AI_INSTANT_TEST(FExecution_MultipleNullProcessors, "System.Pipe.Execution.MultipleNullProcessors");
+IMPLEMENT_AI_INSTANT_TEST(FExecution_MultipleNullProcessors, "System.Mass.Execution.MultipleNullProcessors");
 #endif // WITH_MASSENTITY_DEBUG
 
 
@@ -150,7 +150,7 @@ struct FExecution_Sparse : FEntityTestBase
 	{
 		CA_ASSUME(EntitySubsystem);
 		const float DeltaSeconds = 0.f;
-		FMassProcessingContext PipeContext(*EntitySubsystem, DeltaSeconds);
+		FMassProcessingContext ProcessingContext(*EntitySubsystem, DeltaSeconds);
 		UPipeTestProcessorBase* Processor = NewObject<UPipeTestProcessorBase>(EntitySubsystem);
 		check(Processor);
 
@@ -164,12 +164,12 @@ struct FExecution_Sparse : FEntityTestBase
 		FArchetypeChunkCollection ChunkCollection(FloatsArchetype);
 		// nothing should break. The actual result of processing is getting tested in MassProcessorTests.cpp
 		
-		UE::Mass::Executor::RunSparse(Pipeline, PipeContext, ChunkCollection);
+		UE::Mass::Executor::RunSparse(Pipeline, ProcessingContext, ChunkCollection);
 
 		return true;
 	}
 };
-IMPLEMENT_AI_INSTANT_TEST(FExecution_Sparse, "System.Pipe.Execution.Sparse");
+IMPLEMENT_AI_INSTANT_TEST(FExecution_Sparse, "System.Mass.Execution.Sparse");
 } // FPipeExecutionTest
 
 PRAGMA_ENABLE_OPTIMIZATION
diff --git a/Engine/Plugins/Runtime/MassEntity/Source/MassEntityTestSuite/Private/MassProcessorTest.cpp b/Engine/Plugins/Runtime/MassEntity/Source/MassEntityTestSuite/Private/MassProcessorTest.cpp
--- a/Engine/Plugins/Runtime/MassEntity/Source/MassEntityTestSuite/Private/MassProcessorTest.cpp
+++ b/Engine/Plugins/Runtime/MassEntity/Source/MassEntityTestSuite/Private/MassProcessorTest.cpp
@@ -9,7 +9,7 @@
 #include "MassEntityTestTypes.h"
 #include "MassExecutor.h"
 
-#define LOCTEXT_NAMESPACE "PipeTest"
+#define LOCTEXT_NAMESPACE "MassTest"
 
 PRAGMA_DISABLE_OPTIMIZATION
 
@@ -29,8 +29,8 @@ int32 SimpleProcessorRun(UMassEntitySubsystem& EntitySubsystem)
 			});
 	};
 
-	FMassProcessingContext PipeContext(EntitySubsystem, /*DeltaSeconds=*/0.f);
-	UE::Mass::Executor::Run(*Processor, PipeContext);
+	FMassProcessingContext ProcessingContext(EntitySubsystem, /*DeltaSeconds=*/0.f);
+	UE::Mass::Executor::Run(*Processor, ProcessingContext);
 
 	return EntityProcessedCount;
 }
@@ -47,7 +47,7 @@ struct FProcessorTest_NoEntities : FEntityTestBase
 		return true;
 	}
 };
-IMPLEMENT_AI_INSTANT_TEST(FProcessorTest_NoEntities, "System.Pipe.Processor.NoEntities");
+IMPLEMENT_AI_INSTANT_TEST(FProcessorTest_NoEntities, "System.Mass.Processor.NoEntities");
 
 
 struct FProcessorTest_NoMatchingEntities : FEntityTestBase
@@ -64,7 +64,7 @@ struct FProcessorTest_NoMatchingEntities : FEntityTestBase
 		return true;
 	}
 };
-IMPLEMENT_AI_INSTANT_TEST(FProcessorTest_NoMatchingEntities, "System.Pipe.Processor.NoMatchingEntities");
+IMPLEMENT_AI_INSTANT_TEST(FProcessorTest_NoMatchingEntities, "System.Mass.Processor.NoMatchingEntities");
 
 struct FProcessorTest_OneMatchingArchetype : FEntityTestBase
 {
@@ -81,7 +81,7 @@ struct FProcessorTest_OneMatchingArchetype : FEntityTestBase
 		return true;
 	}
 };
-IMPLEMENT_AI_INSTANT_TEST(FProcessorTest_OneMatchingArchetype, "System.Pipe.Processor.OneMatchingArchetype");
+IMPLEMENT_AI_INSTANT_TEST(FProcessorTest_OneMatchingArchetype, "System.Mass.Processor.OneMatchingArchetype");
 
 
 struct FProcessorTest_MultipleMatchingArchetype : FEntityTestBase
@@ -102,7 +102,7 @@ struct FProcessorTest_MultipleMatchingArchetype : FEntityTestBase
 		return true;
 	}
 };
-IMPLEMENT_AI_INSTANT_TEST(FProcessorTest_MultipleMatchingArchetype, "System.Pipe.Processor.MultipleMatchingArchetype");
+IMPLEMENT_AI_INSTANT_TEST(FProcessorTest_MultipleMatchingArchetype, "System.Mass.Processor.MultipleMatchingArchetype");
 
 
 struct FProcessorTest_Requirements : FEntityTestBase
@@ -114,7 +114,7 @@ struct FProcessorTest_Requirements : FEntityTestBase
 		return true;
 	}
 };
-IMPLEMENT_AI_INSTANT_TEST(FProcessorTest_Requirements, "System.Pipe.Processor.Requirements");
+IMPLEMENT_AI_INSTANT_TEST(FProcessorTest_Requirements, "System.Mass.Processor.Requirements");
 
 
 } // FMassProcessorTestTest
diff --git a/Engine/Plugins/Runtime/MassEntity/Source/MassEntityTestSuite/Private/MassQueryTest.cpp b/Engine/Plugins/Runtime/MassEntity/Source/MassEntityTestSuite/Private/MassQueryTest.cpp
--- a/Engine/Plugins/Runtime/MassEntity/Source/MassEntityTestSuite/Private/MassQueryTest.cpp
+++ b/Engine/Plugins/Runtime/MassEntity/Source/MassEntityTestSuite/Private/MassQueryTest.cpp
@@ -9,7 +9,7 @@
 #include "MassEntityTestTypes.h"
 #include "MassExecutor.h"
 
-#define LOCTEXT_NAMESPACE "PipeTest"
+#define LOCTEXT_NAMESPACE "MassTest"
 
 PRAGMA_DISABLE_OPTIMIZATION
 
@@ -43,7 +43,7 @@ struct FQueryTest_ProcessorRequirements : FEntityTestBase
 		return true;
 	}
 };
-IMPLEMENT_AI_INSTANT_TEST(FQueryTest_ProcessorRequirements, "System.Pipe.Query.ProcessorRequiements");
+IMPLEMENT_AI_INSTANT_TEST(FQueryTest_ProcessorRequirements, "System.Mass.Query.ProcessorRequiements");
 
 
 struct FQueryTest_ExplicitRequirements : FEntityTestBase
@@ -62,7 +62,7 @@ struct FQueryTest_ExplicitRequirements : FEntityTestBase
 		return true;
 	}
 };
-IMPLEMENT_AI_INSTANT_TEST(FQueryTest_ExplicitRequirements, "System.Pipe.Query.ExplicitRequiements");
+IMPLEMENT_AI_INSTANT_TEST(FQueryTest_ExplicitRequirements, "System.Mass.Query.ExplicitRequiements");
 
 
 struct FQueryTest_FragmentViewBinding : FEntityTestBase
@@ -91,15 +91,15 @@ struct FQueryTest_FragmentViewBinding : FEntityTestBase
 				});
 		};
 
-		FMassProcessingContext PipeContext(*EntitySubsystem, /*DeltaSeconds=*/0.f);
-		UE::Mass::Executor::Run(*Processor, PipeContext);
+		FMassProcessingContext ProcessingContext(*EntitySubsystem, /*DeltaSeconds=*/0.f);
+		UE::Mass::Executor::Run(*Processor, ProcessingContext);
 
 		AITEST_EQUAL("Fragment value should have changed to the expected value", TestedFragment.Value, 13.f);
 
 		return true;
 	}
 };
-IMPLEMENT_AI_INSTANT_TEST(FQueryTest_FragmentViewBinding, "System.Pipe.Query.FragmentViewBinding");
+IMPLEMENT_AI_INSTANT_TEST(FQueryTest_FragmentViewBinding, "System.Mass.Query.FragmentViewBinding");
 
 struct FQueryTest_ExecuteSingleArchetype : FEntityTestBase
 {
@@ -137,7 +137,7 @@ struct FQueryTest_ExecuteSingleArchetype : FEntityTestBase
 		return true;
 	}
 };
-IMPLEMENT_AI_INSTANT_TEST(FQueryTest_ExecuteSingleArchetype, "System.Pipe.Query.ExecuteSingleArchetype");
+IMPLEMENT_AI_INSTANT_TEST(FQueryTest_ExecuteSingleArchetype, "System.Mass.Query.ExecuteSingleArchetype");
 
 
 struct FQueryTest_ExecuteMultipleArchetypes : FEntityTestBase
@@ -182,7 +182,7 @@ struct FQueryTest_ExecuteMultipleArchetypes : FEntityTestBase
 	}
 };
 
-IMPLEMENT_AI_INSTANT_TEST(FQueryTest_ExecuteMultipleArchetypes, "System.Pipe.Query.ExecuteMultipleArchetypes");
+IMPLEMENT_AI_INSTANT_TEST(FQueryTest_ExecuteMultipleArchetypes, "System.Mass.Query.ExecuteMultipleArchetypes");
 
 
 struct FQueryTest_ExecuteSparse : FEntityTestBase
@@ -243,7 +243,7 @@ struct FQueryTest_ExecuteSparse : FEntityTestBase
 		return true;
 	}
 };
-IMPLEMENT_AI_INSTANT_TEST(FQueryTest_ExecuteSparse, "System.Pipe.Query.ExecuteSparse");
+IMPLEMENT_AI_INSTANT_TEST(FQueryTest_ExecuteSparse, "System.Mass.Query.ExecuteSparse");
 
 
 struct FQueryTest_TagPresent : FEntityTestBase
@@ -266,7 +266,7 @@ struct FQueryTest_TagPresent : FEntityTestBase
 		return true;
 	}
 };
-IMPLEMENT_AI_INSTANT_TEST(FQueryTest_TagPresent, "System.Pipe.Query.TagPresent");
+IMPLEMENT_AI_INSTANT_TEST(FQueryTest_TagPresent, "System.Mass.Query.TagPresent");
 
 
 struct FQueryTest_TagAbsent : FEntityTestBase
@@ -290,7 +290,7 @@ struct FQueryTest_TagAbsent : FEntityTestBase
 		return true;
 	}
 };
-IMPLEMENT_AI_INSTANT_TEST(FQueryTest_TagAbsent, "System.Pipe.Query.TagAbsent");
+IMPLEMENT_AI_INSTANT_TEST(FQueryTest_TagAbsent, "System.Mass.Query.TagAbsent");
 
 
 /** using a fragment as a tag */
@@ -311,7 +311,7 @@ struct FQueryTest_FragmentPresent : FEntityTestBase
 		return true;
 	}
 };
-IMPLEMENT_AI_INSTANT_TEST(FQueryTest_FragmentPresent, "System.Pipe.Query.FragmentPresent");
+IMPLEMENT_AI_INSTANT_TEST(FQueryTest_FragmentPresent, "System.Mass.Query.FragmentPresent");
 
 
 struct FQueryTest_OnlySingleAbsentFragment : FEntityTestBase
@@ -334,7 +334,7 @@ struct FQueryTest_OnlySingleAbsentFragment : FEntityTestBase
 		return true;
 	}
 };
-IMPLEMENT_AI_INSTANT_TEST(FQueryTest_OnlySingleAbsentFragment, "System.Pipe.Query.OnlySingleAbsentFragment");
+IMPLEMENT_AI_INSTANT_TEST(FQueryTest_OnlySingleAbsentFragment, "System.Mass.Query.OnlySingleAbsentFragment");
 
 
 struct FQueryTest_OnlyMultipleAbsentFragments : FEntityTestBase
@@ -358,7 +358,7 @@ struct FQueryTest_OnlyMultipleAbsentFragments : FEntityTestBase
 		return true;
 	}
 };
-IMPLEMENT_AI_INSTANT_TEST(FQueryTest_OnlyMultipleAbsentFragments, "System.Pipe.Query.OnlyMultipleAbsentFragments");
+IMPLEMENT_AI_INSTANT_TEST(FQueryTest_OnlyMultipleAbsentFragments, "System.Mass.Query.OnlyMultipleAbsentFragments");
 
 
 struct FQueryTest_AbsentAndPresentFragments : FEntityTestBase
@@ -379,7 +379,7 @@ struct FQueryTest_AbsentAndPresentFragments : FEntityTestBase
 		return true;
 	}
 };
-IMPLEMENT_AI_INSTANT_TEST(FQueryTest_AbsentAndPresentFragments, "System.Pipe.Query.AbsentAndPresentFragments");
+IMPLEMENT_AI_INSTANT_TEST(FQueryTest_AbsentAndPresentFragments, "System.Mass.Query.AbsentAndPresentFragments");
 
 
 struct FQueryTest_SingleOptionalFragment : FEntityTestBase
@@ -399,7 +399,7 @@ struct FQueryTest_SingleOptionalFragment : FEntityTestBase
 		return true;
 	}
 };
-IMPLEMENT_AI_INSTANT_TEST(FQueryTest_SingleOptionalFragment, "System.Pipe.Query.SingleOptionalFragment");
+IMPLEMENT_AI_INSTANT_TEST(FQueryTest_SingleOptionalFragment, "System.Mass.Query.SingleOptionalFragment");
 
 
 struct FQueryTest_MultipleOptionalFragment : FEntityTestBase
@@ -417,7 +417,7 @@ struct FQueryTest_MultipleOptionalFragment : FEntityTestBase
 		return true;
 	}
 };
-IMPLEMENT_AI_INSTANT_TEST(FQueryTest_MultipleOptionalFragment, "System.Pipe.Query.MultipleOptionalFragment");
+IMPLEMENT_AI_INSTANT_TEST(FQueryTest_MultipleOptionalFragment, "System.Mass.Query.MultipleOptionalFragment");
 
 
 /** This test configures a query to fetch archetypes that have a Float fragment (we have two of these) with an optional 
@@ -465,7 +465,7 @@ struct FQueryTest_UsingOptionalFragment : FEntityTestBase
 		return true;
 	}
 };
-IMPLEMENT_AI_INSTANT_TEST(FQueryTest_UsingOptionalFragment, "System.Pipe.Query.UsingOptionalFragment");
+IMPLEMENT_AI_INSTANT_TEST(FQueryTest_UsingOptionalFragment, "System.Mass.Query.UsingOptionalFragment");
 
 
 struct FQueryTest_AnyFragment : FEntityTestBase
@@ -508,7 +508,7 @@ struct FQueryTest_AnyFragment : FEntityTestBase
 		return true;
 	}
 };
-IMPLEMENT_AI_INSTANT_TEST(FQueryTest_AnyFragment, "System.Pipe.Query.AnyFragment");
+IMPLEMENT_AI_INSTANT_TEST(FQueryTest_AnyFragment, "System.Mass.Query.AnyFragment");
 
 
 struct FQueryTest_AnyTag : FEntityTestBase
@@ -540,7 +540,7 @@ struct FQueryTest_AnyTag : FEntityTestBase
 		return true;
 	}
 };
-IMPLEMENT_AI_INSTANT_TEST(FQueryTest_AnyTag, "System.Pipe.Query.AnyTag");
+IMPLEMENT_AI_INSTANT_TEST(FQueryTest_AnyTag, "System.Mass.Query.AnyTag");
 
 } // FPipeQueryTest
 
diff --git a/Engine/Plugins/Runtime/MassEntity/Source/MassEntityTestSuite/Public/MassEntityTestSuiteModule.h b/Engine/Plugins/Runtime/MassEntity/Source/MassEntityTestSuite/Public/MassEntityTestSuiteModule.h
--- a/Engine/Plugins/Runtime/MassEntity/Source/MassEntityTestSuite/Public/MassEntityTestSuiteModule.h
+++ b/Engine/Plugins/Runtime/MassEntity/Source/MassEntityTestSuite/Public/MassEntityTestSuiteModule.h
@@ -8,7 +8,7 @@
 /**
 * The public interface to this module
 */
-class IPipeTestSuiteModule : public IModuleInterface
+class IMassEntityTestSuiteModule : public IModuleInterface
 {
 
 public:
@@ -19,9 +19,9 @@ public:
 	*
 	* @return Returns singleton instance, loading the module on demand if needed
 	*/
-	static inline IPipeTestSuiteModule& Get()
+	static inline IMassEntityTestSuiteModule& Get()
 	{
-		return FModuleManager::LoadModuleChecked<IPipeTestSuiteModule>("PipeTestSuite");
+		return FModuleManager::LoadModuleChecked<IMassEntityTestSuiteModule>("MassEntityTestSuite");
 	}
 
 	/**
@@ -31,7 +31,7 @@ public:
 	*/
 	static inline bool IsAvailable()
 	{
-		return FModuleManager::Get().IsModuleLoaded("PipeTestSuite");
+		return FModuleManager::Get().IsModuleLoaded("MassEntityTestSuite");
 	}
 };
 
diff --git a/Engine/Plugins/Runtime/MassGameplay/Source/MassCommon/Private/MassCommonModule.cpp b/Engine/Plugins/Runtime/MassGameplay/Source/MassCommon/Private/MassCommonModule.cpp
--- a/Engine/Plugins/Runtime/MassGameplay/Source/MassCommon/Private/MassCommonModule.cpp
+++ b/Engine/Plugins/Runtime/MassGameplay/Source/MassCommon/Private/MassCommonModule.cpp
@@ -5,8 +5,6 @@
 #include "MassCommonTypes.h"
 
 
-DEFINE_LOG_CATEGORY(LogMass);
-
 class FMassCommonModule : public IMassCommonModule
 {
 	/** IModuleInterface implementation */
diff --git a/Engine/Plugins/Runtime/MassGameplay/Source/MassCommon/Public/MassCommonTypes.h b/Engine/Plugins/Runtime/MassGameplay/Source/MassCommon/Public/MassCommonTypes.h
--- a/Engine/Plugins/Runtime/MassGameplay/Source/MassCommon/Public/MassCommonTypes.h
+++ b/Engine/Plugins/Runtime/MassGameplay/Source/MassCommon/Public/MassCommonTypes.h
@@ -13,8 +13,6 @@
 class UStaticMesh;
 class UMaterialInterface;
 
-MASSCOMMON_API DECLARE_LOG_CATEGORY_EXTERN(LogMass, Warning, All);
-
 #if WITH_MASSGAMEPLAY_DEBUG
 namespace UE::Mass::Debug
 {
diff --git a/Engine/Plugins/Runtime/MassGameplay/Source/MassSimulation/Private/MassSimulationSettings.cpp b/Engine/Plugins/Runtime/MassGameplay/Source/MassSimulation/Private/MassSimulationSettings.cpp
--- a/Engine/Plugins/Runtime/MassGameplay/Source/MassSimulation/Private/MassSimulationSettings.cpp
+++ b/Engine/Plugins/Runtime/MassGameplay/Source/MassSimulation/Private/MassSimulationSettings.cpp
@@ -10,10 +10,10 @@ void UMassSimulationSettings::PostInitProperties()
 {
 	Super::PostInitProperties();
 
-	FPipeEditorModule* PipeEditorModule = FModuleManager::GetModulePtr<FPipeEditorModule>("MassEntityEditor");
-	if (PipeEditorModule)
+	FMassEntityEditorModule* MassEntityEditorModule = FModuleManager::GetModulePtr<FMassEntityEditorModule>("MassEntityEditor");
+	if (MassEntityEditorModule)
 	{
-		PipeEditorModule->GetOnAssetPropertiesChanged().AddUObject(this, &UMassSimulationSettings::OnAssetPropertiesChanged);
+		MassEntityEditorModule->GetOnAssetPropertiesChanged().AddUObject(this, &UMassSimulationSettings::OnAssetPropertiesChanged);
 	}
 }
 
diff --git a/Engine/Plugins/Runtime/MassGameplay/Source/MassSpawner/Private/MassSpawnerSubsystem.cpp b/Engine/Plugins/Runtime/MassGameplay/Source/MassSpawner/Private/MassSpawnerSubsystem.cpp
--- a/Engine/Plugins/Runtime/MassGameplay/Source/MassSpawner/Private/MassSpawnerSubsystem.cpp
+++ b/Engine/Plugins/Runtime/MassGameplay/Source/MassSpawner/Private/MassSpawnerSubsystem.cpp
@@ -48,9 +48,9 @@ namespace UE::Mass::SpawnerSubsystem
 			FMassRuntimePipeline DestructionPipeline;
 			DestructionPipeline.InitializeFromArray(Destructors, SpawnerSubSystem);
 
-			FMassProcessingContext PipeContext(EntitySystem, /*TimeDelta=*/0.0f);
+			FMassProcessingContext ProcessingContext(EntitySystem, /*TimeDelta=*/0.0f);
 
-			UE::Mass::Executor::RunSparse(DestructionPipeline, PipeContext, Chunks);
+			UE::Mass::Executor::RunSparse(DestructionPipeline, ProcessingContext, Chunks);
 		}
 
 		return Destructors.Num() > 0;
@@ -144,8 +144,8 @@ bool UMassSpawnerSubsystem::SpawnEntities(const FMassEntityTemplate& EntityTempl
 		{
 			FMassEntityTemplate& MutableEntityTemplate = const_cast<FMassEntityTemplate&>(EntityTemplate);
 
-			FMassProcessingContext PipeContext(*EntitySystem, /*TimeDelta=*/0.0f);
-			UE::Mass::Executor::RunSparse(MutableEntityTemplate.GetMutableInitializationPipeline(), PipeContext, ChunkCollection);
+			FMassProcessingContext ProcessingContext(*EntitySystem, /*TimeDelta=*/0.0f);
+			UE::Mass::Executor::RunSparse(MutableEntityTemplate.GetMutableInitializationPipeline(), ProcessingContext, ChunkCollection);
 		}
 		return true;
 	}
@@ -235,7 +235,7 @@ void UMassSpawnerSubsystem::DestroyEntities(const FMassEntityTemplateID Template
 		TArray<FArchetypeChunkCollection> ChunkCollections;
 		UE::Mass::SpawnerSubsystem::CreateSparseChunks(*EntitySystem, Entities, ChunkCollections);
 
-		FMassProcessingContext PipeContext(*EntitySystem, /*TimeDelta=*/0.0f);
+		FMassProcessingContext ProcessingContext(*EntitySystem, /*TimeDelta=*/0.0f);
 
 		// Run destructors
 		bool bDestructorsRun = false;
@@ -260,7 +260,7 @@ void UMassSpawnerSubsystem::DestroyEntities(const FMassEntityTemplateID Template
 			{
 				if (Chunks.GetArchetype().IsValid())
 				{
-					UE::Mass::Executor::RunSparse(MutableEntityTemplate.GetMutableDeinitializationPipeline(), PipeContext, Chunks);
+					UE::Mass::Executor::RunSparse(MutableEntityTemplate.GetMutableDeinitializationPipeline(), ProcessingContext, Chunks);
 				}
 			}
 		}
@@ -309,9 +309,9 @@ void UMassSpawnerSubsystem::DoSpawning(const FMassEntityTemplate& EntityTemplate
 	{
 		FMassEntityTemplate& MutableEntityTemplate = const_cast<FMassEntityTemplate&>(EntityTemplate);
 
-		FMassProcessingContext PipeContext(*EntitySystem, /*TimeDelta=*/0.0f);
-		PipeContext.AuxData = AuxData;
-		UE::Mass::Executor::RunSparse(MutableEntityTemplate.GetMutableInitializationPipeline(), PipeContext, ChunkCollection);
+		FMassProcessingContext ProcessingContext(*EntitySystem, /*TimeDelta=*/0.0f);
+		ProcessingContext.AuxData = AuxData;
+		UE::Mass::Executor::RunSparse(MutableEntityTemplate.GetMutableInitializationPipeline(), ProcessingContext, ChunkCollection);
 	}
 
 	OutEntities.Append(MoveTemp(SpawnedEntities));
