Made all MassXXXSettings and MassProcessors use the same file: DefaultMass.ini

diff --git a/Engine/Plugins/AI/MassAI/Source/MassAIMovement/Public/MassAvoidanceSettings.h b/Engine/Plugins/AI/MassAI/Source/MassAIMovement/Public/MassAvoidanceSettings.h
--- a/Engine/Plugins/AI/MassAI/Source/MassAIMovement/Public/MassAvoidanceSettings.h
+++ b/Engine/Plugins/AI/MassAI/Source/MassAIMovement/Public/MassAvoidanceSettings.h
@@ -24,7 +24,7 @@ struct MASSAIMOVEMENT_API FMassAvoidanceVelocityFilter
 	float HighSpeed = 50.f;
 };
 
-UCLASS(config = Engine, defaultconfig, meta = (DisplayName = "Mass Avoidance"))
+UCLASS(config = Mass, defaultconfig, meta = (DisplayName = "Mass Avoidance"))
 class MASSAIMOVEMENT_API UMassAvoidanceSettings : public UDeveloperSettings
 {
 	GENERATED_BODY()
diff --git a/Engine/Plugins/AI/MassAI/Source/MassAIMovement/Public/MassMovementSettings.h b/Engine/Plugins/AI/MassAI/Source/MassAIMovement/Public/MassMovementSettings.h
--- a/Engine/Plugins/AI/MassAI/Source/MassAIMovement/Public/MassMovementSettings.h
+++ b/Engine/Plugins/AI/MassAI/Source/MassAIMovement/Public/MassMovementSettings.h
@@ -7,7 +7,7 @@
 #include "MassAIMovementTypes.h"
 #include "MassMovementSettings.generated.h"
 
-UCLASS(config = Engine, defaultconfig, meta = (DisplayName = "Mass Movement"))
+UCLASS(config = Mass, defaultconfig, meta = (DisplayName = "Mass Movement"))
 class MASSAIMOVEMENT_API UMassMovementSettings : public UDeveloperSettings
 {
 	GENERATED_BODY()
diff --git a/Engine/Plugins/Runtime/MassEntity/Source/MassEntity/Public/MassEntitySettings.h b/Engine/Plugins/Runtime/MassEntity/Source/MassEntity/Public/MassEntitySettings.h
--- a/Engine/Plugins/Runtime/MassEntity/Source/MassEntity/Public/MassEntitySettings.h
+++ b/Engine/Plugins/Runtime/MassEntity/Source/MassEntity/Public/MassEntitySettings.h
@@ -44,9 +44,9 @@ struct FMassProcessingPhaseConfig
 };
 
 /**
- * Implements the settings for MassSimulation
+ * Implements the settings for MassEntity plugin
  */
-UCLASS(config = Game, defaultconfig, DisplayName = "Mass")
+UCLASS(config = Mass, defaultconfig, DisplayName = "Mass")
 class MASSENTITY_API UMassSettings : public UDeveloperSettings
 {
 	GENERATED_BODY()
diff --git a/Engine/Plugins/Runtime/MassEntity/Source/MassEntity/Public/MassProcessor.h b/Engine/Plugins/Runtime/MassEntity/Source/MassEntity/Public/MassProcessor.h
--- a/Engine/Plugins/Runtime/MassEntity/Source/MassEntity/Public/MassProcessor.h
+++ b/Engine/Plugins/Runtime/MassEntity/Source/MassEntity/Public/MassProcessor.h
@@ -38,7 +38,7 @@ struct FMassProcessorExecutionOrder
 };
 
 
-UCLASS(abstract, EditInlineNew, CollapseCategories, config = Game, defaultconfig)
+UCLASS(abstract, EditInlineNew, CollapseCategories, config = Mass, defaultconfig)
 class MASSENTITY_API UMassProcessor : public UObject
 {
 	GENERATED_BODY()
diff --git a/Engine/Plugins/Runtime/MassGameplay/Source/MassSimulation/Public/MassSimulationSettings.h b/Engine/Plugins/Runtime/MassGameplay/Source/MassSimulation/Public/MassSimulationSettings.h
--- a/Engine/Plugins/Runtime/MassGameplay/Source/MassSimulation/Public/MassSimulationSettings.h
+++ b/Engine/Plugins/Runtime/MassGameplay/Source/MassSimulation/Public/MassSimulationSettings.h
@@ -12,7 +12,7 @@ class UMassSchematic;
 /**
  * Implements the settings for MassSimulation
  */
-UCLASS(config = Game, defaultconfig, DisplayName = "Mass")
+UCLASS(config = Mass, defaultconfig, DisplayName = "Mass Simulation")
 class MASSSIMULATION_API UMassSimulationSettings : public UDeveloperSettings
 {
 	GENERATED_BODY()
