Renamed all Component-containing functions and member variable names in Mass, replaced with "Fragment". Updated a lot of comments and logged text as well.

diff --git a/Engine/Plugins/Runtime/MassEntity/Source/MassEntity/Private/MassArchetypeData.cpp b/Engine/Plugins/Runtime/MassEntity/Source/MassEntity/Private/MassArchetypeData.cpp
--- a/Engine/Plugins/Runtime/MassEntity/Source/MassEntity/Private/MassArchetypeData.cpp
+++ b/Engine/Plugins/Runtime/MassEntity/Source/MassEntity/Private/MassArchetypeData.cpp
@@ -9,32 +9,32 @@
 // FMassArchetypeData
 namespace UE::Mass::Core 
 {
-	constexpr static bool bBitwiseRelocateComponents = true;
+	constexpr static bool bBitwiseRelocateFragments = true;
 }// UE::Mass::Core
 
-void FMassArchetypeData::ForEachComponentType(TFunction< void(const UScriptStruct* /*Component*/)> Function) const
+void FMassArchetypeData::ForEachFragmentType(TFunction< void(const UScriptStruct* /*Fragment*/)> Function) const
 {
-	for (const FMassArchetypeFragmentConfig& ComponentData : ComponentConfigs)
+	for (const FMassArchetypeFragmentConfig& FragmentData : FragmentConfigs)
 	{
-		Function(ComponentData.ComponentType);
+		Function(FragmentData.FragmentType);
 	}
 }
 
-bool FMassArchetypeData::HasComponentType(const UScriptStruct* ComponentType) const
+bool FMassArchetypeData::HasFragmentType(const UScriptStruct* FragmentType) const
 {
-	return (ComponentType && CompositionDescriptor.Components.Contains(*ComponentType));
+	return (FragmentType && CompositionDescriptor.Fragments.Contains(*FragmentType));
 }
 
 bool FMassArchetypeData::IsEquivalent(TConstArrayView<const UScriptStruct*> OtherList) const
 {
-	if (OtherList.Num() != ComponentConfigs.Num())
+	if (OtherList.Num() != FragmentConfigs.Num())
 	{
 		return false;
 	}
 
-	for (int32 Index = 0; Index < ComponentConfigs.Num(); ++Index)
+	for (int32 Index = 0; Index < FragmentConfigs.Num(); ++Index)
 	{
-		if (OtherList[Index] != ComponentConfigs[Index].ComponentType)
+		if (OtherList[Index] != FragmentConfigs[Index].FragmentType)
 		{
 			return false;
 		}
@@ -43,68 +43,68 @@ bool FMassArchetypeData::IsEquivalent(TConstArrayView<const UScriptStruct*> Othe
 	return true;
 }
 
-void FMassArchetypeData::Initialize(const FMassFragmentBitSet& Components, const FMassTagBitSet& Tags, const FMassChunkFragmentBitSet& ChunkComponents)
+void FMassArchetypeData::Initialize(const FMassFragmentBitSet& Fragments, const FMassTagBitSet& Tags, const FMassChunkFragmentBitSet& ChunkFragments)
 {
-	TArray<const UScriptStruct*, TInlineAllocator<16>> SortedComponentList;
-	Components.ExportTypes(SortedComponentList);
+	TArray<const UScriptStruct*, TInlineAllocator<16>> SortedFragmentList;
+	Fragments.ExportTypes(SortedFragmentList);
 
-	SortedComponentList.Sort(FMassSorterOperator<UScriptStruct>());
+	SortedFragmentList.Sort(FMassSorterOperator<UScriptStruct>());
 
-	// Figure out how many bytes all of the individual components (and metadata) will cost per entity
-	int32 ComponentSizeTallyBytes = 0;
+	// Figure out how many bytes all of the individual fragments (and metadata) will cost per entity
+	int32 FragmentSizeTallyBytes = 0;
 
 	// Alignment padding computation is currently very conservative and over-estimated.
 	int32 AlignmentPadding = 0;
 	{
 		// Save room for the 'metadata' (entity array)
-		ComponentSizeTallyBytes += sizeof(FMassEntityHandle);
+		FragmentSizeTallyBytes += sizeof(FMassEntityHandle);
 
-		// Tally up the component sizes and place them in the index map
-		ComponentConfigs.AddDefaulted(SortedComponentList.Num());
-		ComponentIndexMap.Reserve(SortedComponentList.Num());
+		// Tally up the fragment sizes and place them in the index map
+		FragmentConfigs.AddDefaulted(SortedFragmentList.Num());
+		FragmentIndexMap.Reserve(SortedFragmentList.Num());
 
-		for (int32 ComponentIndex = 0; ComponentIndex < SortedComponentList.Num(); ++ComponentIndex)
+		for (int32 FragmentIndex = 0; FragmentIndex < SortedFragmentList.Num(); ++FragmentIndex)
 		{
-			const UScriptStruct* ComponentType = SortedComponentList[ComponentIndex];
-			checkSlow(ComponentType);
-			ComponentConfigs[ComponentIndex].ComponentType = ComponentType;
+			const UScriptStruct* FragmentType = SortedFragmentList[FragmentIndex];
+			checkSlow(FragmentType);
+			FragmentConfigs[FragmentIndex].FragmentType = FragmentType;
 			
-			AlignmentPadding += ComponentType->GetMinAlignment();
-			ComponentSizeTallyBytes += ComponentType->GetStructureSize();
+			AlignmentPadding += FragmentType->GetMinAlignment();
+			FragmentSizeTallyBytes += FragmentType->GetStructureSize();
 
-			ComponentIndexMap.Add(ComponentType, ComponentIndex);
-			CompositionDescriptor.Components.Add(*ComponentType);
+			FragmentIndexMap.Add(FragmentType, FragmentIndex);
+			CompositionDescriptor.Fragments.Add(*FragmentType);
 		}
-		checkSlow(Components.IsEquivalent(CompositionDescriptor.Components));
+		checkSlow(Fragments.IsEquivalent(CompositionDescriptor.Fragments));
 	}
 
 	CompositionDescriptor.Tags = Tags;
-	CompositionDescriptor.ChunkComponents = ChunkComponents;
-	// create the chunk component instances
-	TArray<const UScriptStruct*, TInlineAllocator<16>> ChunkComponentList;
-	ChunkComponents.ExportTypes(ChunkComponentList);
-	ChunkComponentList.Sort(FMassSorterOperator<UScriptStruct>());
-	for (const UScriptStruct* ChunkComponentType : ChunkComponentList)
+	CompositionDescriptor.ChunkFragments = ChunkFragments;
+	// create the chunk fragment instances
+	TArray<const UScriptStruct*, TInlineAllocator<16>> ChunkFragmentList;
+	ChunkFragments.ExportTypes(ChunkFragmentList);
+	ChunkFragmentList.Sort(FMassSorterOperator<UScriptStruct>());
+	for (const UScriptStruct* ChunkFragmentType : ChunkFragmentList)
 	{
-		check(ChunkComponentType);
-		ChunkComponentTemplates.Add(FInstancedStruct(ChunkComponentType));
+		check(ChunkFragmentType);
+		ChunkFragmentTemplates.Add(FInstancedStruct(ChunkFragmentType));
 	}
 
-	TotalBytesPerEntity = ComponentSizeTallyBytes;
+	TotalBytesPerEntity = FragmentSizeTallyBytes;
 	int32 ChunkAvailableSize = GetChunkAllocSize() - AlignmentPadding;
 	check(TotalBytesPerEntity <= ChunkAvailableSize);
 
 	NumEntitiesPerChunk = ChunkAvailableSize / TotalBytesPerEntity;
 
-	// Set up the offsets for each component into the chunk data
+	// Set up the offsets for each fragment into the chunk data
 	EntityListOffsetWithinChunk = 0;
 	int32 CurrentOffset = NumEntitiesPerChunk * sizeof(FMassEntityHandle);
-	for (FMassArchetypeFragmentConfig& ComponentData : ComponentConfigs)
+	for (FMassArchetypeFragmentConfig& FragmentData : FragmentConfigs)
 	{
-		CurrentOffset = Align(CurrentOffset, ComponentData.ComponentType->GetMinAlignment());
-		ComponentData.ArrayOffsetWithinChunk = CurrentOffset;
-		const int32 SizeOfThisComponentArray = NumEntitiesPerChunk * ComponentData.ComponentType->GetStructureSize();
-		CurrentOffset += SizeOfThisComponentArray;
+		CurrentOffset = Align(CurrentOffset, FragmentData.FragmentType->GetMinAlignment());
+		FragmentData.ArrayOffsetWithinChunk = CurrentOffset;
+		const int32 SizeOfThisFragmentArray = NumEntitiesPerChunk * FragmentData.FragmentType->GetStructureSize();
+		CurrentOffset += SizeOfThisFragmentArray;
 	}
 }
 
@@ -112,26 +112,26 @@ void FMassArchetypeData::InitializeWithSibling(const FMassArchetypeData& Sibling
 {
 	checkf(IsInitialized() == false, TEXT("Trying to %s but this archetype has already been initialized"));
 
-	ComponentConfigs = SiblingArchetype.ComponentConfigs;
-	ComponentIndexMap = SiblingArchetype.ComponentIndexMap;
-	CompositionDescriptor.Components = SiblingArchetype.CompositionDescriptor.Components;
+	FragmentConfigs = SiblingArchetype.FragmentConfigs;
+	FragmentIndexMap = SiblingArchetype.FragmentIndexMap;
+	CompositionDescriptor.Fragments = SiblingArchetype.CompositionDescriptor.Fragments;
 	CompositionDescriptor.Tags = OverrideTags;
-	CompositionDescriptor.ChunkComponents = SiblingArchetype.CompositionDescriptor.ChunkComponents;
-	ChunkComponentTemplates = SiblingArchetype.ChunkComponentTemplates;
+	CompositionDescriptor.ChunkFragments = SiblingArchetype.CompositionDescriptor.ChunkFragments;
+	ChunkFragmentTemplates = SiblingArchetype.ChunkFragmentTemplates;
 
 	TotalBytesPerEntity = SiblingArchetype.TotalBytesPerEntity;
 	NumEntitiesPerChunk = SiblingArchetype.NumEntitiesPerChunk;
 
-	// Set up the offsets for each component into the chunk data
+	// Set up the offsets for each fragment into the chunk data
 	EntityListOffsetWithinChunk = 0;
 }
 
 void FMassArchetypeData::AddEntity(FMassEntityHandle Entity)
 {
-	AddEntityInternal(Entity, true/*bInitializeComponents*/);
+	AddEntityInternal(Entity, true/*bInitializeFragments*/);
 }
 
-int32 FMassArchetypeData::AddEntityInternal(FMassEntityHandle Entity, const bool bInitializeComponents)
+int32 FMassArchetypeData::AddEntityInternal(FMassEntityHandle Entity, const bool bInitializeFragments)
 {
 	int32 IndexWithinChunk = 0;
 	int32 AbsoluteIndex = 0;
@@ -173,25 +173,25 @@ int32 FMassArchetypeData::AddEntityInternal(FMassEntityHandle Entity, const bool
 		if (EmptyChunkIndex != INDEX_NONE)
 		{
 			DestinationChunk = &Chunks[EmptyChunkIndex];
-			DestinationChunk->Recycle(ChunkComponentTemplates);
+			DestinationChunk->Recycle(ChunkFragmentTemplates);
 			AbsoluteIndex = EmptyAbsoluteIndex;
 		}
 		else
 		{
-			DestinationChunk = &Chunks.Emplace_GetRef(GetChunkAllocSize(), ChunkComponentTemplates);
+			DestinationChunk = &Chunks.Emplace_GetRef(GetChunkAllocSize(), ChunkFragmentTemplates);
 		}
 
 		check(DestinationChunk);
 		DestinationChunk->AddInstance();
 	}
 
-	// Initialize the component memory
-	if (bInitializeComponents)
+	// Initialize the fragment memory
+	if (bInitializeFragments)
 	{
-		for (const FMassArchetypeFragmentConfig& ComponentConfig : ComponentConfigs)
+		for (const FMassArchetypeFragmentConfig& FragmentConfig : FragmentConfigs)
 		{
-			void* ComponentPtr = ComponentConfig.GetComponentData(DestinationChunk->GetRawMemory(), IndexWithinChunk);
-			ComponentConfig.ComponentType->InitializeStruct(ComponentPtr);
+			void* FragmentPtr = FragmentConfig.GetFragmentData(DestinationChunk->GetRawMemory(), IndexWithinChunk);
+			FragmentConfig.FragmentType->InitializeStruct(FragmentPtr);
 		}
 	}
 
@@ -205,10 +205,10 @@ int32 FMassArchetypeData::AddEntityInternal(FMassEntityHandle Entity, const bool
 void FMassArchetypeData::RemoveEntity(FMassEntityHandle Entity)
 {
 	const int32 AbsoluteIndex = EntityMap.FindAndRemoveChecked(Entity.Index);
-	RemoveEntityInternal(AbsoluteIndex, true/*bDestroyComponents*/);
+	RemoveEntityInternal(AbsoluteIndex, true/*bDestroyFragments*/);
 }
 
-void FMassArchetypeData::RemoveEntityInternal(const int32 AbsoluteIndex, const bool bDestroyComponents)
+void FMassArchetypeData::RemoveEntityInternal(const int32 AbsoluteIndex, const bool bDestroyFragments)
 {
 	const int32 ChunkIndex = AbsoluteIndex / NumEntitiesPerChunk;
 	const int32 IndexWithinChunk = AbsoluteIndex % NumEntitiesPerChunk;
@@ -218,37 +218,37 @@ void FMassArchetypeData::RemoveEntityInternal(const int32 AbsoluteIndex, const b
 	Chunk.RemoveInstance();
 	const int32 IndexToSwapFrom = Chunk.GetNumInstances();
 
-	checkf(bDestroyComponents || UE::Mass::Core::bBitwiseRelocateComponents, TEXT("We allow not to destroy components only in bit wise relocation mode."));
+	checkf(bDestroyFragments || UE::Mass::Core::bBitwiseRelocateFragments, TEXT("We allow not to destroy fragments only in bit wise relocation mode."));
 
 	// Remove and swap the last entry in the chunk to the location of the removed item (if it's not the same as the dying entry)
 	if (IndexToSwapFrom != IndexWithinChunk)
 	{
-		for (const FMassArchetypeFragmentConfig& ComponentConfig : ComponentConfigs)
+		for (const FMassArchetypeFragmentConfig& FragmentConfig : FragmentConfigs)
 		{
-			void* DyingComponentPtr = ComponentConfig.GetComponentData(Chunk.GetRawMemory(), IndexWithinChunk);
-			void* MovingComponentPtr = ComponentConfig.GetComponentData(Chunk.GetRawMemory(), IndexToSwapFrom);
+			void* DyingFragmentPtr = FragmentConfig.GetFragmentData(Chunk.GetRawMemory(), IndexWithinChunk);
+			void* MovingFragmentPtr = FragmentConfig.GetFragmentData(Chunk.GetRawMemory(), IndexToSwapFrom);
 
-			if (UE::Mass::Core::bBitwiseRelocateComponents)
+			if (UE::Mass::Core::bBitwiseRelocateFragments)
 			{
-				// Destroy component data
-				if (bDestroyComponents)
+				// Destroy fragment data
+				if (bDestroyFragments)
 				{
-					ComponentConfig.ComponentType->DestroyStruct(DyingComponentPtr);
+					FragmentConfig.FragmentType->DestroyStruct(DyingFragmentPtr);
 				}
 
 				// Move last entry
-				FMemory::Memcpy(DyingComponentPtr, MovingComponentPtr, ComponentConfig.ComponentType->GetStructureSize());
+				FMemory::Memcpy(DyingFragmentPtr, MovingFragmentPtr, FragmentConfig.FragmentType->GetStructureSize());
 			}
 			else
 			{
-				// Destroy & initialize the component data
-				ComponentConfig.ComponentType->ClearScriptStruct(DyingComponentPtr);
+				// Destroy & initialize the fragment data
+				FragmentConfig.FragmentType->ClearScriptStruct(DyingFragmentPtr);
 
 				// Copy last entry
-				ComponentConfig.ComponentType->CopyScriptStruct(DyingComponentPtr, MovingComponentPtr);
+				FragmentConfig.FragmentType->CopyScriptStruct(DyingFragmentPtr, MovingFragmentPtr);
 
 				// Destroy last entry
-				ComponentConfig.ComponentType->DestroyStruct(MovingComponentPtr);
+				FragmentConfig.FragmentType->DestroyStruct(MovingFragmentPtr);
 			}
 		}
 
@@ -257,13 +257,13 @@ void FMassArchetypeData::RemoveEntityInternal(const int32 AbsoluteIndex, const b
 		Chunk.GetEntityArrayElementRef(EntityListOffsetWithinChunk, IndexWithinChunk) = EntityBeingSwapped;
 		EntityMap.FindChecked(EntityBeingSwapped.Index) = AbsoluteIndex;
 	}
-	else if (!UE::Mass::Core::bBitwiseRelocateComponents || bDestroyComponents)
+	else if (!UE::Mass::Core::bBitwiseRelocateFragments || bDestroyFragments)
 	{
-		for (const FMassArchetypeFragmentConfig& ComponentConfig : ComponentConfigs)
+		for (const FMassArchetypeFragmentConfig& FragmentConfig : FragmentConfigs)
 		{
-			// Destroy the component data
-			void* DyingComponentPtr = ComponentConfig.GetComponentData(Chunk.GetRawMemory(), IndexWithinChunk);
-			ComponentConfig.ComponentType->DestroyStruct(DyingComponentPtr);
+			// Destroy the fragment data
+			void* DyingFragmentPtr = FragmentConfig.GetFragmentData(Chunk.GetRawMemory(), IndexWithinChunk);
+			FragmentConfig.FragmentType->DestroyStruct(DyingFragmentPtr);
 		}
 	}
 
@@ -304,29 +304,29 @@ void FMassArchetypeData::BatchDestroyEntityChunks(const FArchetypeChunkCollectio
 			const int32 SwapStartIndex = Chunk.GetNumInstances() - NumberToMove;
 			checkf((SubchunkInfo.SubchunkStart + NumberToMove - 1) < SwapStartIndex, TEXT("Remove and Move ranges overlap"));
 
-			for (const FMassArchetypeFragmentConfig& ComponentConfig : ComponentConfigs)
+			for (const FMassArchetypeFragmentConfig& FragmentConfig : FragmentConfigs)
 			{
-				void* DyingComponentPtr = ComponentConfig.GetComponentData(Chunk.GetRawMemory(), SubchunkInfo.SubchunkStart);
-				void* MovingComponentPtr = ComponentConfig.GetComponentData(Chunk.GetRawMemory(), SwapStartIndex);
+				void* DyingFragmentPtr = FragmentConfig.GetFragmentData(Chunk.GetRawMemory(), SubchunkInfo.SubchunkStart);
+				void* MovingFragmentPtr = FragmentConfig.GetFragmentData(Chunk.GetRawMemory(), SwapStartIndex);
 
-				if (UE::Mass::Core::bBitwiseRelocateComponents)
+				if (UE::Mass::Core::bBitwiseRelocateFragments)
 				{
-					// Destroy the components we'll replace by the following copy
-					ComponentConfig.ComponentType->DestroyStruct(DyingComponentPtr, NumberToMove);
+					// Destroy the fragments we'll replace by the following copy
+					FragmentConfig.FragmentType->DestroyStruct(DyingFragmentPtr, NumberToMove);
 
-					// Swap components to the empty space just created.
-					FMemory::Memcpy(DyingComponentPtr, MovingComponentPtr, ComponentConfig.ComponentType->GetStructureSize() * NumberToMove);
+					// Swap fragments to the empty space just created.
+					FMemory::Memcpy(DyingFragmentPtr, MovingFragmentPtr, FragmentConfig.FragmentType->GetStructureSize() * NumberToMove);
 				}
 				else
 				{
-					// Clear components that we will copy over. Clear destroys and initializes the components, which is needed for CopyScriptStruct().
-					ComponentConfig.ComponentType->ClearScriptStruct(DyingComponentPtr, NumberToMove);
+					// Clear fragments that we will copy over. Clear destroys and initializes the fragments, which is needed for CopyScriptStruct().
+					FragmentConfig.FragmentType->ClearScriptStruct(DyingFragmentPtr, NumberToMove);
 
-					// Swap components into the empty space just created.
-					ComponentConfig.ComponentType->CopyScriptStruct(DyingComponentPtr, MovingComponentPtr, NumberToMove);
+					// Swap fragments into the empty space just created.
+					FragmentConfig.FragmentType->CopyScriptStruct(DyingFragmentPtr, MovingFragmentPtr, NumberToMove);
 
-					// Destroy the components that were moved.
-					ComponentConfig.ComponentType->DestroyStruct(MovingComponentPtr, NumberToMove);
+					// Destroy the fragments that were moved.
+					FragmentConfig.FragmentType->DestroyStruct(MovingFragmentPtr, NumberToMove);
 				}
 			}
 
@@ -346,11 +346,11 @@ void FMassArchetypeData::BatchDestroyEntityChunks(const FArchetypeChunkCollectio
 			// just clean up the rest. Note that we explicitly do not clean the spots vacated by entities moved from 
 			// the back of the chunk - if we did the risk calling DestroyStruct on them multiple times
 			const int32 CutStartIndex = SubchunkInfo.SubchunkStart + NumberToMove;
-			for (const FMassArchetypeFragmentConfig& ComponentConfig : ComponentConfigs)
+			for (const FMassArchetypeFragmentConfig& FragmentConfig : FragmentConfigs)
 			{
-				// Destroy the component data
-				void* DyingComponentPtr = ComponentConfig.GetComponentData(Chunk.GetRawMemory(), CutStartIndex);
-				ComponentConfig.ComponentType->DestroyStruct(DyingComponentPtr, NumberToCut);
+				// Destroy the fragment data
+				void* DyingFragmentPtr = FragmentConfig.GetFragmentData(Chunk.GetRawMemory(), CutStartIndex);
+				FragmentConfig.FragmentType->DestroyStruct(DyingFragmentPtr, NumberToCut);
 			}
 		}
 
@@ -370,72 +370,72 @@ void FMassArchetypeData::BatchDestroyEntityChunks(const FArchetypeChunkCollectio
 	}
 }
 
-bool FMassArchetypeData::HasComponentDataForEntity(const UScriptStruct* ComponentType, int32 EntityIndex) const
+bool FMassArchetypeData::HasFragmentDataForEntity(const UScriptStruct* FragmentType, int32 EntityIndex) const
 {
-	return (ComponentType && CompositionDescriptor.Components.Contains(*ComponentType));
+	return (FragmentType && CompositionDescriptor.Fragments.Contains(*FragmentType));
 }
 
-void* FMassArchetypeData::GetComponentDataForEntityChecked(const UScriptStruct* ComponentType, int32 EntityIndex) const
+void* FMassArchetypeData::GetFragmentDataForEntityChecked(const UScriptStruct* FragmentType, int32 EntityIndex) const
 {
 	const FInternalEntityHandle InternalIndex = MakeEntityHandle(EntityIndex);
 	
-	// failing the below Find means given entity's archetype is missing given ComponentType
-	const int32 ComponentIndex = ComponentIndexMap.FindChecked(ComponentType);
-	return GetComponentData(ComponentIndex, InternalIndex);
+	// failing the below Find means given entity's archetype is missing given FragmentType
+	const int32 FragmentIndex = FragmentIndexMap.FindChecked(FragmentType);
+	return GetFragmentData(FragmentIndex, InternalIndex);
 }
 
-void* FMassArchetypeData::GetComponentDataForEntity(const UScriptStruct* ComponentType, int32 EntityIndex) const
+void* FMassArchetypeData::GetFragmentDataForEntity(const UScriptStruct* FragmentType, int32 EntityIndex) const
 {
-	if (const int32* ComponentIndex = ComponentIndexMap.Find(ComponentType))
+	if (const int32* FragmentIndex = FragmentIndexMap.Find(FragmentType))
 	{
 		FInternalEntityHandle InternalIndex = MakeEntityHandle(EntityIndex);
-		// failing the below Find means given entity's archetype is missing given ComponentType
-		return GetComponentData(*ComponentIndex, InternalIndex);
+		// failing the below Find means given entity's archetype is missing given FragmentType
+		return GetFragmentData(*FragmentIndex, InternalIndex);
 	}
 	return nullptr;
 }
 
-void FMassArchetypeData::SetComponentsData(const FMassEntityHandle Entity, TArrayView<const FInstancedStruct> ComponentInstances)
+void FMassArchetypeData::SetFragmentsData(const FMassEntityHandle Entity, TArrayView<const FInstancedStruct> FragmentInstances)
 {
 	FInternalEntityHandle InternalIndex = MakeEntityHandle(Entity);
 
-	for (const FInstancedStruct& Instance : ComponentInstances)
+	for (const FInstancedStruct& Instance : FragmentInstances)
 	{
-		const UScriptStruct* ComponentType = Instance.GetScriptStruct();
-		check(ComponentType);
-		const int32 ComponentIndex = ComponentIndexMap.FindChecked(ComponentType);
-		void* ComponentMemory = GetComponentData(ComponentIndex, InternalIndex);
-		// No UE::Mass::Core::bBitwiseRelocateComponents, this isn't a move component
-		ComponentType->CopyScriptStruct(ComponentMemory, Instance.GetMemory());
+		const UScriptStruct* FragmentType = Instance.GetScriptStruct();
+		check(FragmentType);
+		const int32 FragmentIndex = FragmentIndexMap.FindChecked(FragmentType);
+		void* FragmentMemory = GetFragmentData(FragmentIndex, InternalIndex);
+		// No UE::Mass::Core::bBitwiseRelocateFragments, this isn't a move fragment
+		FragmentType->CopyScriptStruct(FragmentMemory, Instance.GetMemory());
 	}
 }
 
-void FMassArchetypeData::SetComponentData(const FArchetypeChunkCollection& ChunkCollection, const FInstancedStruct& ComponentSource)
+void FMassArchetypeData::SetFragmentData(const FArchetypeChunkCollection& ChunkCollection, const FInstancedStruct& FragmentSource)
 {
-	check(ComponentSource.IsValid());
-	const UScriptStruct* ComponentType = ComponentSource.GetScriptStruct();
-	check(ComponentType);
-	const int32 ComponentIndex = ComponentIndexMap.FindChecked(ComponentType);
-	const int32 ComponentTypeSize = ComponentType->GetStructureSize();
-	const uint8* ComponentSourceMemory = ComponentSource.GetMemory();
-	check(ComponentSourceMemory);
+	check(FragmentSource.IsValid());
+	const UScriptStruct* FragmentType = FragmentSource.GetScriptStruct();
+	check(FragmentType);
+	const int32 FragmentIndex = FragmentIndexMap.FindChecked(FragmentType);
+	const int32 FragmentTypeSize = FragmentType->GetStructureSize();
+	const uint8* FragmentSourceMemory = FragmentSource.GetMemory();
+	check(FragmentSourceMemory);
 	
 	for (FMassArchetypeChunkIterator ChunkIterator(ChunkCollection); ChunkIterator; ++ChunkIterator)
 	{
-		uint8* ComponentMemory = (uint8*)ComponentConfigs[ComponentIndex].GetComponentData(Chunks[ChunkIterator->ChunkIndex].GetRawMemory(), ChunkIterator->SubchunkStart);
-		for (int i = ChunkIterator->Length; i; --i, ComponentMemory += ComponentTypeSize)
+		uint8* FragmentMemory = (uint8*)FragmentConfigs[FragmentIndex].GetFragmentData(Chunks[ChunkIterator->ChunkIndex].GetRawMemory(), ChunkIterator->SubchunkStart);
+		for (int i = ChunkIterator->Length; i; --i, FragmentMemory += FragmentTypeSize)
 		{
-			// No UE::Mass::Core::bBitwiseRelocateComponents, this isn't a move of a component
-			ComponentType->CopyScriptStruct(ComponentMemory, ComponentSourceMemory);
+			// No UE::Mass::Core::bBitwiseRelocateFragments, this isn't a move of a fragment
+			FragmentType->CopyScriptStruct(FragmentMemory, FragmentSourceMemory);
 		}
 	}
 }
 
-void FMassArchetypeData::SetDefaultChunkComponentValue(FConstStructView InstancedStruct)
+void FMassArchetypeData::SetDefaultChunkFragmentValue(FConstStructView InstancedStruct)
 {
-	if (ensure(InstancedStruct.GetScriptStruct()) && CompositionDescriptor.ChunkComponents.Contains(*InstancedStruct.GetScriptStruct()))
+	if (ensure(InstancedStruct.GetScriptStruct()) && CompositionDescriptor.ChunkFragments.Contains(*InstancedStruct.GetScriptStruct()))
 	{
-		if (FInstancedStruct* Element = ChunkComponentTemplates.FindByPredicate([&InstancedStruct](const FInstancedStruct& Element) { return Element.GetScriptStruct() == InstancedStruct.GetScriptStruct(); }))
+		if (FInstancedStruct* Element = ChunkFragmentTemplates.FindByPredicate([&InstancedStruct](const FInstancedStruct& Element) { return Element.GetScriptStruct() == InstancedStruct.GetScriptStruct(); }))
 		{
 			*Element = InstancedStruct;
 		}
@@ -451,44 +451,44 @@ void FMassArchetypeData::MoveEntityToAnotherArchetype(const FMassEntityHandle En
 	const int32 IndexWithinChunk = AbsoluteIndex % NumEntitiesPerChunk;
 	FMassArchetypeChunk& Chunk = Chunks[ChunkIndex];
 
-	constexpr bool bInitializeComponentsDuringCreation = !UE::Mass::Core::bBitwiseRelocateComponents;
-	const int32 NewAbsoluteIndex = NewArchetype.AddEntityInternal(Entity, bInitializeComponentsDuringCreation);
+	constexpr bool bInitializeFragmentsDuringCreation = !UE::Mass::Core::bBitwiseRelocateFragments;
+	const int32 NewAbsoluteIndex = NewArchetype.AddEntityInternal(Entity, bInitializeFragmentsDuringCreation);
 	const int32 NewChunkIndex = NewAbsoluteIndex / NewArchetype.NumEntitiesPerChunk;
 	const int32 NewIndexWithinChunk = NewAbsoluteIndex % NewArchetype.NumEntitiesPerChunk;
 	FMassArchetypeChunk& NewChunk = NewArchetype.Chunks[NewChunkIndex];
 
-	// for every NewArchetype's component see if it was in the old component as well and if so copy it's value. 
-	// If not then initialize the component (unless it has already been initialized based on bInitializeComponentsDuringCreation
+	// for every NewArchetype's fragment see if it was in the old fragment as well and if so copy it's value. 
+	// If not then initialize the fragment (unless it has already been initialized based on bInitializeFragmentsDuringCreation
 	// value
-	for (const FMassArchetypeFragmentConfig& NewComponentConfig : NewArchetype.ComponentConfigs)
+	for (const FMassArchetypeFragmentConfig& NewFragmentConfig : NewArchetype.FragmentConfigs)
 	{
-		const int32* OldComponentIndex = ComponentIndexMap.Find(NewComponentConfig.ComponentType);
-		void* Dst = NewComponentConfig.GetComponentData(NewChunk.GetRawMemory(), NewIndexWithinChunk);
+		const int32* OldFragmentIndex = FragmentIndexMap.Find(NewFragmentConfig.FragmentType);
+		void* Dst = NewFragmentConfig.GetFragmentData(NewChunk.GetRawMemory(), NewIndexWithinChunk);
 
-		// Only copy if the component type exists in both archetypes
-		if (OldComponentIndex)
+		// Only copy if the fragment type exists in both archetypes
+		if (OldFragmentIndex)
 		{
-			const void* Src = ComponentConfigs[*OldComponentIndex].GetComponentData(Chunk.GetRawMemory(), IndexWithinChunk);
-			if (UE::Mass::Core::bBitwiseRelocateComponents)
+			const void* Src = FragmentConfigs[*OldFragmentIndex].GetFragmentData(Chunk.GetRawMemory(), IndexWithinChunk);
+			if (UE::Mass::Core::bBitwiseRelocateFragments)
 			{
-				FMemory::Memcpy(Dst, Src, NewComponentConfig.ComponentType->GetStructureSize());
+				FMemory::Memcpy(Dst, Src, NewFragmentConfig.FragmentType->GetStructureSize());
 			}
 			else
 			{
-				NewComponentConfig.ComponentType->CopyScriptStruct(Dst, Src);
+				NewFragmentConfig.FragmentType->CopyScriptStruct(Dst, Src);
 			}
 		}
-		else if (bInitializeComponentsDuringCreation == false)
+		else if (bInitializeFragmentsDuringCreation == false)
 		{
-			// the component's unique to the NewArchetype need to be initialized
+			// the fragment's unique to the NewArchetype need to be initialized
 			// @todo we're doing it for tags here as well. A tiny bit of perf lost. Probably not worth adding a check
 			// but something to keep in mind. Will go away once tags are more of an archetype fragment than entity's
-			NewComponentConfig.ComponentType->InitializeStruct(Dst);
+			NewFragmentConfig.FragmentType->InitializeStruct(Dst);
 		}
 	}
 
-	constexpr bool bDestroyComponents = !UE::Mass::Core::bBitwiseRelocateComponents;
-	RemoveEntityInternal(AbsoluteIndex, bDestroyComponents);
+	constexpr bool bDestroyFragments = !UE::Mass::Core::bBitwiseRelocateFragments;
+	RemoveEntityInternal(AbsoluteIndex, bDestroyFragments);
 }
 
 void FMassArchetypeData::ExecuteFunction(FMassExecutionContext& RunContext, const FMassExecuteFunction& Function, const FMassQueryRequirementIndicesMapping& RequirementMapping, const FArchetypeChunkCollection& ChunkCollection)
@@ -507,8 +507,8 @@ void FMassArchetypeData::ExecuteFunction(FMassExecutionContext& RunContext, cons
 			checkf((ChunkIterator->SubchunkStart + ChunkLength) <= Chunk.GetNumInstances() && ChunkLength > 0, TEXT("Invalid subchunk, it is going over the number of instances in the chunk or it is empty."));
 
 			RunContext.SetCurrentChunkSerialModificationNumber(Chunk.GetSerialModificationNumber());
-			BindChunkComponentRequirements(RunContext, RequirementMapping.ChunkComponents, Chunk);
-			BindEntityRequirements(RunContext, RequirementMapping.EntityComponents, Chunk, ChunkIterator->SubchunkStart, ChunkLength);
+			BindChunkFragmentRequirements(RunContext, RequirementMapping.ChunkFragments, Chunk);
+			BindEntityRequirements(RunContext, RequirementMapping.EntityFragments, Chunk, ChunkIterator->SubchunkStart, ChunkLength);
 
 			Function(RunContext);
 		}
@@ -524,11 +524,11 @@ void FMassArchetypeData::ExecuteFunction(FMassExecutionContext& RunContext, cons
 		if (Chunk.GetNumInstances())
 		{
 			RunContext.SetCurrentChunkSerialModificationNumber(Chunk.GetSerialModificationNumber());
-			BindChunkComponentRequirements(RunContext, RequirementMapping.ChunkComponents, Chunk);
+			BindChunkFragmentRequirements(RunContext, RequirementMapping.ChunkFragments, Chunk);
 
 			if (!ChunkCondition || ChunkCondition(RunContext))
 			{
-				BindEntityRequirements(RunContext, RequirementMapping.EntityComponents, Chunk, 0, Chunk.GetNumInstances());
+				BindEntityRequirements(RunContext, RequirementMapping.EntityFragments, Chunk, 0, Chunk.GetNumInstances());
 				Function(RunContext);
 			}
 		}
@@ -544,11 +544,11 @@ void FMassArchetypeData::ExecutionFunctionForChunk(FMassExecutionContext RunCont
 	{
 		RunContext.SetCurrentArchetypeData(*this);
 		RunContext.SetCurrentChunkSerialModificationNumber(Chunk.GetSerialModificationNumber());
-		BindChunkComponentRequirements(RunContext, RequirementMapping.ChunkComponents, Chunk);
+		BindChunkFragmentRequirements(RunContext, RequirementMapping.ChunkFragments, Chunk);
 
 		if (!ChunkCondition || ChunkCondition(RunContext))
 		{
-			BindEntityRequirements(RunContext, RequirementMapping.EntityComponents, Chunk, ChunkInfo.SubchunkStart, ChunkLength);
+			BindEntityRequirements(RunContext, RequirementMapping.EntityFragments, Chunk, ChunkInfo.SubchunkStart, ChunkLength);
 			Function(RunContext);
 		}
 	}
@@ -604,26 +604,26 @@ void FMassArchetypeData::CompactEntities(const double TimeAllowed)
 		const int32 ToIndex = ChunkToFill->GetNumInstances();
 		check(NumberOfEntitiesToMove > 0);
 
-		for (const FMassArchetypeFragmentConfig& ComponentConfig : ComponentConfigs)
+		for (const FMassArchetypeFragmentConfig& FragmentConfig : FragmentConfigs)
 		{
-			void* FromComponentPtr = ComponentConfig.GetComponentData(ChunkToEmpty->GetRawMemory(), FromIndex);
-			void* ToComponentPtr = ComponentConfig.GetComponentData(ChunkToFill->GetRawMemory(), ToIndex);
+			void* FromFragmentPtr = FragmentConfig.GetFragmentData(ChunkToEmpty->GetRawMemory(), FromIndex);
+			void* ToFragmentPtr = FragmentConfig.GetFragmentData(ChunkToFill->GetRawMemory(), ToIndex);
 
-			if (UE::Mass::Core::bBitwiseRelocateComponents)
+			if (UE::Mass::Core::bBitwiseRelocateFragments)
 			{
 				// Move all entries
-				FMemory::Memcpy(ToComponentPtr, FromComponentPtr, ComponentConfig.ComponentType->GetStructureSize() * NumberOfEntitiesToMove);
+				FMemory::Memcpy(ToFragmentPtr, FromFragmentPtr, FragmentConfig.FragmentType->GetStructureSize() * NumberOfEntitiesToMove);
 			}
 			else
 			{
-				// Destroy & initialize the component data
-				ComponentConfig.ComponentType->ClearScriptStruct(ToComponentPtr, NumberOfEntitiesToMove);
+				// Destroy & initialize the fragment data
+				FragmentConfig.FragmentType->ClearScriptStruct(ToFragmentPtr, NumberOfEntitiesToMove);
 
 				// Copy all entries
-				ComponentConfig.ComponentType->CopyScriptStruct(ToComponentPtr, FromComponentPtr, NumberOfEntitiesToMove);
+				FragmentConfig.FragmentType->CopyScriptStruct(ToFragmentPtr, FromFragmentPtr, NumberOfEntitiesToMove);
 
 				// Destroy all entries
-				ComponentConfig.ComponentType->DestroyStruct(FromComponentPtr, NumberOfEntitiesToMove);
+				FragmentConfig.FragmentType->DestroyStruct(FromFragmentPtr, NumberOfEntitiesToMove);
 			}
 		}
 
@@ -644,69 +644,69 @@ void FMassArchetypeData::CompactEntities(const double TimeAllowed)
 	}
 }
 
-void FMassArchetypeData::GetRequirementsComponentMapping(TConstArrayView<FMassFragmentRequirement> Requirements, FMassFragmentIndicesMapping& OutComponentIndices)
+void FMassArchetypeData::GetRequirementsFragmentMapping(TConstArrayView<FMassFragmentRequirement> Requirements, FMassFragmentIndicesMapping& OutFragmentIndices)
 {
-	OutComponentIndices.Reset(Requirements.Num());
+	OutFragmentIndices.Reset(Requirements.Num());
 	for (const FMassFragmentRequirement& Requirement : Requirements)
 	{
 		if (Requirement.RequiresBinding())
 		{
-			const int32* ComponentIndex = ComponentIndexMap.Find(Requirement.StructType);
-			check(ComponentIndex != nullptr || Requirement.IsOptional());
-			OutComponentIndices.Add(ComponentIndex ? *ComponentIndex : INDEX_NONE);
+			const int32* FragmentIndex = FragmentIndexMap.Find(Requirement.StructType);
+			check(FragmentIndex != nullptr || Requirement.IsOptional());
+			OutFragmentIndices.Add(FragmentIndex ? *FragmentIndex : INDEX_NONE);
 		}
 	}
 }
 
-void FMassArchetypeData::GetRequirementsChunkComponentMapping(TConstArrayView<FMassFragmentRequirement> ChunkRequirements, FMassFragmentIndicesMapping& OutComponentIndices)
+void FMassArchetypeData::GetRequirementsChunkFragmentMapping(TConstArrayView<FMassFragmentRequirement> ChunkRequirements, FMassFragmentIndicesMapping& OutFragmentIndices)
 {
-	int32 LastFoundComponentIndex = -1;
-	OutComponentIndices.Reset(ChunkRequirements.Num());
+	int32 LastFoundFragmentIndex = -1;
+	OutFragmentIndices.Reset(ChunkRequirements.Num());
 	for (const FMassFragmentRequirement& Requirement : ChunkRequirements)
 	{
 		if (Requirement.RequiresBinding())
 		{
-			int32 ComponentIndex = INDEX_NONE;
-			for (int32 i = LastFoundComponentIndex + 1; i < ChunkComponentTemplates.Num(); ++i)
+			int32 FragmentIndex = INDEX_NONE;
+			for (int32 i = LastFoundFragmentIndex + 1; i < ChunkFragmentTemplates.Num(); ++i)
 			{
-				if (ChunkComponentTemplates[i].GetScriptStruct()->IsChildOf(Requirement.StructType))
+				if (ChunkFragmentTemplates[i].GetScriptStruct()->IsChildOf(Requirement.StructType))
 				{
-					ComponentIndex = i;
+					FragmentIndex = i;
 					break;
 				}
 			}
 
-			check(ComponentIndex != INDEX_NONE);
-			OutComponentIndices.Add(ComponentIndex);
-			LastFoundComponentIndex = ComponentIndex;
+			check(FragmentIndex != INDEX_NONE);
+			OutFragmentIndices.Add(FragmentIndex);
+			LastFoundFragmentIndex = FragmentIndex;
 		}
 	}
 }
 
-void FMassArchetypeData::BindEntityRequirements(FMassExecutionContext& RunContext, const FMassFragmentIndicesMapping& EntityComponentsMapping, FMassArchetypeChunk& Chunk, const int32 SubchunkStart, const int32 SubchunkLength)
+void FMassArchetypeData::BindEntityRequirements(FMassExecutionContext& RunContext, const FMassFragmentIndicesMapping& EntityFragmentsMapping, FMassArchetypeChunk& Chunk, const int32 SubchunkStart, const int32 SubchunkLength)
 {
 	// auto-correcting number of entities to process in case SubchunkStart +  SubchunkLength > Chunk.GetNumInstances()
 	const int32 NumEntities = SubchunkLength >= 0 ? FMath::Min(SubchunkLength, Chunk.GetNumInstances() - SubchunkStart) : Chunk.GetNumInstances();
 	check(SubchunkStart >= 0 && SubchunkStart < Chunk.GetNumInstances());
 
-	if (EntityComponentsMapping.Num() > 0)
+	if (EntityFragmentsMapping.Num() > 0)
 	{
-		check(RunContext.GetMutableRequirements().Num() == EntityComponentsMapping.Num());
+		check(RunContext.GetMutableRequirements().Num() == EntityFragmentsMapping.Num());
 
-		for (int i = 0; i < EntityComponentsMapping.Num(); ++i)
+		for (int i = 0; i < EntityFragmentsMapping.Num(); ++i)
 		{
-			FMassExecutionContext::FFragmentView& Requirement = RunContext.ComponentViews[i];
-			const int32 ComponentIndex = EntityComponentsMapping[i];
+			FMassExecutionContext::FFragmentView& Requirement = RunContext.FragmentViews[i];
+			const int32 FragmentIndex = EntityFragmentsMapping[i];
 
-			check(ComponentIndex != INDEX_NONE || Requirement.Requirement.IsOptional());
-			if (ComponentIndex != INDEX_NONE)
+			check(FragmentIndex != INDEX_NONE || Requirement.Requirement.IsOptional());
+			if (FragmentIndex != INDEX_NONE)
 			{
-				Requirement.ComponentView = TArrayView<FMassFragment>((FMassFragment*)GetComponentData(ComponentIndex, Chunk.GetRawMemory(), SubchunkStart), NumEntities);
+				Requirement.FragmentView = TArrayView<FMassFragment>((FMassFragment*)GetFragmentData(FragmentIndex, Chunk.GetRawMemory(), SubchunkStart), NumEntities);
 			}
 			else
 			{
 				// @todo this might not be needed
-				Requirement.ComponentView = TArrayView<FMassFragment>();
+				Requirement.FragmentView = TArrayView<FMassFragment>();
 			}
 		}
 	}
@@ -715,15 +715,15 @@ void FMassArchetypeData::BindEntityRequirements(FMassExecutionContext& RunContex
 		// Map in the required data arrays from the current chunk to the array views
 		for (FMassExecutionContext::FFragmentView& Requirement : RunContext.GetMutableRequirements())
 		{
-			const int32* ComponentIndex = ComponentIndexMap.Find(Requirement.Requirement.StructType);
-			check(ComponentIndex != nullptr || Requirement.Requirement.IsOptional());
-			if (ComponentIndex)
+			const int32* FragmentIndex = FragmentIndexMap.Find(Requirement.Requirement.StructType);
+			check(FragmentIndex != nullptr || Requirement.Requirement.IsOptional());
+			if (FragmentIndex)
 			{
-				Requirement.ComponentView = TArrayView<FMassFragment>((FMassFragment*)GetComponentData(*ComponentIndex, Chunk.GetRawMemory(), SubchunkStart), NumEntities);
+				Requirement.FragmentView = TArrayView<FMassFragment>((FMassFragment*)GetFragmentData(*FragmentIndex, Chunk.GetRawMemory(), SubchunkStart), NumEntities);
 			}
 			else
 			{
-				Requirement.ComponentView = TArrayView<FMassFragment>();
+				Requirement.FragmentView = TArrayView<FMassFragment>();
 			}
 		}
 	}
@@ -731,28 +731,28 @@ void FMassArchetypeData::BindEntityRequirements(FMassExecutionContext& RunContex
 	RunContext.EntityListView = TArrayView<FMassEntityHandle>(&Chunk.GetEntityArrayElementRef(EntityListOffsetWithinChunk, SubchunkStart), NumEntities);
 }
 
-void FMassArchetypeData::BindChunkComponentRequirements(FMassExecutionContext& RunContext, const FMassFragmentIndicesMapping& ChunkComponentsMapping, FMassArchetypeChunk& Chunk)
+void FMassArchetypeData::BindChunkFragmentRequirements(FMassExecutionContext& RunContext, const FMassFragmentIndicesMapping& ChunkFragmentsMapping, FMassArchetypeChunk& Chunk)
 {
-	if (ChunkComponentsMapping.Num() > 0)
+	if (ChunkFragmentsMapping.Num() > 0)
 	{
-		check(RunContext.GetMutableChunkRequirements().Num() == ChunkComponentsMapping.Num());
+		check(RunContext.GetMutableChunkRequirements().Num() == ChunkFragmentsMapping.Num());
 
-		for (int i = 0; i < ChunkComponentsMapping.Num(); ++i)
+		for (int i = 0; i < ChunkFragmentsMapping.Num(); ++i)
 		{
-			FMassExecutionContext::FChunkFragmentView& ChunkRequirement = RunContext.ChunkComponents[i];
-			const int32 ChunkComponentIndex = ChunkComponentsMapping[i];
+			FMassExecutionContext::FChunkFragmentView& ChunkRequirement = RunContext.ChunkFragments[i];
+			const int32 ChunkFragmentIndex = ChunkFragmentsMapping[i];
 
-			check(ChunkComponentIndex != INDEX_NONE);
-			ChunkRequirement.ChunkComponentView = Chunk.GetMutableChunkComponentViewChecked(ChunkComponentIndex);
+			check(ChunkFragmentIndex != INDEX_NONE);
+			ChunkRequirement.ChunkFragmentView = Chunk.GetMutableChunkFragmentViewChecked(ChunkFragmentIndex);
 		}
 	}
 	else
 	{
 		for (FMassExecutionContext::FChunkFragmentView& ChunkRequirement : RunContext.GetMutableChunkRequirements())
 		{
-			FInstancedStruct* ChunkComponentInstance = Chunk.FindMutableChunkComponent(ChunkRequirement.Requirement.StructType);
-			check(ChunkComponentInstance != nullptr);
-			ChunkRequirement.ChunkComponentView = FStructView(*ChunkComponentInstance);
+			FInstancedStruct* ChunkFragmentInstance = Chunk.FindMutableChunkFragment(ChunkRequirement.Requirement.StructType);
+			check(ChunkFragmentInstance != nullptr);
+			ChunkRequirement.ChunkFragmentView = FStructView(*ChunkFragmentInstance);
 		}
 	}
 }
@@ -760,11 +760,11 @@ void FMassArchetypeData::BindChunkComponentRequirements(FMassExecutionContext& R
 SIZE_T FMassArchetypeData::GetAllocatedSize() const
 {
 	return sizeof(FMassArchetypeData) +
-		ComponentConfigs.GetAllocatedSize() +
+		FragmentConfigs.GetAllocatedSize() +
 		Chunks.GetAllocatedSize() +
 		(Chunks.Num() * GetChunkAllocSize()) +
 		EntityMap.GetAllocatedSize() +
-		ComponentIndexMap.GetAllocatedSize();
+		FragmentIndexMap.GetAllocatedSize();
 }
 
 FString FMassArchetypeData::DebugGetDescription() const
@@ -775,13 +775,13 @@ FString FMassArchetypeData::DebugGetDescription() const
 	ArchetypeDebugName.Append(TEXT("<"));
 
 	bool bNeedsComma = false;
-	for (const FMassArchetypeFragmentConfig& ComponentConfig : ComponentConfigs)
+	for (const FMassArchetypeFragmentConfig& FragmentConfig : FragmentConfigs)
 	{
 		if (bNeedsComma)
 		{
 			ArchetypeDebugName.Append(TEXT(","));
 		}
-		ArchetypeDebugName.Append(ComponentConfig.ComponentType->GetName());
+		ArchetypeDebugName.Append(FragmentConfig.FragmentType->GetName());
 		bNeedsComma = true;
 	}
 
@@ -798,17 +798,17 @@ void FMassArchetypeData::DebugPrintArchetype(FOutputDevice& Ar)
 	FStringOutputDevice TagsDecription;
 	CompositionDescriptor.Tags.DebugGetStringDesc(TagsDecription);
 	Ar.Logf(ELogVerbosity::Log, TEXT("Tags: %s"), *TagsDecription);
-	Ar.Logf(ELogVerbosity::Log, TEXT("Components: %s"), *DebugGetDescription());
+	Ar.Logf(ELogVerbosity::Log, TEXT("Fragments: %s"), *DebugGetDescription());
 	Ar.Logf(ELogVerbosity::Log, TEXT("\tChunks: %d x %d KB = %d KB total"), Chunks.Num(), GetChunkAllocSize() / 1024, (GetChunkAllocSize()*Chunks.Num()) / 1024);
 	
-	int ChunkWithComponentsCount = 0;
+	int ChunkWithFragmentsCount = 0;
 	for (FMassArchetypeChunk& Chunk : Chunks)
 	{
-		ChunkWithComponentsCount += Chunk.DebugGetChunkComponentCount() > 0 ? 1 : 0;
+		ChunkWithFragmentsCount += Chunk.DebugGetChunkFragmentCount() > 0 ? 1 : 0;
 	}
-	if (ChunkWithComponentsCount)
+	if (ChunkWithFragmentsCount)
 	{
-		Ar.Logf(ELogVerbosity::Log, TEXT("\tChunks with components: %d"), ChunkWithComponentsCount);
+		Ar.Logf(ELogVerbosity::Log, TEXT("\tChunks with fragments: %d"), ChunkWithFragmentsCount);
 	}
 
 	const int32 CurrentEntityCapacity = Chunks.Num() * NumEntitiesPerChunk;
@@ -837,10 +837,10 @@ void FMassArchetypeData::DebugPrintArchetype(FOutputDevice& Ar)
 
 	Ar.Logf(ELogVerbosity::Log, TEXT("\tOffset 0x%04X: Entity[] (%d bytes each)"), EntityListOffsetWithinChunk, sizeof(FMassEntityHandle));
 	int32 TotalBytesOfValidData = sizeof(FMassEntityHandle) * NumEntitiesPerChunk;
-	for (const FMassArchetypeFragmentConfig& ComponentConfig : ComponentConfigs)
+	for (const FMassArchetypeFragmentConfig& FragmentConfig : FragmentConfigs)
 	{
-		TotalBytesOfValidData += ComponentConfig.ComponentType->GetStructureSize() * NumEntitiesPerChunk;
-		Ar.Logf(ELogVerbosity::Log, TEXT("\tOffset 0x%04X: %s[] (%d bytes each)"), ComponentConfig.ArrayOffsetWithinChunk, *ComponentConfig.ComponentType->GetName(), ComponentConfig.ComponentType->GetStructureSize());
+		TotalBytesOfValidData += FragmentConfig.FragmentType->GetStructureSize() * NumEntitiesPerChunk;
+		Ar.Logf(ELogVerbosity::Log, TEXT("\tOffset 0x%04X: %s[] (%d bytes each)"), FragmentConfig.ArrayOffsetWithinChunk, *FragmentConfig.FragmentType->GetName(), FragmentConfig.FragmentType->GetStructureSize());
 	}
 
 	//@TODO: Print out padding in between things?
@@ -860,28 +860,28 @@ void FMassArchetypeData::DebugPrintArchetype(FOutputDevice& Ar)
 
 void FMassArchetypeData::DebugPrintEntity(FMassEntityHandle Entity, FOutputDevice& Ar, const TCHAR* InPrefix) const
 {
-	for (const FMassArchetypeFragmentConfig& ComponentConfig : ComponentConfigs)
+	for (const FMassArchetypeFragmentConfig& FragmentConfig : FragmentConfigs)
 	{
-		void* Data = GetComponentDataForEntityChecked(ComponentConfig.ComponentType, Entity.Index);
+		void* Data = GetFragmentDataForEntityChecked(FragmentConfig.FragmentType, Entity.Index);
 		
-		FString ComponentName = ComponentConfig.ComponentType->GetName();
-		ComponentName.RemoveFromStart(InPrefix);
+		FString FragmentName = FragmentConfig.FragmentType->GetName();
+		FragmentName.RemoveFromStart(InPrefix);
 
 		FString ValueStr;
-		ComponentConfig.ComponentType->ExportText(ValueStr, Data, /*Default*/nullptr, /*OwnerObject*/nullptr, EPropertyPortFlags::PPF_IncludeTransient, /*ExportRootScope*/nullptr);
+		FragmentConfig.FragmentType->ExportText(ValueStr, Data, /*Default*/nullptr, /*OwnerObject*/nullptr, EPropertyPortFlags::PPF_IncludeTransient, /*ExportRootScope*/nullptr);
 
-		Ar.Logf(TEXT("%s: %s"), *ComponentName, *ValueStr);
+		Ar.Logf(TEXT("%s: %s"), *FragmentName, *ValueStr);
 	}
 }
 
 #endif // WITH_MASSENTITY_DEBUG
 
-void FMassArchetypeData::REMOVEME_GetArrayViewForComponentInChunk(int32 ChunkIndex, const UScriptStruct* ComponentType, void*& OutChunkBase, int32& OutNumEntities)
+void FMassArchetypeData::REMOVEME_GetArrayViewForFragmentInChunk(int32 ChunkIndex, const UScriptStruct* FragmentType, void*& OutChunkBase, int32& OutNumEntities)
 {
 	const FMassArchetypeChunk& Chunk = Chunks[ChunkIndex];
-	const int32 ComponentIndex = ComponentIndexMap.FindChecked(ComponentType);
+	const int32 FragmentIndex = FragmentIndexMap.FindChecked(FragmentType);
 
-	OutChunkBase = ComponentConfigs[ComponentIndex].GetComponentData(Chunk.GetRawMemory(), 0);
+	OutChunkBase = FragmentConfigs[FragmentIndex].GetFragmentData(Chunk.GetRawMemory(), 0);
 	OutNumEntities = Chunk.GetNumInstances();
 }
 
diff --git a/Engine/Plugins/Runtime/MassEntity/Source/MassEntity/Private/MassArchetypeData.h b/Engine/Plugins/Runtime/MassEntity/Source/MassEntity/Private/MassArchetypeData.h
--- a/Engine/Plugins/Runtime/MassEntity/Source/MassEntity/Private/MassArchetypeData.h
+++ b/Engine/Plugins/Runtime/MassEntity/Source/MassEntity/Private/MassArchetypeData.h
@@ -17,11 +17,11 @@ private:
 	uint8* RawMemory = nullptr;
 	int32 NumInstances = 0;
 	int32 SerialModificationNumber = 0;
-	TArray<FInstancedStruct> ChunkComponentData;
+	TArray<FInstancedStruct> ChunkFragmentData;
 
 public:
-	explicit FMassArchetypeChunk(int32 AllocSize, TConstArrayView<FInstancedStruct> InChunkComponentTemplates)
-		: ChunkComponentData(InChunkComponentTemplates)
+	explicit FMassArchetypeChunk(int32 AllocSize, TConstArrayView<FInstancedStruct> InChunkFragmentTemplates)
+		: ChunkFragmentData(InChunkFragmentTemplates)
 	{
 		RawMemory = (uint8*)FMemory::Malloc(AllocSize);
 	}
@@ -75,51 +75,51 @@ public:
 		return SerialModificationNumber;
 	}
 
-	FStructView GetMutableChunkComponentViewChecked(const int32 Index) { return FStructView(ChunkComponentData[Index]); }
+	FStructView GetMutableChunkFragmentViewChecked(const int32 Index) { return FStructView(ChunkFragmentData[Index]); }
 
-	FInstancedStruct* FindMutableChunkComponent(const UScriptStruct* Type)
+	FInstancedStruct* FindMutableChunkFragment(const UScriptStruct* Type)
 	{
-		return ChunkComponentData.FindByPredicate([Type](const FInstancedStruct& Element)
+		return ChunkFragmentData.FindByPredicate([Type](const FInstancedStruct& Element)
 			{
 				return Element.GetScriptStruct()->IsChildOf(Type);
 			});
 	}
 
-	void Recycle(TConstArrayView<FInstancedStruct> InChunkComponentsTemplate)
+	void Recycle(TConstArrayView<FInstancedStruct> InChunkFragmentsTemplate)
 	{
 		checkf(NumInstances == 0, TEXT("Recycling a chunk that is not empty."));
 		SerialModificationNumber++;
-		ChunkComponentData = InChunkComponentsTemplate;
+		ChunkFragmentData = InChunkFragmentsTemplate;
 	}
 
 #if WITH_MASSENTITY_DEBUG
-	int32 DebugGetChunkComponentCount() const { return ChunkComponentData.Num(); }
+	int32 DebugGetChunkFragmentCount() const { return ChunkFragmentData.Num(); }
 #endif // WITH_MASSENTITY_DEBUG
 };
 
-// Information for a single component type in an archetype
+// Information for a single fragment type in an archetype
 struct FMassArchetypeFragmentConfig
 {
-	const UScriptStruct* ComponentType = nullptr;
+	const UScriptStruct* FragmentType = nullptr;
 	int32 ArrayOffsetWithinChunk = 0;
 
-	void* GetComponentData(uint8* ChunkBase, int32 IndexWithinChunk) const
+	void* GetFragmentData(uint8* ChunkBase, int32 IndexWithinChunk) const
 	{
-		return ChunkBase + ArrayOffsetWithinChunk + (IndexWithinChunk * ComponentType->GetStructureSize());
+		return ChunkBase + ArrayOffsetWithinChunk + (IndexWithinChunk * FragmentType->GetStructureSize());
 	}
 };
 
-// An archetype is defined by a collection of unique component types (no duplicates).
-// Order doesn't matter, there will only ever be one FMassArchetypeData per unique set of component types per entity manager subsystem
+// An archetype is defined by a collection of unique fragment types (no duplicates).
+// Order doesn't matter, there will only ever be one FMassArchetypeData per unique set of fragment types per entity manager subsystem
 struct FMassArchetypeData
 {
 private:
-	// One-stop-shop variable describing the archetype's component and tag composition 
+	// One-stop-shop variable describing the archetype's fragment and tag composition 
 	FMassCompositionDescriptor CompositionDescriptor;
 
-	TArray<FInstancedStruct> ChunkComponentTemplates;
+	TArray<FInstancedStruct> ChunkFragmentTemplates;
 
-	TArray<FMassArchetypeFragmentConfig, TInlineAllocator<16>> ComponentConfigs;
+	TArray<FMassArchetypeFragmentConfig, TInlineAllocator<16>> FragmentConfigs;
 	
 	TArray<FMassArchetypeChunk> Chunks;
 
@@ -128,7 +128,7 @@ private:
 	// of loss of encapsulation and extra complexity during archetype changes
 	TMap<int32, int32> EntityMap;
 	
-	TMap<const UScriptStruct*, int32> ComponentIndexMap;
+	TMap<const UScriptStruct*, int32> FragmentIndexMap;
 
 	int32 NumEntitiesPerChunk;
 	int32 TotalBytesPerEntity;
@@ -138,27 +138,27 @@ private:
 	friend FArchetypeChunkCollection;
 
 public:
-	TConstArrayView<FMassArchetypeFragmentConfig> GetComponentConfigs() const { return ComponentConfigs; }
-	const FMassFragmentBitSet& GetComponentBitSet() const { return CompositionDescriptor.Components; }
+	TConstArrayView<FMassArchetypeFragmentConfig> GetFragmentConfigs() const { return FragmentConfigs; }
+	const FMassFragmentBitSet& GetFragmentBitSet() const { return CompositionDescriptor.Fragments; }
 	const FMassTagBitSet& GetTagBitSet() const { return CompositionDescriptor.Tags; }
-	const FMassChunkFragmentBitSet& GetChunkComponentBitSet() const { return CompositionDescriptor.ChunkComponents; }
+	const FMassChunkFragmentBitSet& GetChunkFragmentBitSet() const { return CompositionDescriptor.ChunkFragments; }
 	const FMassCompositionDescriptor& GetCompositionDescriptor() const { return CompositionDescriptor; }
 
-	/** Method to iterate on all the component types */
-	void ForEachComponentType(TFunction< void(const UScriptStruct* /*ComponentType*/)> Function) const;
-	bool HasComponentType(const UScriptStruct* ComponentType) const;
-	bool HasTagType(const UScriptStruct* ComponentType) const { check(ComponentType); return CompositionDescriptor.Tags.Contains(*ComponentType); }
+	/** Method to iterate on all the fragment types */
+	void ForEachFragmentType(TFunction< void(const UScriptStruct* /*FragmentType*/)> Function) const;
+	bool HasFragmentType(const UScriptStruct* FragmentType) const;
+	bool HasTagType(const UScriptStruct* FragmentType) const { check(FragmentType); return CompositionDescriptor.Tags.Contains(*FragmentType); }
 	bool IsEquivalent(TConstArrayView<const UScriptStruct*> OtherList) const;
-	bool IsEquivalent(const FMassFragmentBitSet& InComponentBitSet, const FMassTagBitSet& InTagBitSet, const FMassChunkFragmentBitSet& InChunkComponentsBitSet) const
+	bool IsEquivalent(const FMassFragmentBitSet& InFragmentBitSet, const FMassTagBitSet& InTagBitSet, const FMassChunkFragmentBitSet& InChunkFragmentsBitSet) const
 	{ 
-		return CompositionDescriptor.IsEquivalent(InComponentBitSet, InTagBitSet, InChunkComponentsBitSet);
+		return CompositionDescriptor.IsEquivalent(InFragmentBitSet, InTagBitSet, InChunkFragmentsBitSet);
 	}
 	bool IsEquivalent(const FMassCompositionDescriptor& OtherCompositionDescriptor) const
 	{
 		return CompositionDescriptor.IsEquivalent(OtherCompositionDescriptor);
 	}
 
-	void Initialize(const FMassFragmentBitSet& Components, const FMassTagBitSet& Tags, const FMassChunkFragmentBitSet& ChunkComponents);
+	void Initialize(const FMassFragmentBitSet& Fragments, const FMassTagBitSet& Tags, const FMassChunkFragmentBitSet& ChunkFragments);
 
 	/** 
 	 * A special way of initializing an archetype resulting in a copy of SiblingArchetype's setup with OverrideTags
@@ -170,9 +170,9 @@ public:
 	void RemoveEntity(FMassEntityHandle Entity);
 	void BatchDestroyEntityChunks(const FArchetypeChunkCollection& ChunkCollection, TArray<FMassEntityHandle>& OutEntitiesRemoved);
 
-	bool HasComponentDataForEntity(const UScriptStruct* ComponentType, int32 EntityIndex) const;
-	void* GetComponentDataForEntityChecked(const UScriptStruct* ComponentType, int32 EntityIndex) const;
-	void* GetComponentDataForEntity(const UScriptStruct* ComponentType, int32 EntityIndex) const;
+	bool HasFragmentDataForEntity(const UScriptStruct* FragmentType, int32 EntityIndex) const;
+	void* GetFragmentDataForEntityChecked(const UScriptStruct* FragmentType, int32 EntityIndex) const;
+	void* GetFragmentDataForEntity(const UScriptStruct* FragmentType, int32 EntityIndex) const;
 
 	FORCEINLINE int32 GetInternalIndexForEntity(const int32 EntityIndex) const { return EntityMap.FindChecked(EntityIndex); }
 	int32 GetNumEntitiesPerChunk() const { return NumEntitiesPerChunk; }
@@ -194,34 +194,34 @@ public:
 	void CompactEntities(const double TimeAllowed);
 
 	/**
-	 * Moves the entity from this archetype to another, will only copy all matching component types
+	 * Moves the entity from this archetype to another, will only copy all matching fragment types
 	 * @param Entity is the entity to move
 	 * @param NewArchetype the archetype to move to
 	 */
 	void MoveEntityToAnotherArchetype(const FMassEntityHandle Entity, FMassArchetypeData& NewArchetype);
 
 	/**
-	 * Set all component sources data on specified entity, will check if there are component sources type that does not exist in the archetype
-	 * @param Entity is the entity to set the data of all components
-	 * @param ComponentSources are the components to copy the data from
+	 * Set all fragment sources data on specified entity, will check if there are fragment sources type that does not exist in the archetype
+	 * @param Entity is the entity to set the data of all fragments
+	 * @param FragmentSources are the fragments to copy the data from
 	 */
-	void SetComponentsData(const FMassEntityHandle Entity, TArrayView<const FInstancedStruct> ComponentSources);
+	void SetFragmentsData(const FMassEntityHandle Entity, TArrayView<const FInstancedStruct> FragmentSources);
 
-	/** For all entities indicated by ChunkCollection the function sets the value of component of type
-	 *  ComponentSource.GetScriptStruct to the value represented by ComponentSource.GetMemory */
-	void SetComponentData(const FArchetypeChunkCollection& ChunkCollection, const FInstancedStruct& ComponentSource);
+	/** For all entities indicated by ChunkCollection the function sets the value of fragment of type
+	 *  FragmentSource.GetScriptStruct to the value represented by FragmentSource.GetMemory */
+	void SetFragmentData(const FArchetypeChunkCollection& ChunkCollection, const FInstancedStruct& FragmentSource);
 
-	void SetDefaultChunkComponentValue(FConstStructView InstancedStruct);
+	void SetDefaultChunkFragmentValue(FConstStructView InstancedStruct);
 
-	/** Returns conversion from given Requirements to archetype's component indices */
-	void GetRequirementsComponentMapping(TConstArrayView<FMassFragmentRequirement> Requirements, FMassFragmentIndicesMapping& OutComponentIndices);
+	/** Returns conversion from given Requirements to archetype's fragment indices */
+	void GetRequirementsFragmentMapping(TConstArrayView<FMassFragmentRequirement> Requirements, FMassFragmentIndicesMapping& OutFragmentIndices);
 
-	/** Returns conversion from given ChunkRequirements to archetype's chunk component indices */
-	void GetRequirementsChunkComponentMapping(TConstArrayView<FMassFragmentRequirement> ChunkRequirements, FMassFragmentIndicesMapping& OutComponentIndices);
+	/** Returns conversion from given ChunkRequirements to archetype's chunk fragment indices */
+	void GetRequirementsChunkFragmentMapping(TConstArrayView<FMassFragmentRequirement> ChunkRequirements, FMassFragmentIndicesMapping& OutFragmentIndices);
 
 	SIZE_T GetAllocatedSize() const;
 
-	// Converts the list of components into a user-readable debug string
+	// Converts the list of fragments into a user-readable debug string
 	FString DebugGetDescription() const;
 
 #if WITH_MASSENTITY_DEBUG
@@ -231,24 +231,24 @@ public:
 	void DebugPrintArchetype(FOutputDevice& Ar);
 
 	/**
-	 * Prints out component's values for the specified entity. 
-	 * @param Entity The entity for which we want to print component values
+	 * Prints out fragment's values for the specified entity. 
+	 * @param Entity The entity for which we want to print fragment values
 	 * @param Ar The output device
-	 * @param InPrefix Optional prefix to remove from component names
+	 * @param InPrefix Optional prefix to remove from fragment names
 	 */
 	void DebugPrintEntity(FMassEntityHandle Entity, FOutputDevice& Ar, const TCHAR* InPrefix = TEXT("")) const;
 #endif // WITH_MASSENTITY_DEBUG
 
-	void REMOVEME_GetArrayViewForComponentInChunk(int32 ChunkIndex, const UScriptStruct* ComponentType, void*& OutChunkBase, int32& OutNumEntities);
+	void REMOVEME_GetArrayViewForFragmentInChunk(int32 ChunkIndex, const UScriptStruct* FragmentType, void*& OutChunkBase, int32& OutNumEntities);
 
 	//////////////////////////////////////////////////////////////////////
 	// low level api
-	FORCEINLINE const int32* GetComponentIndex(const UScriptStruct* ComponentType) const { return ComponentIndexMap.Find(ComponentType); }
-	FORCEINLINE int32 GetComponentIndexChecked(const UScriptStruct* ComponentType) const { return ComponentIndexMap.FindChecked(ComponentType); }
+	FORCEINLINE const int32* GetFragmentIndex(const UScriptStruct* FragmentType) const { return FragmentIndexMap.Find(FragmentType); }
+	FORCEINLINE int32 GetFragmentIndexChecked(const UScriptStruct* FragmentType) const { return FragmentIndexMap.FindChecked(FragmentType); }
 
-	FORCEINLINE void* GetComponentData(const int32 ComponentIndex, const FInternalEntityHandle EntityIndex) const
+	FORCEINLINE void* GetFragmentData(const int32 FragmentIndex, const FInternalEntityHandle EntityIndex) const
 	{
-		return ComponentConfigs[ComponentIndex].GetComponentData(EntityIndex.ChunkRawMemory, EntityIndex.IndexWithinChunk);
+		return FragmentConfigs[FragmentIndex].GetFragmentData(EntityIndex.ChunkRawMemory, EntityIndex.IndexWithinChunk);
 	}
 
 	FORCEINLINE FInternalEntityHandle MakeEntityHandle(int32 EntityIndex) const
@@ -264,18 +264,18 @@ public:
 		return MakeEntityHandle(Entity.Index); 
 	}
 
-	bool IsInitialized() const { return TotalBytesPerEntity > 0 && ComponentConfigs.IsEmpty() == false; }
+	bool IsInitialized() const { return TotalBytesPerEntity > 0 && FragmentConfigs.IsEmpty() == false; }
 
 protected:
-	FORCEINLINE void* GetComponentData(const int32 ComponentIndex, uint8* ChunkRawMemory, const int32 IndexWithinChunk) const
+	FORCEINLINE void* GetFragmentData(const int32 FragmentIndex, uint8* ChunkRawMemory, const int32 IndexWithinChunk) const
 	{
-		return ComponentConfigs[ComponentIndex].GetComponentData(ChunkRawMemory, IndexWithinChunk);
+		return FragmentConfigs[FragmentIndex].GetFragmentData(ChunkRawMemory, IndexWithinChunk);
 	}
 
-	void BindEntityRequirements(FMassExecutionContext& RunContext, const FMassFragmentIndicesMapping& EntityComponentsMapping, FMassArchetypeChunk& Chunk, const int32 SubchunkStart, const int32 SubchunkLength);
-	void BindChunkComponentRequirements(FMassExecutionContext& RunContext, const FMassFragmentIndicesMapping& ChunkComponentsMapping, FMassArchetypeChunk& Chunk);
+	void BindEntityRequirements(FMassExecutionContext& RunContext, const FMassFragmentIndicesMapping& EntityFragmentsMapping, FMassArchetypeChunk& Chunk, const int32 SubchunkStart, const int32 SubchunkLength);
+	void BindChunkFragmentRequirements(FMassExecutionContext& RunContext, const FMassFragmentIndicesMapping& ChunkFragmentsMapping, FMassArchetypeChunk& Chunk);
 
 private:
-	int32 AddEntityInternal(FMassEntityHandle Entity, const bool bInitializeComponents);
-	void RemoveEntityInternal(const int32 AbsoluteIndex, const bool bDestroyComponents);
+	int32 AddEntityInternal(FMassEntityHandle Entity, const bool bInitializeFragments);
+	void RemoveEntityInternal(const int32 AbsoluteIndex, const bool bDestroyFragments);
 };
diff --git a/Engine/Plugins/Runtime/MassEntity/Source/MassEntity/Private/MassCommandBuffer.cpp b/Engine/Plugins/Runtime/MassEntity/Source/MassEntity/Private/MassCommandBuffer.cpp
--- a/Engine/Plugins/Runtime/MassEntity/Source/MassEntity/Private/MassCommandBuffer.cpp
+++ b/Engine/Plugins/Runtime/MassEntity/Source/MassEntity/Private/MassCommandBuffer.cpp
@@ -119,9 +119,9 @@ void FMassCommandBuffer::MoveAppend(FMassCommandBuffer& Other)
 //////////////////////////////////////////////////////////////////////
 // Command implementations
 
-void FCommandAddComponentInstance::Execute(UMassEntitySubsystem& System) const
+void FCommandAddFragmentInstance::Execute(UMassEntitySubsystem& System) const
 {
-	System.AddComponentInstanceListToEntity(TargetEntity, MakeArrayView(&Struct, 1));
+	System.AddFragmentInstanceListToEntity(TargetEntity, MakeArrayView(&Struct, 1));
 }
 
 void FCommandSwapTags::Execute(UMassEntitySubsystem& System) const
@@ -145,34 +145,34 @@ void FCommandSwapTags::Execute(UMassEntitySubsystem& System) const
 	}
 }
 
-void FBuildEntityFromComponentInstance::Execute(UMassEntitySubsystem& System) const
+void FBuildEntityFromFragmentInstance::Execute(UMassEntitySubsystem& System) const
 {
 	System.BuildEntity(TargetEntity, MakeArrayView(&Struct, 1));
 }
 
-void FBuildEntityFromComponentInstances::Execute(UMassEntitySubsystem& System) const
+void FBuildEntityFromFragmentInstances::Execute(UMassEntitySubsystem& System) const
 {
 	System.BuildEntity(TargetEntity, Instances);
 }
 
-void FCommandAddComponent::Execute(UMassEntitySubsystem& System) const
+void FCommandAddFragment::Execute(UMassEntitySubsystem& System) const
 {
-	System.AddComponentToEntity(TargetEntity, StructParam);
+	System.AddFragmentToEntity(TargetEntity, StructParam);
 }
 
-void FCommandRemoveComponent::Execute(UMassEntitySubsystem& System) const
+void FCommandRemoveFragment::Execute(UMassEntitySubsystem& System) const
 {
-	System.RemoveComponentFromEntity(TargetEntity, StructParam);
+	System.RemoveFragmentFromEntity(TargetEntity, StructParam);
 }
 
-void FCommandAddComponentList::Execute(UMassEntitySubsystem& System) const
+void FCommandAddFragmentList::Execute(UMassEntitySubsystem& System) const
 {
-	System.AddComponentListToEntity(TargetEntity, ComponentList);
+	System.AddFragmentListToEntity(TargetEntity, FragmentList);
 }
 
-void FCommandRemoveComponentList::Execute(UMassEntitySubsystem& System) const
+void FCommandRemoveFragmentList::Execute(UMassEntitySubsystem& System) const
 {
-	System.RemoveComponentListFromEntity(TargetEntity, ComponentList);
+	System.RemoveFragmentListFromEntity(TargetEntity, FragmentList);
 }
 
 void FCommandAddTag::Execute(UMassEntitySubsystem & System) const 
diff --git a/Engine/Plugins/Runtime/MassEntity/Source/MassEntity/Private/MassEntityDebug.cpp b/Engine/Plugins/Runtime/MassEntity/Source/MassEntity/Private/MassEntityDebug.cpp
--- a/Engine/Plugins/Runtime/MassEntity/Source/MassEntity/Private/MassEntityDebug.cpp
+++ b/Engine/Plugins/Runtime/MassEntity/Source/MassEntity/Private/MassEntityDebug.cpp
@@ -11,7 +11,7 @@ DEFINE_ENUM_TO_STRING(EMassProcessingPhase);
 namespace UE::Mass::Debug
 {
 
-FString DebugGetComponentAccessString(EMassFragmentAccess Access)
+FString DebugGetFragmentAccessString(EMassFragmentAccess Access)
 {
 #if WITH_MASSENTITY_DEBUG
 	switch (Access)
diff --git a/Engine/Plugins/Runtime/MassEntity/Source/MassEntity/Private/MassEntityQuery.cpp b/Engine/Plugins/Runtime/MassEntity/Source/MassEntity/Private/MassEntityQuery.cpp
--- a/Engine/Plugins/Runtime/MassEntity/Source/MassEntity/Private/MassEntityQuery.cpp
+++ b/Engine/Plugins/Runtime/MassEntity/Source/MassEntity/Private/MassEntityQuery.cpp
@@ -20,18 +20,18 @@ FMassEntityQuery::FMassEntityQuery()
 FMassEntityQuery::FMassEntityQuery(std::initializer_list<UScriptStruct*> InitList)
 	: FMassEntityQuery()
 {
-	for (const UScriptStruct* ComponentType : InitList)
+	for (const UScriptStruct* FragmentType : InitList)
 	{
-		AddRequirement(ComponentType, EMassFragmentAccess::ReadWrite, EMassFragmentPresence::All);
+		AddRequirement(FragmentType, EMassFragmentAccess::ReadWrite, EMassFragmentPresence::All);
 	}
 }
 
 FMassEntityQuery::FMassEntityQuery(TConstArrayView<const UScriptStruct*> InitList)
 	: FMassEntityQuery()
 {
-	for (const UScriptStruct* ComponentType : InitList)
+	for (const UScriptStruct* FragmentType : InitList)
 	{
-		AddRequirement(ComponentType, EMassFragmentAccess::ReadWrite, EMassFragmentPresence::All);
+		AddRequirement(FragmentType, EMassFragmentAccess::ReadWrite, EMassFragmentPresence::All);
 	}
 }
 
@@ -46,10 +46,10 @@ void FMassEntityQuery::ReadCommandlineParams()
 
 void FMassEntityQuery::SortRequirements()
 {
-	// we're sorting the Requirements the same way ArchetypeData's ComponentConfig is sorted (see FMassArchetypeData::Initialize)
-	// so that when we access ArchetypeData.ComponentConfigs in FMassArchetypeData::BindRequirementsWithMapping
-	// (via GetComponentData call) the access is sequential (i.e. not random) and there's a higher chance the memory
-	// ComponentConfigs we want to access have already been fetched and are available in processor cache.
+	// we're sorting the Requirements the same way ArchetypeData's FragmentConfig is sorted (see FMassArchetypeData::Initialize)
+	// so that when we access ArchetypeData.FragmentConfigs in FMassArchetypeData::BindRequirementsWithMapping
+	// (via GetFragmentData call) the access is sequential (i.e. not random) and there's a higher chance the memory
+	// FragmentConfigs we want to access have already been fetched and are available in processor cache.
 	Requirements.Sort(FMassSorterOperator<FMassFragmentRequirement>());
 	ChunkRequirements.Sort(FMassSorterOperator<FMassFragmentRequirement>());
 }
@@ -70,14 +70,14 @@ void FMassEntityQuery::CacheArchetypes(UMassEntitySubsystem& InEntitySubsystem)
 
 			TRACE_CPUPROFILER_EVENT_SCOPE_STR("Pipe RequirementsBinding")
 			const TConstArrayView<FMassFragmentRequirement> LocalRequirements = GetRequirements();
-			ArchetypeComponentMapping.Reset(ValidArchetypes.Num());
-			ArchetypeComponentMapping.AddDefaulted(ValidArchetypes.Num());
+			ArchetypeFragmentMapping.Reset(ValidArchetypes.Num());
+			ArchetypeFragmentMapping.AddDefaulted(ValidArchetypes.Num());
 			for (int i = 0; i < ValidArchetypes.Num(); ++i)
 			{
-				ValidArchetypes[i].DataPtr->GetRequirementsComponentMapping(LocalRequirements, ArchetypeComponentMapping[i].EntityComponents);
+				ValidArchetypes[i].DataPtr->GetRequirementsFragmentMapping(LocalRequirements, ArchetypeFragmentMapping[i].EntityFragments);
 				if (ChunkRequirements.Num())
 				{
-					ValidArchetypes[i].DataPtr->GetRequirementsChunkComponentMapping(ChunkRequirements, ArchetypeComponentMapping[i].ChunkComponents);
+					ValidArchetypes[i].DataPtr->GetRequirementsChunkFragmentMapping(ChunkRequirements, ArchetypeFragmentMapping[i].ChunkFragments);
 				}
 			}
 		}
@@ -90,7 +90,7 @@ void FMassEntityQuery::CacheArchetypes(UMassEntitySubsystem& InEntitySubsystem)
 
 bool FMassEntityQuery::CheckValidity() const
 {
-	return RequiredAllComponents.IsEmpty() == false || RequiredAnyComponents.IsEmpty() == false || RequiredOptionalComponents.IsEmpty() == false;
+	return RequiredAllFragments.IsEmpty() == false || RequiredAnyFragments.IsEmpty() == false || RequiredOptionalFragments.IsEmpty() == false;
 }
 
 bool FMassEntityQuery::DoesArchetypeMatchRequirements(const FArchetypeHandle& ArchetypeHandle) const
@@ -108,9 +108,9 @@ bool FMassEntityQuery::DoesArchetypeMatchRequirements(const FArchetypeHandle& Ar
 			continue;
 		}
 
-		const bool bNeedsComponent = (Requirement.Presence == EMassFragmentPresence::All);
-		const bool bHasComponent = Archetype->HasComponentType(Requirement.StructType);
-		if (bNeedsComponent != bHasComponent)
+		const bool bNeedsFragment = (Requirement.Presence == EMassFragmentPresence::All);
+		const bool bHasFragment = Archetype->HasFragmentType(Requirement.StructType);
+		if (bNeedsFragment != bHasFragment)
 		{
 			return false;
 		}
@@ -158,8 +158,8 @@ void FMassEntityQuery::ForEachEntityChunk(UMassEntitySubsystem& EntitySubsystem,
 		{
 			FArchetypeHandle& Archetype = ValidArchetypes[i];
 			check(Archetype.IsValid());
-			Archetype.DataPtr->ExecuteFunction(ExecutionContext, ExecuteFunction, ArchetypeComponentMapping[i], ChunkCondition);
-			ExecutionContext.ClearComponentViews();
+			Archetype.DataPtr->ExecuteFunction(ExecutionContext, ExecuteFunction, ArchetypeFragmentMapping[i], ChunkCondition);
+			ExecutionContext.ClearFragmentViews();
 #if WITH_MASSENTITY_DEBUG
 			NumEntitiesToProcess += ExecutionContext.GetNumEntities();
 #endif
@@ -233,7 +233,7 @@ void FMassEntityQuery::ParallelForEachEntityChunk(UMassEntitySubsystem& EntitySu
 	ParallelFor(Jobs.Num(), [this, ExecutionContext, &ExecuteFunction, Jobs](const int32 JobIndex)
 	{
 		Jobs[JobIndex].Archetype.ExecutionFunctionForChunk(ExecutionContext, ExecuteFunction
-			, Jobs[JobIndex].ArchetypeIndex != INDEX_NONE ? ArchetypeComponentMapping[Jobs[JobIndex].ArchetypeIndex] : FMassQueryRequirementIndicesMapping()
+			, Jobs[JobIndex].ArchetypeIndex != INDEX_NONE ? ArchetypeFragmentMapping[Jobs[JobIndex].ArchetypeIndex] : FMassQueryRequirementIndicesMapping()
 			, Jobs[JobIndex].ChunkInfo
 			, ChunkCondition);
 	});
@@ -302,7 +302,7 @@ FString FMassFragmentRequirement::DebugGetDescription() const
 {
 #if WITH_MASSENTITY_DEBUG
 	return FString::Printf(TEXT("%s%s[%s]"), IsOptional() ? TEXT("?") : (Presence == EMassFragmentPresence::None ? TEXT("-") : TEXT("+"))
-		, *GetNameSafe(StructType), *UE::Mass::Debug::DebugGetComponentAccessString(AccessMode));
+		, *GetNameSafe(StructType), *UE::Mass::Debug::DebugGetFragmentAccessString(AccessMode));
 #else
 	return {};
 #endif
