Renamed all Component-containing functions and member variable names in Mass, replaced with "Fragment". Updated a lot of comments and logged text as well.

diff --git a/Engine/Plugins/Runtime/MassEntity/Source/MassEntity/Private/MassEntitySubsystem.cpp b/Engine/Plugins/Runtime/MassEntity/Source/MassEntity/Private/MassEntitySubsystem.cpp
--- a/Engine/Plugins/Runtime/MassEntity/Source/MassEntity/Private/MassEntitySubsystem.cpp
+++ b/Engine/Plugins/Runtime/MassEntity/Source/MassEntity/Private/MassEntitySubsystem.cpp
@@ -13,7 +13,6 @@ const FMassEntityHandle UMassEntitySubsystem::InvalidEntity;
 //@TODO: No re-entrance safety while running a system (e.g., preventing someone from adding/removing entities or altering archetypes, etc...)
 //@TODO: Do we allow GCable types?  If so, need to implement AddReferencedObjects
 //@TODO: Do a UObject stat scope around the system Run call
-//@TODO: Tag components end up taking up space per-entity when they shouldn't
 
 //////////////////////////////////////////////////////////////////////
 // UMassEntitySubsystem
@@ -29,11 +28,11 @@ void UMassEntitySubsystem::GetResourceSizeEx(FResourceSizeEx& CumulativeResource
 	const SIZE_T MyExtraSize = Entities.GetAllocatedSize() + 
 		EntityFreeIndexList.GetAllocatedSize() +
 		DeferredCommandBuffer->GetAllocatedSize() +
-		ComponentHashToArchetypeMap.GetAllocatedSize() +
-		ComponentTypeToArchetypeMap.GetAllocatedSize();
+		FragmentHashToArchetypeMap.GetAllocatedSize() +
+		FragmentTypeToArchetypeMap.GetAllocatedSize();
 	CumulativeResourceSize.AddDedicatedSystemMemoryBytes(MyExtraSize);
 
-	for (const auto& KVP : ComponentHashToArchetypeMap)
+	for (const auto& KVP : FragmentHashToArchetypeMap)
 	{
 		for (const TSharedPtr<FMassArchetypeData>& ArchetypePtr : KVP.Value)
 		{
@@ -52,15 +51,15 @@ void UMassEntitySubsystem::Initialize(FSubsystemCollectionBase& Collection)
 
 	// initialize the bitsets
 
-	FMassFragmentBitSet Components;
+	FMassFragmentBitSet Fragments;
 	FMassTagBitSet Tags;
-	FMassChunkFragmentBitSet ChunkComponents;
+	FMassChunkFragmentBitSet ChunkFragments;
 
 	for (TObjectIterator<UStruct> StructIt; StructIt; ++StructIt)
 	{
 		if (StructIt->IsChildOf(FMassFragment::StaticStruct()))
 		{
-			Components.Add(*CastChecked<UScriptStruct>(*StructIt));
+			Fragments.Add(*CastChecked<UScriptStruct>(*StructIt));
 		}
 		else if (StructIt->IsChildOf(FMassTag::StaticStruct()))
 		{
@@ -68,7 +67,7 @@ void UMassEntitySubsystem::Initialize(FSubsystemCollectionBase& Collection)
 		}
 		else if (StructIt->IsChildOf(FMassChunkFragment::StaticStruct()))
 		{
-			ChunkComponents.Add(*CastChecked<UScriptStruct>(*StructIt));
+			ChunkFragments.Add(*CastChecked<UScriptStruct>(*StructIt));
 		}
 	}
 }
@@ -78,18 +77,18 @@ void UMassEntitySubsystem::Deinitialize()
 	//@TODO: Should actually do this...
 }
 
-FArchetypeHandle UMassEntitySubsystem::CreateArchetype(TConstArrayView<const UScriptStruct*> ComponentsAngTagsList)
+FArchetypeHandle UMassEntitySubsystem::CreateArchetype(TConstArrayView<const UScriptStruct*> FragmentsAngTagsList)
 {
-	FMassChunkFragmentBitSet ChunkComponents;
+	FMassChunkFragmentBitSet ChunkFragments;
 	FMassTagBitSet Tags;
-	TArray<const UScriptStruct*, TInlineAllocator<16>> ComponentList;
-	ComponentList.Reserve(ComponentsAngTagsList.Num());
+	TArray<const UScriptStruct*, TInlineAllocator<16>> FragmentList;
+	FragmentList.Reserve(FragmentsAngTagsList.Num());
 
-	for (const UScriptStruct* Type : ComponentsAngTagsList)
+	for (const UScriptStruct* Type : FragmentsAngTagsList)
 	{
 		if (Type->IsChildOf(FMassFragment::StaticStruct()))
 		{
-			ComponentList.Add(Type);
+			FragmentList.Add(Type);
 		}
 		else if (Type->IsChildOf(FMassTag::StaticStruct()))
 		{
@@ -97,36 +96,36 @@ FArchetypeHandle UMassEntitySubsystem::CreateArchetype(TConstArrayView<const USc
 		}
 		else if (Type->IsChildOf(FMassChunkFragment::StaticStruct()))
 		{
-			ChunkComponents.Add(*Type);
+			ChunkFragments.Add(*Type);
 		}
 		else
 		{
-			UE_LOG(LogMass, Warning, TEXT("%s: %s is not a valid component nor tag type. Ignoring.")
+			UE_LOG(LogMass, Warning, TEXT("%s: %s is not a valid fragment nor tag type. Ignoring.")
 				, ANSI_TO_TCHAR(__FUNCTION__), *GetNameSafe(Type));
 		}
 	}
 
-	return CreateArchetype(FMassCompositionDescriptor(FMassFragmentBitSet(ComponentList), Tags, ChunkComponents));
+	return CreateArchetype(FMassCompositionDescriptor(FMassFragmentBitSet(FragmentList), Tags, ChunkFragments));
 }
 
-FArchetypeHandle UMassEntitySubsystem::CreateArchetype(TConstArrayView<const UScriptStruct*> ComponentList, const FMassTagBitSet& Tags, const FMassChunkFragmentBitSet& ChunkComponents)
+FArchetypeHandle UMassEntitySubsystem::CreateArchetype(TConstArrayView<const UScriptStruct*> FragmentList, const FMassTagBitSet& Tags, const FMassChunkFragmentBitSet& ChunkFragments)
 {
-	return CreateArchetype(FMassCompositionDescriptor(ComponentList, Tags, ChunkComponents));
+	return CreateArchetype(FMassCompositionDescriptor(FragmentList, Tags, ChunkFragments));
 }
 
-FArchetypeHandle UMassEntitySubsystem::CreateArchetype(const TSharedPtr<FMassArchetypeData>& SourceArchetype, const FMassFragmentBitSet& NewComponents)
+FArchetypeHandle UMassEntitySubsystem::CreateArchetype(const TSharedPtr<FMassArchetypeData>& SourceArchetype, const FMassFragmentBitSet& NewFragments)
 {
 	check(SourceArchetype.IsValid());
-	checkf(NewComponents.IsEmpty() == false, TEXT("%s Adding an empty component list to an archetype is not supported."), ANSI_TO_TCHAR(__FUNCTION__));
+	checkf(NewFragments.IsEmpty() == false, TEXT("%s Adding an empty fragment list to an archetype is not supported."), ANSI_TO_TCHAR(__FUNCTION__));
 
-	return CreateArchetype(FMassCompositionDescriptor(NewComponents + SourceArchetype->GetComponentBitSet(), SourceArchetype->GetTagBitSet(), SourceArchetype->GetChunkComponentBitSet()));
+	return CreateArchetype(FMassCompositionDescriptor(NewFragments + SourceArchetype->GetFragmentBitSet(), SourceArchetype->GetTagBitSet(), SourceArchetype->GetChunkFragmentBitSet()));
 }
 
 FArchetypeHandle UMassEntitySubsystem::CreateArchetype(const FMassCompositionDescriptor& Composition)
 {
 	const uint32 TypeHash = Composition.CalculateHash();
 
-	TArray<TSharedPtr<FMassArchetypeData>>& HashRow = ComponentHashToArchetypeMap.FindOrAdd(TypeHash);
+	TArray<TSharedPtr<FMassArchetypeData>>& HashRow = FragmentHashToArchetypeMap.FindOrAdd(TypeHash);
 
 	FArchetypeHandle Result;
 	for (TSharedPtr<FMassArchetypeData>& Ptr : HashRow)
@@ -142,13 +141,13 @@ FArchetypeHandle UMassEntitySubsystem::CreateArchetype(const FMassCompositionDes
 	{
 		// Create a new archetype
 		TSharedPtr<FMassArchetypeData> NewArchetype = MakeShareable(new FMassArchetypeData);
-		NewArchetype->Initialize(Composition.Components, Composition.Tags, Composition.ChunkComponents);
+		NewArchetype->Initialize(Composition.Fragments, Composition.Tags, Composition.ChunkFragments);
 		HashRow.Add(NewArchetype);
 
-		for (const FMassArchetypeFragmentConfig& ComponentConfig : NewArchetype->GetComponentConfigs())
+		for (const FMassArchetypeFragmentConfig& FragmentConfig : NewArchetype->GetFragmentConfigs())
 		{
-			checkSlow(ComponentConfig.ComponentType)
-			ComponentTypeToArchetypeMap.FindOrAdd(ComponentConfig.ComponentType).Add(NewArchetype);
+			checkSlow(FragmentConfig.FragmentType)
+			FragmentTypeToArchetypeMap.FindOrAdd(FragmentConfig.FragmentType).Add(NewArchetype);
 		}
 
 		Result.DataPtr = NewArchetype;
@@ -163,14 +162,14 @@ FArchetypeHandle UMassEntitySubsystem::InternalCreateSiblingArchetype(const TSha
 {
 	checkSlow(SourceArchetype.IsValid());
 	const FMassArchetypeData& SourceArchetypeRef = *SourceArchetype.Get();
-	const uint32 TypeHash = FMassCompositionDescriptor::CalculateHash(SourceArchetypeRef.GetComponentBitSet(), OverrideTags, SourceArchetypeRef.GetChunkComponentBitSet());
+	const uint32 TypeHash = FMassCompositionDescriptor::CalculateHash(SourceArchetypeRef.GetFragmentBitSet(), OverrideTags, SourceArchetypeRef.GetChunkFragmentBitSet());
 
-	TArray<TSharedPtr<FMassArchetypeData>>& HashRow = ComponentHashToArchetypeMap.FindOrAdd(TypeHash);
+	TArray<TSharedPtr<FMassArchetypeData>>& HashRow = FragmentHashToArchetypeMap.FindOrAdd(TypeHash);
 
 	FArchetypeHandle Result;
 	for (TSharedPtr<FMassArchetypeData>& Ptr : HashRow)
 	{
-		if (Ptr->IsEquivalent(SourceArchetypeRef.GetComponentBitSet(), OverrideTags, SourceArchetypeRef.GetChunkComponentBitSet()))
+		if (Ptr->IsEquivalent(SourceArchetypeRef.GetFragmentBitSet(), OverrideTags, SourceArchetypeRef.GetChunkFragmentBitSet()))
 		{
 			Result.DataPtr = Ptr;
 			break;
@@ -184,10 +183,10 @@ FArchetypeHandle UMassEntitySubsystem::InternalCreateSiblingArchetype(const TSha
 		NewArchetype->InitializeWithSibling(SourceArchetypeRef, OverrideTags);
 		HashRow.Add(NewArchetype);
 
-		for (const FMassArchetypeFragmentConfig& ComponentConfig : NewArchetype->GetComponentConfigs())
+		for (const FMassArchetypeFragmentConfig& FragmentConfig : NewArchetype->GetFragmentConfigs())
 		{
-			checkSlow(ComponentConfig.ComponentType)
-			ComponentTypeToArchetypeMap.FindOrAdd(ComponentConfig.ComponentType).Add(NewArchetype);
+			checkSlow(FragmentConfig.FragmentType)
+			FragmentTypeToArchetypeMap.FindOrAdd(FragmentConfig.FragmentType).Add(NewArchetype);
 		}
 
 		Result.DataPtr = NewArchetype;
@@ -208,18 +207,18 @@ FArchetypeHandle UMassEntitySubsystem::GetArchetypeForEntity(FMassEntityHandle E
 	return Result;
 }
 
-void UMassEntitySubsystem::ForEachArchetypeComponentType(const FArchetypeHandle Archetype, TFunction< void(const UScriptStruct* /*ComponentType*/)> Function)
+void UMassEntitySubsystem::ForEachArchetypeFragmentType(const FArchetypeHandle Archetype, TFunction< void(const UScriptStruct* /*FragmentType*/)> Function)
 {
 	check(Archetype.DataPtr.IsValid());
 	const FMassArchetypeData& ArchetypeData = *Archetype.DataPtr.Get();
-	ArchetypeData.ForEachComponentType(Function);
+	ArchetypeData.ForEachFragmentType(Function);
 }
 
-void UMassEntitySubsystem::SetDefaultChunkComponentValue(const FArchetypeHandle Archetype, FConstStructView InstancedStruct)
+void UMassEntitySubsystem::SetDefaultChunkFragmentValue(const FArchetypeHandle Archetype, FConstStructView InstancedStruct)
 {
 	check(Archetype.DataPtr.IsValid());
 	FMassArchetypeData& ArchetypeData = *Archetype.DataPtr.Get();
-	ArchetypeData.SetDefaultChunkComponentValue(InstancedStruct);
+	ArchetypeData.SetDefaultChunkFragmentValue(InstancedStruct);
 }
 
 void UMassEntitySubsystem::DoEntityCompaction(const double TimeAllowed)
@@ -227,7 +226,7 @@ void UMassEntitySubsystem::DoEntityCompaction(const double TimeAllowed)
 	const double TimeAllowedEnd = FPlatformTime::Seconds() + TimeAllowed;
 
 	bool bReachedTimeLimit = false;
-	for (const auto& KVP : ComponentHashToArchetypeMap)
+	for (const auto& KVP : FragmentHashToArchetypeMap)
 	{
 		for (const TSharedPtr<FMassArchetypeData>& ArchetypePtr : KVP.Value)
 		{
@@ -255,18 +254,18 @@ FMassEntityHandle UMassEntitySubsystem::CreateEntity(const FArchetypeHandle Arch
 	return Entity;
 }
 
-FMassEntityHandle UMassEntitySubsystem::CreateEntity(TConstArrayView<FInstancedStruct> ComponentInstanceList)
+FMassEntityHandle UMassEntitySubsystem::CreateEntity(TConstArrayView<FInstancedStruct> FragmentInstanceList)
 {
-	check(ComponentInstanceList.Num() > 0);
+	check(FragmentInstanceList.Num() > 0);
 
-	FArchetypeHandle Archetype = CreateArchetype(FMassCompositionDescriptor(ComponentInstanceList, FMassTagBitSet(), FMassChunkFragmentBitSet()));
+	FArchetypeHandle Archetype = CreateArchetype(FMassCompositionDescriptor(FragmentInstanceList, FMassTagBitSet(), FMassChunkFragmentBitSet()));
 	check(Archetype.DataPtr.IsValid());
 
 	FMassEntityHandle Entity = ReserveEntity();
 	InternalBuildEntity(Entity, Archetype);
 
 	FEntityData& EntityData = Entities[Entity.Index];
-	EntityData.CurrentArchetype->SetComponentsData(Entity, ComponentInstanceList);
+	EntityData.CurrentArchetype->SetFragmentsData(Entity, FragmentInstanceList);
 
 	return Entity;
 }
@@ -297,18 +296,18 @@ void UMassEntitySubsystem::BuildEntity(FMassEntityHandle Entity, FArchetypeHandl
 	InternalBuildEntity(Entity, Archetype);
 }
 
-void UMassEntitySubsystem::BuildEntity(FMassEntityHandle Entity, TConstArrayView<FInstancedStruct> ComponentInstanceList)
+void UMassEntitySubsystem::BuildEntity(FMassEntityHandle Entity, TConstArrayView<FInstancedStruct> FragmentInstanceList)
 {
-	check(ComponentInstanceList.Num() > 0);
+	check(FragmentInstanceList.Num() > 0);
 	checkf(!IsEntityBuilt(Entity), TEXT("Expecting an entity that is not already built"));
 
-	FArchetypeHandle Archetype = CreateArchetype(FMassCompositionDescriptor(ComponentInstanceList, FMassTagBitSet(), FMassChunkFragmentBitSet()));
+	FArchetypeHandle Archetype = CreateArchetype(FMassCompositionDescriptor(FragmentInstanceList, FMassTagBitSet(), FMassChunkFragmentBitSet()));
 	check(Archetype.DataPtr.IsValid());
 
 	InternalBuildEntity(Entity, Archetype);
 
 	FEntityData& EntityData = Entities[Entity.Index];
-	EntityData.CurrentArchetype->SetComponentsData(Entity, ComponentInstanceList);
+	EntityData.CurrentArchetype->SetFragmentsData(Entity, FragmentInstanceList);
 }
 
 int32 UMassEntitySubsystem::BatchCreateEntities(const FArchetypeHandle Archetype, const int32 Count, TArray<FMassEntityHandle>& OutEntities)
@@ -412,21 +411,21 @@ void UMassEntitySubsystem::BatchDestroyEntityChunks(const FArchetypeChunkCollect
 	}
 }
 
-void UMassEntitySubsystem::AddComponentToEntity(FMassEntityHandle Entity, const UScriptStruct* ComponentType)
+void UMassEntitySubsystem::AddFragmentToEntity(FMassEntityHandle Entity, const UScriptStruct* FragmentType)
 {
-	checkf(ComponentType, TEXT("Null component type passed in to %s"), ANSI_TO_TCHAR(__FUNCTION__));
+	checkf(FragmentType, TEXT("Null fragment type passed in to %s"), ANSI_TO_TCHAR(__FUNCTION__));
 	checkf(IsProcessing() == false, TEXT("Synchronous API function %s called during mass processing. Use asynchronous API instead."), ANSI_TO_TCHAR(__FUNCTION__));
 
 	CheckIfEntityIsActive(Entity);
 
-	InternalAddComponentListToEntityChecked(Entity, FMassFragmentBitSet(*ComponentType));
+	InternalAddFragmentListToEntityChecked(Entity, FMassFragmentBitSet(*FragmentType));
 }
 
-void UMassEntitySubsystem::AddComponentListToEntity(FMassEntityHandle Entity, TConstArrayView<const UScriptStruct*> ComponentList)
+void UMassEntitySubsystem::AddFragmentListToEntity(FMassEntityHandle Entity, TConstArrayView<const UScriptStruct*> FragmentList)
 {
 	CheckIfEntityIsActive(Entity);
 
-	InternalAddComponentListToEntityChecked(Entity, FMassFragmentBitSet(ComponentList));
+	InternalAddFragmentListToEntityChecked(Entity, FMassFragmentBitSet(FragmentList));
 }
 
 void UMassEntitySubsystem::AddCompositionToEntity_GetDelta(FMassEntityHandle Entity, FMassCompositionDescriptor& Composition)
@@ -437,13 +436,13 @@ void UMassEntitySubsystem::AddCompositionToEntity_GetDelta(FMassEntityHandle Ent
 	FMassArchetypeData* OldArchetype = EntityData.CurrentArchetype.Get();
 	check(OldArchetype);
 
-	Composition.Components -= OldArchetype->GetCompositionDescriptor().Components;
+	Composition.Fragments -= OldArchetype->GetCompositionDescriptor().Fragments;
 	Composition.Tags -= OldArchetype->GetCompositionDescriptor().Tags;
 
 	if (Composition.IsEmpty() == false)
 	{
 		FMassCompositionDescriptor NewDescriptor = OldArchetype->GetCompositionDescriptor();
-		NewDescriptor.Components += Composition.Components;
+		NewDescriptor.Fragments += Composition.Fragments;
 		NewDescriptor.Tags+= Composition.Tags;
 
 		const FArchetypeHandle NewArchetypeHandle = CreateArchetype(NewDescriptor);
@@ -468,7 +467,7 @@ void UMassEntitySubsystem::RemoveCompositionFromEntity(FMassEntityHandle Entity,
 	check(OldArchetype);
 
 	FMassCompositionDescriptor NewDescriptor = OldArchetype->GetCompositionDescriptor();
-	NewDescriptor.Components -= InDescriptor.Components;
+	NewDescriptor.Fragments -= InDescriptor.Fragments;
 	NewDescriptor.Tags -= InDescriptor.Tags;
 
 	if (NewDescriptor.IsEmpty() == false)
@@ -500,30 +499,30 @@ void UMassEntitySubsystem::InternalReleaseEntity(FMassEntityHandle Entity)
 	EntityFreeIndexList.Add(Entity.Index);
 }
 
-void UMassEntitySubsystem::InternalAddComponentListToEntityChecked(FMassEntityHandle Entity, const FMassFragmentBitSet& InComponents)
+void UMassEntitySubsystem::InternalAddFragmentListToEntityChecked(FMassEntityHandle Entity, const FMassFragmentBitSet& InFragments)
 {
 	FEntityData& EntityData = Entities[Entity.Index];
 	FMassArchetypeData* OldArchetype = EntityData.CurrentArchetype.Get();
 	check(OldArchetype);
 	
-	ensureMsgf(OldArchetype->GetComponentBitSet().HasAny(InComponents) == false, TEXT("Trying to add a new component type to an entity, but it already has some of them."));
+	ensureMsgf(OldArchetype->GetFragmentBitSet().HasAny(InFragments) == false, TEXT("Trying to add a new fragment type to an entity, but it already has some of them."));
 
-	const FMassFragmentBitSet NewComponents = InComponents - OldArchetype->GetComponentBitSet();
-	if (NewComponents.IsEmpty() == false)
+	const FMassFragmentBitSet NewFragments = InFragments - OldArchetype->GetFragmentBitSet();
+	if (NewFragments.IsEmpty() == false)
 	{
-		InternalAddComponentListToEntity(Entity, NewComponents);
+		InternalAddFragmentListToEntity(Entity, NewFragments);
 	}
 }
 
-void UMassEntitySubsystem::InternalAddComponentListToEntity(FMassEntityHandle Entity, const FMassFragmentBitSet& NewComponents)
+void UMassEntitySubsystem::InternalAddFragmentListToEntity(FMassEntityHandle Entity, const FMassFragmentBitSet& NewFragments)
 {
-	checkf(NewComponents.IsEmpty() == false, TEXT("%s is intended for internal calls with non empty NewComponents parameter"), ANSI_TO_TCHAR(__FUNCTION__));
+	checkf(NewFragments.IsEmpty() == false, TEXT("%s is intended for internal calls with non empty NewFragments parameter"), ANSI_TO_TCHAR(__FUNCTION__));
 	check(Entities.IsValidIndex(Entity.Index));
 	FEntityData& EntityData = Entities[Entity.Index];
 	check(EntityData.CurrentArchetype.IsValid());
 
 	// fetch or create the new archetype
-	const FArchetypeHandle NewArchetypeHandle = CreateArchetype(EntityData.CurrentArchetype, NewComponents);
+	const FArchetypeHandle NewArchetypeHandle = CreateArchetype(EntityData.CurrentArchetype, NewFragments);
 
 	if (NewArchetypeHandle.DataPtr != EntityData.CurrentArchetype)
 	{
@@ -535,25 +534,25 @@ void UMassEntitySubsystem::InternalAddComponentListToEntity(FMassEntityHandle En
 	}
 }
 
-void UMassEntitySubsystem::AddComponentInstanceListToEntity(FMassEntityHandle Entity, TConstArrayView<FInstancedStruct> ComponentInstanceList)
+void UMassEntitySubsystem::AddFragmentInstanceListToEntity(FMassEntityHandle Entity, TConstArrayView<FInstancedStruct> FragmentInstanceList)
 {
 	checkf(IsProcessing() == false, TEXT("Synchronous API function %s called during mass processing. Use asynchronous API instead."), ANSI_TO_TCHAR(__FUNCTION__));
 
 	CheckIfEntityIsActive(Entity);
-	checkf(ComponentInstanceList.Num() > 0, TEXT("Need to specify at least one component instances for this operation"));
+	checkf(FragmentInstanceList.Num() > 0, TEXT("Need to specify at least one fragment instances for this operation"));
 
-	InternalAddComponentListToEntityChecked(Entity, FMassFragmentBitSet(ComponentInstanceList));
+	InternalAddFragmentListToEntityChecked(Entity, FMassFragmentBitSet(FragmentInstanceList));
 
 	FEntityData& EntityData = Entities[Entity.Index];
-	EntityData.CurrentArchetype->SetComponentsData(Entity, ComponentInstanceList);
+	EntityData.CurrentArchetype->SetFragmentsData(Entity, FragmentInstanceList);
 }
 
-void UMassEntitySubsystem::RemoveComponentFromEntity(FMassEntityHandle Entity, const UScriptStruct* ComponentType)
+void UMassEntitySubsystem::RemoveFragmentFromEntity(FMassEntityHandle Entity, const UScriptStruct* FragmentType)
 {
-	RemoveComponentListFromEntity(Entity, MakeArrayView(&ComponentType, 1));
+	RemoveFragmentListFromEntity(Entity, MakeArrayView(&FragmentType, 1));
 }
 
-void UMassEntitySubsystem::RemoveComponentListFromEntity(FMassEntityHandle Entity, TConstArrayView<const UScriptStruct*> ComponentList)
+void UMassEntitySubsystem::RemoveFragmentListFromEntity(FMassEntityHandle Entity, TConstArrayView<const UScriptStruct*> FragmentList)
 {
 	checkf(IsProcessing() == false, TEXT("Synchronous API function %s called during mass processing. Use asynchronous API instead."), ANSI_TO_TCHAR(__FUNCTION__));
 
@@ -563,12 +562,12 @@ void UMassEntitySubsystem::RemoveComponentListFromEntity(FMassEntityHandle Entit
 	FMassArchetypeData* OldArchetype = EntityData.CurrentArchetype.Get();
 	check(OldArchetype);
 
-	const FMassFragmentBitSet ComponentsToRemove(ComponentList);
+	const FMassFragmentBitSet FragmentsToRemove(FragmentList);
 
-	if (ensureMsgf(OldArchetype->GetComponentBitSet().HasAny(ComponentsToRemove), TEXT("Trying to remove a list of components from an entity but none of the components given was found.")))
+	if (ensureMsgf(OldArchetype->GetFragmentBitSet().HasAny(FragmentsToRemove), TEXT("Trying to remove a list of fragments from an entity but none of the fragments given was found.")))
 	{
-		// If all the components got removed this will result in fetching of the empty archetype
-		const FArchetypeHandle NewArchetypeHandle = CreateArchetype(FMassCompositionDescriptor(OldArchetype->GetComponentBitSet() - ComponentsToRemove, OldArchetype->GetTagBitSet(), OldArchetype->GetChunkComponentBitSet()));
+		// If all the fragments got removed this will result in fetching of the empty archetype
+		const FArchetypeHandle NewArchetypeHandle = CreateArchetype(FMassCompositionDescriptor(OldArchetype->GetFragmentBitSet() - FragmentsToRemove, OldArchetype->GetTagBitSet(), OldArchetype->GetChunkFragmentBitSet()));
 
 		// Move the entity over
 		FMassArchetypeData* NewArchetype = NewArchetypeHandle.DataPtr.Get();
@@ -668,40 +667,40 @@ void UMassEntitySubsystem::MoveEntityToAnotherArchetype(FMassEntityHandle Entity
 	EntityData.CurrentArchetype = NewArchetypeHandle.DataPtr;
 }
 
-void UMassEntitySubsystem::SetEntityComponentsValues(FMassEntityHandle Entity, TArrayView<const FInstancedStruct> ComponentInstanceList)
+void UMassEntitySubsystem::SetEntityFragmentsValues(FMassEntityHandle Entity, TArrayView<const FInstancedStruct> FragmentInstanceList)
 {
 	CheckIfEntityIsActive(Entity);
 
 	FEntityData& EntityData = Entities[Entity.Index];
-	EntityData.CurrentArchetype->SetComponentsData(Entity, ComponentInstanceList);
+	EntityData.CurrentArchetype->SetFragmentsData(Entity, FragmentInstanceList);
 }
 
-void UMassEntitySubsystem::BatchSetEntityComponentsValues(const FArchetypeChunkCollection& SparseEntities, TArrayView<const FInstancedStruct> ComponentInstanceList)
+void UMassEntitySubsystem::BatchSetEntityFragmentsValues(const FArchetypeChunkCollection& SparseEntities, TArrayView<const FInstancedStruct> FragmentInstanceList)
 {
 	FMassArchetypeData* Archetype = SparseEntities.GetArchetype().DataPtr.Get();
 	check(Archetype);
 
-	for (const FInstancedStruct& ComponentTemplate : ComponentInstanceList)
+	for (const FInstancedStruct& FragmentTemplate : FragmentInstanceList)
 	{
-		Archetype->SetComponentData(SparseEntities, ComponentTemplate);
+		Archetype->SetFragmentData(SparseEntities, FragmentTemplate);
 	}
 }
 
-void* UMassEntitySubsystem::InternalGetComponentDataChecked(FMassEntityHandle Entity, const UScriptStruct* ComponentType) const
+void* UMassEntitySubsystem::InternalGetFragmentDataChecked(FMassEntityHandle Entity, const UScriptStruct* FragmentType) const
 {
 	CheckIfEntityIsActive(Entity);
 
-	checkf((ComponentType != nullptr) && ComponentType->IsChildOf(FMassFragment::StaticStruct()), TEXT("InternalGetComponentDataChecked called with an invalid component type '%s'"), *GetPathNameSafe(ComponentType));
+	checkf((FragmentType != nullptr) && FragmentType->IsChildOf(FMassFragment::StaticStruct()), TEXT("InternalGetFragmentDataChecked called with an invalid fragment type '%s'"), *GetPathNameSafe(FragmentType));
 	const FEntityData& EntityData = Entities[Entity.Index];
-	return EntityData.CurrentArchetype->GetComponentDataForEntityChecked(ComponentType, Entity.Index);
+	return EntityData.CurrentArchetype->GetFragmentDataForEntityChecked(FragmentType, Entity.Index);
 }
 
-void* UMassEntitySubsystem::InternalGetComponentDataPtr(FMassEntityHandle Entity, const UScriptStruct* ComponentType) const
+void* UMassEntitySubsystem::InternalGetFragmentDataPtr(FMassEntityHandle Entity, const UScriptStruct* FragmentType) const
 {
 	CheckIfEntityIsActive(Entity);
-	checkf((ComponentType != nullptr) && ComponentType->IsChildOf(FMassFragment::StaticStruct()), TEXT("InternalGetComponentData called with an invalid component type '%s'"), *GetPathNameSafe(ComponentType));
+	checkf((FragmentType != nullptr) && FragmentType->IsChildOf(FMassFragment::StaticStruct()), TEXT("InternalGetFragmentData called with an invalid fragment type '%s'"), *GetPathNameSafe(FragmentType));
 	const FEntityData& EntityData = Entities[Entity.Index];
-	return EntityData.CurrentArchetype->GetComponentDataForEntity(ComponentType, Entity.Index);
+	return EntityData.CurrentArchetype->GetFragmentDataForEntity(FragmentType, Entity.Index);
 }
 
 bool UMassEntitySubsystem::IsEntityValid(FMassEntityHandle Entity) const
@@ -730,21 +729,21 @@ void UMassEntitySubsystem::GetValidArchetypes(const FMassEntityQuery& Query, TAr
 {
 	//@TODO: Not optimized yet, but we call this rarely now, so not a big deal.
 
-	// First get set of all archetypes that contain *any* component
+	// First get set of all archetypes that contain *any* fragment
 	TSet<TSharedPtr<FMassArchetypeData>> AnyArchetypes;
 	for (const FMassFragmentRequirement& Requirement : Query.GetRequirements())
 	{
 		check(Requirement.StructType);
 		if (Requirement.Presence != EMassFragmentPresence::None)
 		{
-			if (TArray<TSharedPtr<FMassArchetypeData>>* pData = ComponentTypeToArchetypeMap.Find(Requirement.StructType))
+			if (TArray<TSharedPtr<FMassArchetypeData>>* pData = FragmentTypeToArchetypeMap.Find(Requirement.StructType))
 			{
 				AnyArchetypes.Append(*pData);
 			}
 		}
 	}
 
-	// Then verify that they contain *all* required components
+	// Then verify that they contain *all* required fragments
 	for (TSharedPtr<FMassArchetypeData>& ArchetypePtr : AnyArchetypes)
 	{
 		FMassArchetypeData& Archetype = *(ArchetypePtr.Get());
@@ -787,65 +786,65 @@ void UMassEntitySubsystem::GetValidArchetypes(const FMassEntityQuery& Query, TAr
 			continue;
 		}
 		
-		if (Archetype.GetComponentBitSet().HasAll(Query.GetRequiredAllComponents()) == false)
+		if (Archetype.GetFragmentBitSet().HasAll(Query.GetRequiredAllFragments()) == false)
 		{
-			// missing some required components, skip.
+			// missing some required fragments, skip.
 #if WITH_MASSENTITY_DEBUG
-			const FMassFragmentBitSet UnsatisfiedComponents = Query.GetRequiredAllComponents() - Archetype.GetComponentBitSet();
+			const FMassFragmentBitSet UnsatisfiedFragments = Query.GetRequiredAllFragments() - Archetype.GetFragmentBitSet();
 			FStringOutputDevice Description;
-			UnsatisfiedComponents.DebugGetStringDesc(Description);
-			UE_LOG(LogMass, VeryVerbose, TEXT("Archetype did not match due to missing Components: %s")
+			UnsatisfiedFragments.DebugGetStringDesc(Description);
+			UE_LOG(LogMass, VeryVerbose, TEXT("Archetype did not match due to missing Fragments: %s")
 				, *Description);
 #endif // WITH_MASSENTITY_DEBUG
 			continue;
 		}
 
-		if (Archetype.GetComponentBitSet().HasNone(Query.GetRequiredNoneComponents()) == false)
+		if (Archetype.GetFragmentBitSet().HasNone(Query.GetRequiredNoneFragments()) == false)
 		{
-			// has some Components required to be absent
+			// has some Fragments required to be absent
 #if WITH_MASSENTITY_DEBUG
-			const FMassFragmentBitSet UnwantedComponents = Query.GetRequiredAllComponents().GetOverlap(Archetype.GetComponentBitSet());
+			const FMassFragmentBitSet UnwantedFragments = Query.GetRequiredAllFragments().GetOverlap(Archetype.GetFragmentBitSet());
 			FStringOutputDevice Description;
-			UnwantedComponents.DebugGetStringDesc(Description);
-			UE_LOG(LogMass, VeryVerbose, TEXT("Archetype has Components required absent: %s")
+			UnwantedFragments.DebugGetStringDesc(Description);
+			UE_LOG(LogMass, VeryVerbose, TEXT("Archetype has Fragments required absent: %s")
 				, *Description);
 #endif // WITH_MASSENTITY_DEBUG
 			continue;
 		}
 
-		if (Query.GetRequiredAnyComponents().IsEmpty() == false 
-			&& Archetype.GetComponentBitSet().HasAny(Query.GetRequiredAnyComponents()) == false)
+		if (Query.GetRequiredAnyFragments().IsEmpty() == false 
+			&& Archetype.GetFragmentBitSet().HasAny(Query.GetRequiredAnyFragments()) == false)
 		{
 #if WITH_MASSENTITY_DEBUG
 			FStringOutputDevice Description;
-			Query.GetRequiredAnyComponents().DebugGetStringDesc(Description);
-			UE_LOG(LogMass, VeryVerbose, TEXT("Archetype did not match due to missing \'any\' components: %s")
+			Query.GetRequiredAnyFragments().DebugGetStringDesc(Description);
+			UE_LOG(LogMass, VeryVerbose, TEXT("Archetype did not match due to missing \'any\' fragments: %s")
 				, *Description);
 #endif // WITH_MASSENTITY_DEBUG
 			continue;
 		}
 
-		if (Archetype.GetChunkComponentBitSet().HasAll(Query.GetRequiredAllChunkComponents()) == false)
+		if (Archetype.GetChunkFragmentBitSet().HasAll(Query.GetRequiredAllChunkFragments()) == false)
 		{
-			// missing some required components, skip.
+			// missing some required fragments, skip.
 #if WITH_MASSENTITY_DEBUG
-			const FMassChunkFragmentBitSet UnsatisfiedComponents = Query.GetRequiredAllChunkComponents() - Archetype.GetChunkComponentBitSet();
+			const FMassChunkFragmentBitSet UnsatisfiedFragments = Query.GetRequiredAllChunkFragments() - Archetype.GetChunkFragmentBitSet();
 			FStringOutputDevice Description;
-			UnsatisfiedComponents.DebugGetStringDesc(Description);
-			UE_LOG(LogMass, VeryVerbose, TEXT("Archetype did not match due to missing Chunk Components: %s")
+			UnsatisfiedFragments.DebugGetStringDesc(Description);
+			UE_LOG(LogMass, VeryVerbose, TEXT("Archetype did not match due to missing Chunk Fragments: %s")
 				, *Description);
 #endif // WITH_MASSENTITY_DEBUG
 			continue;
 		}
 
-		if (Archetype.GetChunkComponentBitSet().HasNone(Query.GetRequiredNoneChunkComponents()) == false)
+		if (Archetype.GetChunkFragmentBitSet().HasNone(Query.GetRequiredNoneChunkFragments()) == false)
 		{
-			// has some Components required to be absent
+			// has some Fragments required to be absent
 #if WITH_MASSENTITY_DEBUG
-			const FMassChunkFragmentBitSet UnwantedComponents = Query.GetRequiredNoneChunkComponents().GetOverlap(Archetype.GetChunkComponentBitSet());
+			const FMassChunkFragmentBitSet UnwantedFragments = Query.GetRequiredNoneChunkFragments().GetOverlap(Archetype.GetChunkFragmentBitSet());
 			FStringOutputDevice Description;
-			UnwantedComponents.DebugGetStringDesc(Description);
-			UE_LOG(LogMass, VeryVerbose, TEXT("Archetype has Chunk Components required absent: %s")
+			UnwantedFragments.DebugGetStringDesc(Description);
+			UE_LOG(LogMass, VeryVerbose, TEXT("Archetype has Chunk Fragments required absent: %s")
 				, *Description);
 #endif // WITH_MASSENTITY_DEBUG
 			continue;
@@ -870,7 +869,7 @@ void UMassEntitySubsystem::DebugPrintArchetypes(FOutputDevice& Ar) const
 	int32 NumBuckets = 0;
 	int32 NumArchetypes = 0;
 	int32 LongestArchetypeBucket = 0;
-	for (const auto& KVP : ComponentHashToArchetypeMap)
+	for (const auto& KVP : FragmentHashToArchetypeMap)
 	{
 		for (const TSharedPtr<FMassArchetypeData>& ArchetypePtr : KVP.Value)
 		{
@@ -883,7 +882,7 @@ void UMassEntitySubsystem::DebugPrintArchetypes(FOutputDevice& Ar) const
 		++NumBuckets;
 	}
 
-	Ar.Logf(ELogVerbosity::Log, TEXT("ComponentHashToArchetypeMap: %d archetypes across %d buckets, longest bucket is %d"),
+	Ar.Logf(ELogVerbosity::Log, TEXT("FragmentHashToArchetypeMap: %d archetypes across %d buckets, longest bucket is %d"),
 		NumArchetypes, NumBuckets, LongestArchetypeBucket);
 }
 
@@ -891,14 +890,14 @@ void UMassEntitySubsystem::DebugPrintEntity(int32 Index, FOutputDevice& Ar, cons
 {
 	if (Index >= Entities.Num())
 	{
-		Ar.Logf(ELogVerbosity::Log, TEXT("Unable to list components values for out of range index in %s"), *GetPathNameSafe(this));
+		Ar.Logf(ELogVerbosity::Log, TEXT("Unable to list fragments values for out of range index in %s"), *GetPathNameSafe(this));
 		return;
 	}
 
 	const FEntityData& EntityData = Entities[Index];
 	if (!EntityData.IsValid())
 	{
-		Ar.Logf(ELogVerbosity::Log, TEXT("Unable to list components values for invalid entity in %s"), *GetPathNameSafe(this));
+		Ar.Logf(ELogVerbosity::Log, TEXT("Unable to list fragments values for invalid entity in %s"), *GetPathNameSafe(this));
 	}
 
 	FMassEntityHandle Entity;
@@ -911,16 +910,16 @@ void UMassEntitySubsystem::DebugPrintEntity(FMassEntityHandle Entity, FOutputDev
 {
 	if (!IsEntityActive(Entity))
 	{
-		Ar.Logf(ELogVerbosity::Log, TEXT("Unable to list components values for invalid entity in %s"), *GetPathNameSafe(this));
+		Ar.Logf(ELogVerbosity::Log, TEXT("Unable to list fragments values for invalid entity in %s"), *GetPathNameSafe(this));
 	}
 
-	Ar.Logf(ELogVerbosity::Log, TEXT("Listing components values for Entity[%s] in %s"), *Entity.DebugGetDescription(), *GetPathNameSafe(this));
+	Ar.Logf(ELogVerbosity::Log, TEXT("Listing fragments values for Entity[%s] in %s"), *Entity.DebugGetDescription(), *GetPathNameSafe(this));
 
 	const FEntityData& EntityData = Entities[Entity.Index];
 	FMassArchetypeData* Archetype = EntityData.CurrentArchetype.Get();
 	if (Archetype == nullptr)
 	{
-		Ar.Logf(ELogVerbosity::Log, TEXT("Unable to list components values for invalid entity in %s"), *GetPathNameSafe(this));
+		Ar.Logf(ELogVerbosity::Log, TEXT("Unable to list fragments values for invalid entity in %s"), *GetPathNameSafe(this));
 	}
 	else
 	{
@@ -937,7 +936,7 @@ void UMassEntitySubsystem::DebugGetArchetypesStringDetails(FOutputDevice& Ar, co
 {
 #if WITH_MASSENTITY_DEBUG
 	Ar.SetAutoEmitLineTerminator(true);
-	for (auto Pair : ComponentHashToArchetypeMap)
+	for (auto Pair : FragmentHashToArchetypeMap)
 	{
 		Ar.Logf(ELogVerbosity::Log, TEXT("\n-----------------------------------\nHash: %u"), Pair.Key);
 		for (TSharedPtr<FMassArchetypeData> Archetype : Pair.Value)
@@ -952,11 +951,11 @@ void UMassEntitySubsystem::DebugGetArchetypesStringDetails(FOutputDevice& Ar, co
 #endif // WITH_MASSENTITY_DEBUG
 }
 
-void UMassEntitySubsystem::DebugGetArchetypeComponentTypes(const FArchetypeHandle& Archetype, TArray<const UScriptStruct*>& InOutComponentList) const
+void UMassEntitySubsystem::DebugGetArchetypeFragmentTypes(const FArchetypeHandle& Archetype, TArray<const UScriptStruct*>& InOutFragmentList) const
 {
 	if (Archetype.DataPtr.IsValid())
 	{
-		Archetype.DataPtr->GetCompositionDescriptor().Components.DebugGetStructTypes(InOutComponentList);
+		Archetype.DataPtr->GetCompositionDescriptor().Fragments.DebugGetStructTypes(InOutFragmentList);
 	}
 }
 
@@ -986,7 +985,7 @@ void UMassEntitySubsystem::DebugRemoveAllEntities()
 	}
 }
 
-void UMassEntitySubsystem::DebugGetArchetypeStrings(const FArchetypeHandle& Archetype, TArray<FName>& OutComponentNames, TArray<FName>& OutTagNames)
+void UMassEntitySubsystem::DebugGetArchetypeStrings(const FArchetypeHandle& Archetype, TArray<FName>& OutFragmentNames, TArray<FName>& OutTagNames)
 {
 	if (Archetype.IsValid() == false)
 	{
@@ -995,11 +994,11 @@ void UMassEntitySubsystem::DebugGetArchetypeStrings(const FArchetypeHandle& Arch
 
 	FMassArchetypeData& ArchetypeRef = *Archetype.DataPtr.Get();
 	
-	OutComponentNames.Reserve(ArchetypeRef.GetComponentConfigs().Num());
-	for (const FMassArchetypeFragmentConfig& ComponentConfig : ArchetypeRef.GetComponentConfigs())
+	OutFragmentNames.Reserve(ArchetypeRef.GetFragmentConfigs().Num());
+	for (const FMassArchetypeFragmentConfig& FragmentConfig : ArchetypeRef.GetFragmentConfigs())
 	{
-		checkSlow(ComponentConfig.ComponentType);
-		OutComponentNames.Add(ComponentConfig.ComponentType->GetFName());
+		checkSlow(FragmentConfig.FragmentType);
+		OutFragmentNames.Add(FragmentConfig.FragmentType->GetFName());
 	}
 
 	ArchetypeRef.GetTagBitSet().DebugGetIndividualNames(OutTagNames);
@@ -1038,7 +1037,7 @@ void FMassExecutionContext::FlushDeferred(UMassEntitySubsystem& EntitySystem) co
 
 void FMassExecutionContext::ClearExecutionData()
 {
-	ComponentViews.Reset();
+	FragmentViews.Reset();
 	CurrentArchetypesTagBitSet.Reset();
 	ChunkSerialModificationNumber = -1;
 }
@@ -1062,21 +1061,21 @@ void FMassExecutionContext::SetChunkCollection(FArchetypeChunkCollection&& InChu
 
 void FMassExecutionContext::SetRequirements(TConstArrayView<FMassFragmentRequirement> InRequirements, TConstArrayView<FMassFragmentRequirement> InChunkRequirements)
 { 
-	ComponentViews.Reset();
+	FragmentViews.Reset();
 	for (const FMassFragmentRequirement& Requirement : InRequirements)
 	{
 		if (Requirement.RequiresBinding())
 		{
-			ComponentViews.Add(FFragmentView(Requirement));
+			FragmentViews.Add(FFragmentView(Requirement));
 		}
 	}
 
-	ChunkComponents.Reset();
+	ChunkFragments.Reset();
 	for (const FMassFragmentRequirement& Requirement : InChunkRequirements)
 	{
 		if (Requirement.RequiresBinding())
 		{
-			ChunkComponents.Add(FChunkFragmentView(Requirement));
+			ChunkFragments.Add(FChunkFragmentView(Requirement));
 		}
 	}
 }
diff --git a/Engine/Plugins/Runtime/MassEntity/Source/MassEntity/Private/MassEntityView.cpp b/Engine/Plugins/Runtime/MassEntity/Source/MassEntity/Private/MassEntityView.cpp
--- a/Engine/Plugins/Runtime/MassEntity/Source/MassEntity/Private/MassEntityView.cpp
+++ b/Engine/Plugins/Runtime/MassEntity/Source/MassEntity/Private/MassEntityView.cpp
@@ -24,22 +24,22 @@ FMassEntityView::FMassEntityView(const UMassEntitySubsystem& EntitySubsystem, FM
 	EntityHandle = Archetype->MakeEntityHandle(Entity);
 }
 
-void* FMassEntityView::GetComponentPtr(const UScriptStruct& ComponentType) const
+void* FMassEntityView::GetFragmentPtr(const UScriptStruct& FragmentType) const
 {
 	checkSlow(Archetype && EntityHandle.IsValid());
-	if (const int32* ComponentIndex = Archetype->GetComponentIndex(&ComponentType))
+	if (const int32* FragmentIndex = Archetype->GetFragmentIndex(&FragmentType))
 	{
-		// failing the below Find means given entity's archetype is missing given ComponentType
-		return Archetype->GetComponentData(*ComponentIndex, EntityHandle);
+		// failing the below Find means given entity's archetype is missing given FragmentType
+		return Archetype->GetFragmentData(*FragmentIndex, EntityHandle);
 	}
 	return nullptr;
 }
 
-void* FMassEntityView::GetComponentPtrChecked(const UScriptStruct& ComponentType) const
+void* FMassEntityView::GetFragmentPtrChecked(const UScriptStruct& FragmentType) const
 {
 	checkSlow(Archetype && EntityHandle.IsValid());
-	const int32 ComponentIndex = Archetype->GetComponentIndexChecked(&ComponentType);
-	return Archetype->GetComponentData(ComponentIndex, EntityHandle);
+	const int32 FragmentIndex = Archetype->GetFragmentIndexChecked(&FragmentType);
+	return Archetype->GetFragmentData(FragmentIndex, EntityHandle);
 }
 
 bool FMassEntityView::HasTag(const UScriptStruct& TagType) const
diff --git a/Engine/Plugins/Runtime/MassEntity/Source/MassEntity/Public/MassArchetypeTypes.h b/Engine/Plugins/Runtime/MassEntity/Source/MassEntity/Public/MassArchetypeTypes.h
--- a/Engine/Plugins/Runtime/MassEntity/Source/MassEntity/Public/MassArchetypeTypes.h
+++ b/Engine/Plugins/Runtime/MassEntity/Source/MassEntity/Public/MassArchetypeTypes.h
@@ -138,11 +138,11 @@ struct FMassQueryRequirementIndicesMapping
 {
 	FMassQueryRequirementIndicesMapping() = default;
 
-	FMassFragmentIndicesMapping EntityComponents;
-	FMassFragmentIndicesMapping ChunkComponents;
+	FMassFragmentIndicesMapping EntityFragments;
+	FMassFragmentIndicesMapping ChunkFragments;
 	FORCEINLINE bool IsEmpty() const
 	{
-		return EntityComponents.Num() == 0 || ChunkComponents.Num() == 0;
+		return EntityFragments.Num() == 0 || ChunkFragments.Num() == 0;
 	}
 };
 
diff --git a/Engine/Plugins/Runtime/MassEntity/Source/MassEntity/Public/MassCommandBuffer.h b/Engine/Plugins/Runtime/MassEntity/Source/MassEntity/Public/MassCommandBuffer.h
--- a/Engine/Plugins/Runtime/MassEntity/Source/MassEntity/Public/MassCommandBuffer.h
+++ b/Engine/Plugins/Runtime/MassEntity/Source/MassEntity/Public/MassCommandBuffer.h
@@ -55,15 +55,15 @@ protected:
 };
 
 /**
- * Command dedicated to building an entity from an existing component instance
+ * Command dedicated to building an entity from an existing fragment instance
  */
 USTRUCT()
-struct MASSENTITY_API FBuildEntityFromComponentInstance : public FCommandBufferEntryBase
+struct MASSENTITY_API FBuildEntityFromFragmentInstance : public FCommandBufferEntryBase
 {
 	GENERATED_BODY()
 
-	FBuildEntityFromComponentInstance() = default;
-	FBuildEntityFromComponentInstance(const FMassEntityHandle Entity, FStructView InStruct)
+	FBuildEntityFromFragmentInstance() = default;
+	FBuildEntityFromFragmentInstance(const FMassEntityHandle Entity, FStructView InStruct)
 		: FCommandBufferEntryBase(Entity)
 		, Struct(InStruct)
 	{}
@@ -75,15 +75,15 @@ protected:
 };
 
 /**
-* Command dedicated to building an entity from a list of existing component instances
+* Command dedicated to building an entity from a list of existing fragment instances
 */
 USTRUCT()
-struct MASSENTITY_API FBuildEntityFromComponentInstances : public FCommandBufferEntryBase
+struct MASSENTITY_API FBuildEntityFromFragmentInstances : public FCommandBufferEntryBase
 {
 	GENERATED_BODY()
 
-	FBuildEntityFromComponentInstances() = default;
-	FBuildEntityFromComponentInstances(const FMassEntityHandle Entity, TConstArrayView<FInstancedStruct> InInstances)
+	FBuildEntityFromFragmentInstances() = default;
+	FBuildEntityFromFragmentInstances(const FMassEntityHandle Entity, TConstArrayView<FInstancedStruct> InInstances)
 		: FCommandBufferEntryBase(Entity)
 		, Instances(InInstances)
 	{}
@@ -95,15 +95,15 @@ protected:
 };
 
 /**
- * Command dedicated to add a new component to an existing entity
+ * Command dedicated to add a new fragment to an existing entity
  */
 USTRUCT()
-struct MASSENTITY_API FCommandAddComponent : public FCommandBufferEntryBase
+struct MASSENTITY_API FCommandAddFragment : public FCommandBufferEntryBase
 {
 	GENERATED_BODY()
 
-	FCommandAddComponent() = default;
-	FCommandAddComponent(FMassEntityHandle InEntity, UScriptStruct* InStruct)
+	FCommandAddFragment() = default;
+	FCommandAddFragment(FMassEntityHandle InEntity, UScriptStruct* InStruct)
 		: FCommandBufferEntryBase(InEntity)
 		, StructParam(InStruct)
 	{}
@@ -115,15 +115,15 @@ protected:
 };
 
 /**
- * Command dedicated to add a new component from an existing instance
+ * Command dedicated to add a new fragment from an existing instance
  */
 USTRUCT()
-struct MASSENTITY_API FCommandAddComponentInstance : public FCommandBufferEntryBase
+struct MASSENTITY_API FCommandAddFragmentInstance : public FCommandBufferEntryBase
 {
 	GENERATED_BODY()
 
-	FCommandAddComponentInstance() = default;
-	FCommandAddComponentInstance(const FMassEntityHandle InEntity, FStructView InStruct)
+	FCommandAddFragmentInstance() = default;
+	FCommandAddFragmentInstance(const FMassEntityHandle InEntity, FStructView InStruct)
         : FCommandBufferEntryBase(InEntity)
         , Struct(InStruct)
 	{}
@@ -135,15 +135,15 @@ protected:
 };
 
 /**
- * Command dedicated to remove a component from an existing entity
+ * Command dedicated to remove a fragment from an existing entity
  */
 USTRUCT()
-struct MASSENTITY_API FCommandRemoveComponent : public FCommandBufferEntryBase
+struct MASSENTITY_API FCommandRemoveFragment : public FCommandBufferEntryBase
 {
 	GENERATED_BODY()
 
-	FCommandRemoveComponent() = default;
-	FCommandRemoveComponent(FMassEntityHandle InEntity, UScriptStruct* InStruct)
+	FCommandRemoveFragment() = default;
+	FCommandRemoveFragment(FMassEntityHandle InEntity, UScriptStruct* InStruct)
 		: FCommandBufferEntryBase(InEntity)
 		, StructParam(InStruct)
 	{}
@@ -155,43 +155,43 @@ protected:
 };
 
 /**
- * Command performing addition of a list of components to an existing entity
+ * Command performing addition of a list of fragments to an existing entity
  */
 USTRUCT()
-struct MASSENTITY_API FCommandAddComponentList : public FCommandBufferEntryBase
+struct MASSENTITY_API FCommandAddFragmentList : public FCommandBufferEntryBase
 {
 	GENERATED_BODY()
 
-	FCommandAddComponentList() = default;
-	FCommandAddComponentList(FMassEntityHandle InEntity, TConstArrayView<const UScriptStruct*> InComponentList)
+	FCommandAddFragmentList() = default;
+	FCommandAddFragmentList(FMassEntityHandle InEntity, TConstArrayView<const UScriptStruct*> InFragmentList)
 		: FCommandBufferEntryBase(InEntity)
-		, ComponentList(InComponentList)
+		, FragmentList(InFragmentList)
 	{}
 
 protected:
 	virtual void Execute(UMassEntitySubsystem& System) const override;
 
-	TArray<const UScriptStruct*> ComponentList;
+	TArray<const UScriptStruct*> FragmentList;
 };
 
 /**
- * Command performing a removal of a list of components from an existing entity
+ * Command performing a removal of a list of fragments from an existing entity
  */
 USTRUCT()
-struct MASSENTITY_API FCommandRemoveComponentList : public FCommandBufferEntryBase
+struct MASSENTITY_API FCommandRemoveFragmentList : public FCommandBufferEntryBase
 {
 	GENERATED_BODY()
 
-	FCommandRemoveComponentList() = default;
-	FCommandRemoveComponentList(FMassEntityHandle InEntity, TConstArrayView<const UScriptStruct*> InComponentList)
+	FCommandRemoveFragmentList() = default;
+	FCommandRemoveFragmentList(FMassEntityHandle InEntity, TConstArrayView<const UScriptStruct*> InFragmentList)
 		: FCommandBufferEntryBase(InEntity)
-		, ComponentList(InComponentList)
+		, FragmentList(InFragmentList)
 	{}
 
 protected:
 	virtual void Execute(UMassEntitySubsystem& System) const override;
 
-	TArray<const UScriptStruct*> ComponentList;
+	TArray<const UScriptStruct*> FragmentList;
 };
 
 /**
@@ -262,7 +262,7 @@ protected:
 };
 
 /**
- * Command performing a removal of a collection of components and tags
+ * Command performing a removal of a collection of fragments and tags
  */
 USTRUCT()
 struct MASSENTITY_API FCommandRemoveComposition : public FCommandBufferEntryBase
@@ -309,17 +309,17 @@ public:
 	}
 
 	template<typename T>
-	void AddComponent(FMassEntityHandle Entity)
+	void AddFragment(FMassEntityHandle Entity)
 	{
-		static_assert(TIsDerivedFrom<T, FMassFragment>::IsDerived, "Given struct type is not a valid component type.");
-		EmplaceCommand<FCommandAddComponent>(Entity, T::StaticStruct());
+		static_assert(TIsDerivedFrom<T, FMassFragment>::IsDerived, "Given struct type is not a valid fragment type.");
+		EmplaceCommand<FCommandAddFragment>(Entity, T::StaticStruct());
 	}
 
 	template<typename T>
-	void RemoveComponent(FMassEntityHandle Entity)
+	void RemoveFragment(FMassEntityHandle Entity)
 	{
-		static_assert(TIsDerivedFrom<T, FMassFragment>::IsDerived, "Given struct type is not a valid component type.");
-		EmplaceCommand<FCommandRemoveComponent>(Entity, T::StaticStruct());
+		static_assert(TIsDerivedFrom<T, FMassFragment>::IsDerived, "Given struct type is not a valid fragment type.");
+		EmplaceCommand<FCommandRemoveFragment>(Entity, T::StaticStruct());
 	}
 
 	template<typename T>
diff --git a/Engine/Plugins/Runtime/MassEntity/Source/MassEntity/Public/MassEntityDebug.h b/Engine/Plugins/Runtime/MassEntity/Source/MassEntity/Public/MassEntityDebug.h
--- a/Engine/Plugins/Runtime/MassEntity/Source/MassEntity/Public/MassEntityDebug.h
+++ b/Engine/Plugins/Runtime/MassEntity/Source/MassEntity/Public/MassEntityDebug.h
@@ -13,7 +13,7 @@ MASSENTITY_API DECLARE_ENUM_TO_STRING(EMassProcessingPhase);
 
 namespace UE::Mass::Debug
 {
-	FString DebugGetComponentAccessString(EMassFragmentAccess Access);
+	FString DebugGetFragmentAccessString(EMassFragmentAccess Access);
 	MASSENTITY_API extern void DebugOutputDescription(TConstArrayView<UMassProcessor*> Processors, FOutputDevice& Ar);
 } // namespace UE::Mass::Debug
 
diff --git a/Engine/Plugins/Runtime/MassEntity/Source/MassEntity/Public/MassEntityQuery.h b/Engine/Plugins/Runtime/MassEntity/Source/MassEntity/Public/MassEntityQuery.h
--- a/Engine/Plugins/Runtime/MassEntity/Source/MassEntity/Public/MassEntityQuery.h
+++ b/Engine/Plugins/Runtime/MassEntity/Source/MassEntity/Public/MassEntityQuery.h
@@ -118,29 +118,29 @@ public:
 	 *	ArchetypeDataVersion properties). */
 	void CacheArchetypes(UMassEntitySubsystem& InEntitySubsystem);
 
-	FMassEntityQuery& AddRequirement(const UScriptStruct* ComponentType, const EMassFragmentAccess AccessMode, const EMassFragmentPresence Presence = EMassFragmentPresence::All)
+	FMassEntityQuery& AddRequirement(const UScriptStruct* FragmentType, const EMassFragmentAccess AccessMode, const EMassFragmentPresence Presence = EMassFragmentPresence::All)
 	{
-		checkf(Requirements.FindByPredicate([ComponentType](const FMassFragmentRequirement& Item){ return Item.StructType == ComponentType; }) == nullptr
-			, TEXT("Duplicated requirements are not supported. %s already present"), *GetNameSafe(ComponentType));
+		checkf(Requirements.FindByPredicate([FragmentType](const FMassFragmentRequirement& Item){ return Item.StructType == FragmentType; }) == nullptr
+			, TEXT("Duplicated requirements are not supported. %s already present"), *GetNameSafe(FragmentType));
 		
 		if (Presence != EMassFragmentPresence::None)
 		{
-			Requirements.Emplace(ComponentType, AccessMode, Presence);
+			Requirements.Emplace(FragmentType, AccessMode, Presence);
 		}
 
 		switch (Presence)
 		{
 		case EMassFragmentPresence::All:
-			RequiredAllComponents.Add(*ComponentType);
+			RequiredAllFragments.Add(*FragmentType);
 			break;
 		case EMassFragmentPresence::Any:
-			RequiredAnyComponents.Add(*ComponentType);
+			RequiredAnyFragments.Add(*FragmentType);
 			break;
 		case EMassFragmentPresence::Optional:
-			RequiredOptionalComponents.Add(*ComponentType);
+			RequiredOptionalFragments.Add(*FragmentType);
 			break;
 		case EMassFragmentPresence::None:
-			RequiredNoneComponents.Add(*ComponentType);
+			RequiredNoneFragments.Add(*FragmentType);
 			break;
 		}
 		// force recaching the next time this query is used or the following CacheArchetypes call.
@@ -165,16 +165,16 @@ public:
 		switch (Presence)
 		{
 		case EMassFragmentPresence::All:
-			RequiredAllComponents.Add<T>();
+			RequiredAllFragments.Add<T>();
 			break;
 		case EMassFragmentPresence::Any:
-			RequiredAnyComponents.Add<T>();
+			RequiredAnyFragments.Add<T>();
 			break;
 		case EMassFragmentPresence::Optional:
-			RequiredOptionalComponents.Add<T>();
+			RequiredOptionalFragments.Add<T>();
 			break;
 		case EMassFragmentPresence::None:
-			RequiredNoneComponents.Add<T>();
+			RequiredNoneFragments.Add<T>();
 			break;
 		}
 		// force recaching the next time this query is used or the following CacheArchetypes call.
@@ -240,12 +240,12 @@ public:
 		
 		if (Presence == EMassFragmentPresence::All)
 		{
-			RequiredAllChunkComponents.Add<T>();
+			RequiredAllChunkFragments.Add<T>();
 			ChunkRequirements.Emplace(T::StaticStruct(), AccessMode, Presence);
 		}
 		else
 		{
-			RequiredNoneChunkComponents.Add<T>();
+			RequiredNoneChunkFragments.Add<T>();
 		}
 
 		return *this;
@@ -273,15 +273,15 @@ public:
 	FString DebugGetDescription() const;
 
 	TConstArrayView<FMassFragmentRequirement> GetRequirements() const { return Requirements; }
-	const FMassFragmentBitSet& GetRequiredAllComponents() const { return RequiredAllComponents; }
-	const FMassFragmentBitSet& GetRequiredAnyComponents() const { return RequiredAnyComponents; }
-	const FMassFragmentBitSet& GetRequiredOptionalComponents() const { return RequiredOptionalComponents; }
-	const FMassFragmentBitSet& GetRequiredNoneComponents() const { return RequiredNoneComponents; }
+	const FMassFragmentBitSet& GetRequiredAllFragments() const { return RequiredAllFragments; }
+	const FMassFragmentBitSet& GetRequiredAnyFragments() const { return RequiredAnyFragments; }
+	const FMassFragmentBitSet& GetRequiredOptionalFragments() const { return RequiredOptionalFragments; }
+	const FMassFragmentBitSet& GetRequiredNoneFragments() const { return RequiredNoneFragments; }
 	const FMassTagBitSet& GetRequiredAllTags() const { return RequiredAllTags; }
 	const FMassTagBitSet& GetRequiredAnyTags() const { return RequiredAnyTags; }
 	const FMassTagBitSet& GetRequiredNoneTags() const { return RequiredNoneTags; }
-	const FMassChunkFragmentBitSet& GetRequiredAllChunkComponents() const { return RequiredAllChunkComponents; }
-	const FMassChunkFragmentBitSet& GetRequiredNoneChunkComponents() const { return RequiredNoneChunkComponents; }
+	const FMassChunkFragmentBitSet& GetRequiredAllChunkFragments() const { return RequiredAllChunkFragments; }
+	const FMassChunkFragmentBitSet& GetRequiredNoneChunkFragments() const { return RequiredNoneChunkFragments; }
 
 	const TArray<FArchetypeHandle>& GetArchetypes() const
 	{ 
@@ -323,12 +323,12 @@ protected:
 	FMassTagBitSet RequiredAllTags;
 	FMassTagBitSet RequiredAnyTags;
 	FMassTagBitSet RequiredNoneTags;
-	FMassFragmentBitSet RequiredAllComponents;
-	FMassFragmentBitSet RequiredAnyComponents;
-	FMassFragmentBitSet RequiredOptionalComponents;
-	FMassFragmentBitSet RequiredNoneComponents;
-	FMassChunkFragmentBitSet RequiredAllChunkComponents;
-	FMassChunkFragmentBitSet RequiredNoneChunkComponents;
+	FMassFragmentBitSet RequiredAllFragments;
+	FMassFragmentBitSet RequiredAnyFragments;
+	FMassFragmentBitSet RequiredOptionalFragments;
+	FMassFragmentBitSet RequiredNoneFragments;
+	FMassChunkFragmentBitSet RequiredAllChunkFragments;
+	FMassChunkFragmentBitSet RequiredNoneChunkFragments;
 
 private:
 	/** 
@@ -342,7 +342,7 @@ private:
 	uint32 ArchetypeDataVersion = 0;
 
 	TArray<FArchetypeHandle> ValidArchetypes;
-	TArray<FMassQueryRequirementIndicesMapping> ArchetypeComponentMapping;
+	TArray<FMassQueryRequirementIndicesMapping> ArchetypeFragmentMapping;
 
 	bool bAllowParallelExecution = false;
 };
diff --git a/Engine/Plugins/Runtime/MassEntity/Source/MassEntity/Public/MassEntitySubsystem.h b/Engine/Plugins/Runtime/MassEntity/Source/MassEntity/Public/MassEntitySubsystem.h
--- a/Engine/Plugins/Runtime/MassEntity/Source/MassEntity/Public/MassEntitySubsystem.h
+++ b/Engine/Plugins/Runtime/MassEntity/Source/MassEntity/Public/MassEntitySubsystem.h
@@ -77,33 +77,33 @@ public:
 	//~End of USubsystem interface
 	
 	/** 
-	 * A special, relaxed but slower version of CreateArchetype functions that allows ComponentAngTagsList to contain 
-	 * both components and tags. 
+	 * A special, relaxed but slower version of CreateArchetype functions that allows FragmentAngTagsList to contain 
+	 * both fragments and tags. 
 	 */
-	FArchetypeHandle CreateArchetype(TConstArrayView<const UScriptStruct*> ComponentsAngTagsList);
+	FArchetypeHandle CreateArchetype(TConstArrayView<const UScriptStruct*> FragmentsAngTagsList);
 
-	FArchetypeHandle CreateArchetype(TConstArrayView<const UScriptStruct*> ComponentList, const FMassTagBitSet& Tags, const FMassChunkFragmentBitSet& ChunkComponents);
+	FArchetypeHandle CreateArchetype(TConstArrayView<const UScriptStruct*> FragmentList, const FMassTagBitSet& Tags, const FMassChunkFragmentBitSet& ChunkFragments);
 
 	/** 
-	 *  Creates an archetype like SourceArchetype + NewComponentList. 
-	 *  @param SourceArchetype the archetype used to initially populate the list of components of the archetype being created. 
-	 *  @param NewComponentList list of unique components to add to components fetched from SourceArchetype. Note that 
+	 *  Creates an archetype like SourceArchetype + NewFragmentList. 
+	 *  @param SourceArchetype the archetype used to initially populate the list of fragments of the archetype being created. 
+	 *  @param NewFragmentList list of unique fragments to add to fragments fetched from SourceArchetype. Note that 
 	 *   adding an empty list is not supported and doing so will result in failing a `check`
 	 *  @return a handle of a new archetype
-	 *  @note it's caller's responsibility to ensure that NewComponentList is not empty and contains only component
-	 *   types that SourceArchetype doesn't already have. If the caller cannot guarantee it use of AddComponent functions
+	 *  @note it's caller's responsibility to ensure that NewFragmentList is not empty and contains only fragment
+	 *   types that SourceArchetype doesn't already have. If the caller cannot guarantee it use of AddFragment functions
 	 *   family is recommended.
 	 */
-	FArchetypeHandle CreateArchetype(const TSharedPtr<FMassArchetypeData>& SourceArchetype, const FMassFragmentBitSet& NewComponentList);
+	FArchetypeHandle CreateArchetype(const TSharedPtr<FMassArchetypeData>& SourceArchetype, const FMassFragmentBitSet& NewFragmentList);
 
 	FArchetypeHandle CreateArchetype(const FMassCompositionDescriptor& Descriptor);
 
 
 	FArchetypeHandle GetArchetypeForEntity(FMassEntityHandle Entity) const;
-	/** Method to iterate on all the component types of an archetype */
-	static void ForEachArchetypeComponentType(const FArchetypeHandle Archetype, TFunction< void(const UScriptStruct* /*ComponentType*/)> Function);
+	/** Method to iterate on all the fragment types of an archetype */
+	static void ForEachArchetypeFragmentType(const FArchetypeHandle Archetype, TFunction< void(const UScriptStruct* /*FragmentType*/)> Function);
 
-	void SetDefaultChunkComponentValue(const FArchetypeHandle Archetype, FConstStructView InstancedStruct);
+	void SetDefaultChunkFragmentValue(const FArchetypeHandle Archetype, FConstStructView InstancedStruct);
 
 	/**
 	 * Go through all archetypes and compact entities
@@ -119,9 +119,9 @@ public:
 
 	/**
 	 * Creates fully built entity ready to be used by the subsystem
-	 * @param ComponentInstanceList is the components to create the entity from and initialize values
+	 * @param FragmentInstanceList is the fragments to create the entity from and initialize values
 	 * @return FLWEntity id of the newly created entity */
-	FMassEntityHandle CreateEntity(TConstArrayView<FInstancedStruct> ComponentInstanceList);
+	FMassEntityHandle CreateEntity(TConstArrayView<FInstancedStruct> FragmentInstanceList);
 
 	/** A version of CreateEntity that's creating a number of entities (Count) at one go
 	 *  @param Archetype you want this entity to be
@@ -149,8 +149,8 @@ public:
 	/**
 	 * Builds an entity for it to be ready to be used by the subsystem
 	 * @param Entity to build which was retrieved with ReserveEntity() method
-	 * @param ComponentInstanceList is the components to create the entity from and initialize values*/
-	void BuildEntity(FMassEntityHandle Entity, TConstArrayView<FInstancedStruct> ComponentInstanceList);
+	 * @param FragmentInstanceList is the fragments to create the entity from and initialize values*/
+	void BuildEntity(FMassEntityHandle Entity, TConstArrayView<FInstancedStruct> FragmentInstanceList);
 
 	/*
 	 * Releases a previously reserved entity that was not yet built, otherwise call DestroyEntity
@@ -165,21 +165,21 @@ public:
 
 	void BatchDestroyEntityChunks(const FArchetypeChunkCollection& Chunks);
 
-	void AddComponentToEntity(FMassEntityHandle Entity, const UScriptStruct* ComponentType);
+	void AddFragmentToEntity(FMassEntityHandle Entity, const UScriptStruct* FragmentType);
 
 	/** 
-	 *  Ensures that only unique components are added. 
-	 *  @note It's caller's responsibility to ensure Entity's and ComponentList's validity. 
+	 *  Ensures that only unique fragments are added. 
+	 *  @note It's caller's responsibility to ensure Entity's and FragmentList's validity. 
 	 */
-	void AddComponentListToEntity(FMassEntityHandle Entity, TConstArrayView<const UScriptStruct*> ComponentList);
+	void AddFragmentListToEntity(FMassEntityHandle Entity, TConstArrayView<const UScriptStruct*> FragmentList);
 
-	void AddComponentInstanceListToEntity(FMassEntityHandle Entity, TConstArrayView<FInstancedStruct> ComponentInstanceList);
-	void RemoveComponentFromEntity(FMassEntityHandle Entity, const UScriptStruct* ComponentType);
-	void RemoveComponentListFromEntity(FMassEntityHandle Entity, TConstArrayView<const UScriptStruct*> ComponentList);
+	void AddFragmentInstanceListToEntity(FMassEntityHandle Entity, TConstArrayView<FInstancedStruct> FragmentInstanceList);
+	void RemoveFragmentFromEntity(FMassEntityHandle Entity, const UScriptStruct* FragmentType);
+	void RemoveFragmentListFromEntity(FMassEntityHandle Entity, TConstArrayView<const UScriptStruct*> FragmentList);
 
 	void AddTagToEntity(FMassEntityHandle Entity, const UScriptStruct* TagType);
 	void RemoveTagFromEntity(FMassEntityHandle Entity, const UScriptStruct* TagType);
-	void SwapTagsForEntity(FMassEntityHandle Entity, const UScriptStruct* FromComponentType, const UScriptStruct* ToComponentType);
+	void SwapTagsForEntity(FMassEntityHandle Entity, const UScriptStruct* FromFragmentType, const UScriptStruct* ToFragmentType);
 
 
 	/**
@@ -191,20 +191,20 @@ public:
 	void RemoveCompositionFromEntity(FMassEntityHandle Entity, const FMassCompositionDescriptor& InDescriptor);
 
 	/** 
-	 * Moves an entity over to a new archetype by copying over components common to both archetypes
+	 * Moves an entity over to a new archetype by copying over fragments common to both archetypes
 	 * @param Entity is the entity to move 
 	 * @param NewArchetypeHandle the handle to the new archetype
 	 */
 	void MoveEntityToAnotherArchetype(FMassEntityHandle Entity, FArchetypeHandle NewArchetypeHandle);
 
-	/** Copies values from ComponentInstanceList over to Entity's component. Caller is responsible for ensuring that 
-	 *  the given entity does have given components. Failing this assumption will cause a check-fail.*/
-	void SetEntityComponentsValues(FMassEntityHandle Entity, TArrayView<const FInstancedStruct> ComponentInstanceList);
+	/** Copies values from FragmentInstanceList over to Entity's fragment. Caller is responsible for ensuring that 
+	 *  the given entity does have given fragments. Failing this assumption will cause a check-fail.*/
+	void SetEntityFragmentsValues(FMassEntityHandle Entity, TArrayView<const FInstancedStruct> FragmentInstanceList);
 
-	/** Copies values from ComponentInstanceList over to components of given entities collection. The caller is responsible 
-	 *  for ensuring that the given entity archetype (FArchetypeChunkCollection.Archetype) does have given components. 
+	/** Copies values from FragmentInstanceList over to fragments of given entities collection. The caller is responsible 
+	 *  for ensuring that the given entity archetype (FArchetypeChunkCollection.Archetype) does have given fragments. 
 	 *  Failing this assumption will cause a check-fail. */
-	static void BatchSetEntityComponentsValues(const FArchetypeChunkCollection& SparseEntities, TArrayView<const FInstancedStruct> ComponentInstanceList);
+	static void BatchSetEntityFragmentsValues(const FArchetypeChunkCollection& SparseEntities, TArrayView<const FInstancedStruct> FragmentInstanceList);
 
 	// Return true if it is an valid built entity
 	bool IsEntityActive(FMassEntityHandle Entity) const 
@@ -224,21 +224,21 @@ public:
 	// Asserts that IsEntityBuilt
 	void CheckIfEntityIsActive(FMassEntityHandle Entity) const;
 
-	template <typename ComponentType>
-	ComponentType& GetComponentDataChecked(FMassEntityHandle Entity) const
+	template <typename FragmentType>
+	FragmentType& GetFragmentDataChecked(FMassEntityHandle Entity) const
 	{
-		return *((ComponentType*)InternalGetComponentDataChecked(Entity, ComponentType::StaticStruct()));
+		return *((FragmentType*)InternalGetFragmentDataChecked(Entity, FragmentType::StaticStruct()));
 	}
 
-	template <typename ComponentType>
-	ComponentType* GetComponentDataPtr(FMassEntityHandle Entity) const
+	template <typename FragmentType>
+	FragmentType* GetFragmentDataPtr(FMassEntityHandle Entity) const
 	{
-		return (ComponentType*)InternalGetComponentDataPtr(Entity, ComponentType::StaticStruct());
+		return (FragmentType*)InternalGetFragmentDataPtr(Entity, FragmentType::StaticStruct());
 	}
 
-	FStructView GetComponentDataStruct(FMassEntityHandle Entity, const UScriptStruct* ComponentType) const
+	FStructView GetFragmentDataStruct(FMassEntityHandle Entity, const UScriptStruct* FragmentType) const
 	{
-		return FStructView(ComponentType, static_cast<uint8*>(InternalGetComponentDataPtr(Entity, ComponentType)));
+		return FStructView(FragmentType, static_cast<uint8*>(InternalGetFragmentDataPtr(Entity, FragmentType)));
 	}
 
 	uint32 GetArchetypeDataVersion() const { return ArchetypeDataVersion; }
@@ -258,13 +258,13 @@ public:
 	void DebugPrintArchetypes(FOutputDevice& Ar) const;
 	static void DebugGetStringDesc(const FArchetypeHandle& Archetype, FOutputDevice& Ar);
 	void DebugGetArchetypesStringDetails(FOutputDevice& Ar, const bool bIncludeEmpty = true);
-	void DebugGetArchetypeComponentTypes(const FArchetypeHandle& Archetype, TArray<const UScriptStruct*>& InOutComponentList) const;
+	void DebugGetArchetypeFragmentTypes(const FArchetypeHandle& Archetype, TArray<const UScriptStruct*>& InOutFragmentList) const;
 	int32 DebugGetArchetypeEntitiesCount(const FArchetypeHandle& Archetype) const;	
 	int32 DebugGetEntityCount() const { return Entities.Num() - NumReservedEntities - EntityFreeIndexList.Num(); }
-	int32 DebugGetArchetypesCount() const { return ComponentHashToArchetypeMap.Num(); }
+	int32 DebugGetArchetypesCount() const { return FragmentHashToArchetypeMap.Num(); }
 	void DebugRemoveAllEntities();
 	void DebugForceArchetypeDataVersionBump() { ++ArchetypeDataVersion; }
-	void DebugGetArchetypeStrings(const FArchetypeHandle& Archetype, TArray<FName>& OutComponentNames, TArray<FName>& OutTagNames);
+	void DebugGetArchetypeStrings(const FArchetypeHandle& Archetype, TArray<FName>& OutFragmentNames, TArray<FName>& OutTagNames);
 #endif // WITH_MASSENTITY_DEBUG
 
 protected:
@@ -277,17 +277,17 @@ private:
 	void InternalReleaseEntity(FMassEntityHandle Entity);
 
 	/** 
-	 *  Adds components in ComponentList to Entity. Only the unique components will be added.
+	 *  Adds fragments in FragmentList to Entity. Only the unique fragments will be added.
 	 */
-	void InternalAddComponentListToEntityChecked(FMassEntityHandle Entity, const FMassFragmentBitSet& InComponents);
+	void InternalAddFragmentListToEntityChecked(FMassEntityHandle Entity, const FMassFragmentBitSet& InFragments);
 
 	/** 
-	 *  Similar to InternalAddComponentListToEntity but expects NewComponentList not overlapping with current entity's
-	 *  component list. It's callers responsibility to ensure that's true. Failing this will cause a `check` fail.
+	 *  Similar to InternalAddFragmentListToEntity but expects NewFragmentList not overlapping with current entity's
+	 *  fragment list. It's callers responsibility to ensure that's true. Failing this will cause a `check` fail.
 	 */
-	void InternalAddComponentListToEntity(FMassEntityHandle Entity, const FMassFragmentBitSet& NewComponents);
-	void* InternalGetComponentDataChecked(FMassEntityHandle Entity, const UScriptStruct* ComponentType) const;
-	void* InternalGetComponentDataPtr(FMassEntityHandle Entity, const UScriptStruct* ComponentType) const;
+	void InternalAddFragmentListToEntity(FMassEntityHandle Entity, const FMassFragmentBitSet& NewFragments);
+	void* InternalGetFragmentDataChecked(FMassEntityHandle Entity, const UScriptStruct* FragmentType) const;
+	void* InternalGetFragmentDataPtr(FMassEntityHandle Entity, const UScriptStruct* FragmentType) const;
 
 private:
 	TChunkedArray<FEntityData> Entities;
@@ -300,11 +300,11 @@ private:
 	// the "version" number increased every time an archetype gets added
 	uint32 ArchetypeDataVersion = 0;
 
-	// Map of hash of sorted component list to archetypes with that hash
-	TMap<uint32, TArray<TSharedPtr<FMassArchetypeData>>> ComponentHashToArchetypeMap;
+	// Map of hash of sorted fragment list to archetypes with that hash
+	TMap<uint32, TArray<TSharedPtr<FMassArchetypeData>>> FragmentHashToArchetypeMap;
 
-	// Map to list of archetypes that contain the specified component type
-	TMap<const UScriptStruct*, TArray<TSharedPtr<FMassArchetypeData>>> ComponentTypeToArchetypeMap;
+	// Map to list of archetypes that contain the specified fragment type
+	TMap<const UScriptStruct*, TArray<TSharedPtr<FMassArchetypeData>>> FragmentTypeToArchetypeMap;
 };
 
 
@@ -317,26 +317,26 @@ private:
 	struct FFragmentView 
 	{
 		FMassFragmentRequirement Requirement;
-		TArrayView<FMassFragment> ComponentView;
+		TArrayView<FMassFragment> FragmentView;
 
 		FFragmentView() {}
 		explicit FFragmentView(const FMassFragmentRequirement& InRequirement) : Requirement(InRequirement) {}
 
-		bool operator==(const UScriptStruct* ComponentType) const { return Requirement.StructType == ComponentType; }
+		bool operator==(const UScriptStruct* FragmentType) const { return Requirement.StructType == FragmentType; }
 	};
-	TArray<FFragmentView, TInlineAllocator<8>> ComponentViews;
+	TArray<FFragmentView, TInlineAllocator<8>> FragmentViews;
 
 	struct FChunkFragmentView
 	{
 		FMassFragmentRequirement Requirement;
-		FStructView ChunkComponentView;
+		FStructView ChunkFragmentView;
 
 		FChunkFragmentView() {}
 		explicit FChunkFragmentView(const FMassFragmentRequirement& InRequirement) : Requirement(InRequirement)	{}
 
-		bool operator==(const UScriptStruct* ComponentType) const { return Requirement.StructType == ComponentType; }
+		bool operator==(const UScriptStruct* FragmentType) const { return Requirement.StructType == FragmentType; }
 	};
-	TArray<FChunkFragmentView, TInlineAllocator<4>> ChunkComponents;
+	TArray<FChunkFragmentView, TInlineAllocator<4>> ChunkFragments;
 
 	// mz@todo make this shared ptr thread-safe and never auto-flush in MT environment. 
 	TSharedPtr<FMassCommandBuffer> DeferredCommandBuffer;
@@ -360,8 +360,8 @@ private:
 	 *  calling FLWCCommandBuffer.ReplayBufferAgainstSystem() */
 	bool bFlushDeferredCommands = true;
 
-	TArrayView<FFragmentView> GetMutableRequirements() { return ComponentViews; }
-	TArrayView<FChunkFragmentView> GetMutableChunkRequirements() { return ChunkComponents; }
+	TArrayView<FFragmentView> GetMutableRequirements() { return FragmentViews; }
+	TArrayView<FChunkFragmentView> GetMutableChunkRequirements() { return ChunkFragments; }
 	
 	friend FMassArchetypeData;
 	friend FMassEntityQuery;
@@ -415,58 +415,58 @@ public:
 	int32 GetChunkSerialModificationNumber() const { return ChunkSerialModificationNumber; }
 
 	template<typename T>
-	T& GetMutableChunkComponent() const
+	T& GetMutableChunkFragment() const
 	{
 		static_assert(TIsDerivedFrom<T, FMassChunkFragment>::IsDerived, "Given struct doesn't represent a valid chunk fragment type. Make sure to inherit from FMassChunkFragment or one of its child-types.");
 
 		const UScriptStruct* Type = T::StaticStruct();
-		const FChunkFragmentView* FoundChunkComponentData = ChunkComponents.FindByPredicate([Type](const FChunkFragmentView& Element) { return Element.Requirement.StructType == Type; });
-		checkf(FoundChunkComponentData, TEXT("Chunk Component requirement not found: %s"), *T::StaticStruct()->GetName());
-		return FoundChunkComponentData->ChunkComponentView.GetMutable<T>();
+		const FChunkFragmentView* FoundChunkFragmentData = ChunkFragments.FindByPredicate([Type](const FChunkFragmentView& Element) { return Element.Requirement.StructType == Type; });
+		checkf(FoundChunkFragmentData, TEXT("Chunk Fragment requirement not found: %s"), *T::StaticStruct()->GetName());
+		return FoundChunkFragmentData->ChunkFragmentView.GetMutable<T>();
 	}
 
 	template<typename T>
-	const T& GetChunkComponent() const
+	const T& GetChunkFragment() const
 	{
 		static_assert(TIsDerivedFrom<T, FMassChunkFragment>::IsDerived, "Given struct doesn't represent a valid chunk fragment type. Make sure to inherit from FMassChunkFragment or one of its child-types.");
 
 		const UScriptStruct* Type = T::StaticStruct();
-		const FChunkFragmentView* FoundChunkComponentData = ChunkComponents.FindByPredicate([Type](const FChunkFragmentView& Element) { return Element.Requirement.StructType == Type; } );
-		checkf(FoundChunkComponentData, TEXT("Chunk Component requirement not found: %s"), *T::StaticStruct()->GetName());
-		return FoundChunkComponentData->ChunkComponentView.Get<T>();
+		const FChunkFragmentView* FoundChunkFragmentData = ChunkFragments.FindByPredicate([Type](const FChunkFragmentView& Element) { return Element.Requirement.StructType == Type; } );
+		checkf(FoundChunkFragmentData, TEXT("Chunk Fragment requirement not found: %s"), *T::StaticStruct()->GetName());
+		return FoundChunkFragmentData->ChunkFragmentView.Get<T>();
 	}
 
 	template<typename TFragment>
-	TArrayView<TFragment> GetMutableComponentView()
+	TArrayView<TFragment> GetMutableFragmentView()
 	{
-		const UScriptStruct* ComponentType = TFragment::StaticStruct();
-		const FFragmentView* View = ComponentViews.FindByPredicate([ComponentType](const FFragmentView& Element) { return Element.Requirement.StructType == ComponentType; });
-		//checkfSlow(View != nullptr, TEXT("Requested component type not bound"));
-		//checkfSlow(View->Requirement.AccessMode == EMassFragmentAccess::ReadWrite, TEXT("Requested component has not been bound for writing"));
-		return MakeArrayView<TFragment>((TFragment*)View->ComponentView.GetData(), View->ComponentView.Num());
+		const UScriptStruct* FragmentType = TFragment::StaticStruct();
+		const FFragmentView* View = FragmentViews.FindByPredicate([FragmentType](const FFragmentView& Element) { return Element.Requirement.StructType == FragmentType; });
+		//checkfSlow(View != nullptr, TEXT("Requested fragment type not bound"));
+		//checkfSlow(View->Requirement.AccessMode == EMassFragmentAccess::ReadWrite, TEXT("Requested fragment has not been bound for writing"));
+		return MakeArrayView<TFragment>((TFragment*)View->FragmentView.GetData(), View->FragmentView.Num());
 	}
 
 	template<typename TFragment>
-	TConstArrayView<TFragment> GetComponentView() const
+	TConstArrayView<TFragment> GetFragmentView() const
 	{
-		const UScriptStruct* ComponentType = TFragment::StaticStruct();
-		const FFragmentView* View = ComponentViews.FindByPredicate([ComponentType](const FFragmentView& Element) { return Element.Requirement.StructType == ComponentType; });
-		//checkfSlow(View != nullptr, TEXT("Requested component type not bound"));
-		return TConstArrayView<TFragment>((const TFragment*)View->ComponentView.GetData(), View->ComponentView.Num());
+		const UScriptStruct* FragmentType = TFragment::StaticStruct();
+		const FFragmentView* View = FragmentViews.FindByPredicate([FragmentType](const FFragmentView& Element) { return Element.Requirement.StructType == FragmentType; });
+		//checkfSlow(View != nullptr, TEXT("Requested fragment type not bound"));
+		return TConstArrayView<TFragment>((const TFragment*)View->FragmentView.GetData(), View->FragmentView.Num());
 	}
 
-	TConstArrayView<FMassFragment> GetComponentComponentView(const UScriptStruct* ComponentType) const
+	TConstArrayView<FMassFragment> GetFragmentFragmentView(const UScriptStruct* FragmentType) const
 	{
-		const FFragmentView* View = ComponentViews.FindByPredicate([ComponentType](const FFragmentView& Element) { return Element.Requirement.StructType == ComponentType; });
+		const FFragmentView* View = FragmentViews.FindByPredicate([FragmentType](const FFragmentView& Element) { return Element.Requirement.StructType == FragmentType; });
 		checkSlow(View);
-		return TConstArrayView<FMassFragment>((const FMassFragment*)View->ComponentView.GetData(), View->ComponentView.Num());;
+		return TConstArrayView<FMassFragment>((const FMassFragment*)View->FragmentView.GetData(), View->FragmentView.Num());;
 	}
 
-	TArrayView<FMassFragment> GetMutableComponentView(const UScriptStruct* ComponentType) 
+	TArrayView<FMassFragment> GetMutableFragmentView(const UScriptStruct* FragmentType) 
 	{
-		const FFragmentView* View = ComponentViews.FindByPredicate([ComponentType](const FFragmentView& Element) { return Element.Requirement.StructType == ComponentType; });
+		const FFragmentView* View = FragmentViews.FindByPredicate([FragmentType](const FFragmentView& Element) { return Element.Requirement.StructType == FragmentType; });
 		checkSlow(View);
-		return View->ComponentView;
+		return View->FragmentView;
 	}
 
 	/** Sparse chunk related operation */
@@ -482,15 +482,15 @@ public:
 
 protected:
 	void SetRequirements(TConstArrayView<FMassFragmentRequirement> InRequirements, TConstArrayView<FMassFragmentRequirement> InChunkRequirements);
-	void ClearComponentViews()
+	void ClearFragmentViews()
 	{
-		for (FFragmentView& View : ComponentViews)
+		for (FFragmentView& View : FragmentViews)
 		{
-			View.ComponentView = TArrayView<FMassFragment>();
+			View.FragmentView = TArrayView<FMassFragment>();
 		}
-		for (FChunkFragmentView& View : ChunkComponents)
+		for (FChunkFragmentView& View : ChunkFragments)
 		{
-			View.ChunkComponentView.Reset();
+			View.ChunkFragmentView.Reset();
 		}
 	}
 };
diff --git a/Engine/Plugins/Runtime/MassEntity/Source/MassEntity/Public/MassEntityTypes.h b/Engine/Plugins/Runtime/MassEntity/Source/MassEntity/Public/MassEntityTypes.h
--- a/Engine/Plugins/Runtime/MassEntity/Source/MassEntity/Public/MassEntityTypes.h
+++ b/Engine/Plugins/Runtime/MassEntity/Source/MassEntity/Public/MassEntityTypes.h
@@ -8,7 +8,7 @@
 
 MASSENTITY_API DECLARE_LOG_CATEGORY_EXTERN(LogMass, Warning, All);
 
-// This is the base class for all lightweight components
+// This is the base class for all lightweight fragments
 USTRUCT()
 struct FMassFragment
 {
@@ -36,7 +36,7 @@ struct FMassChunkFragment
 };
 
 // A handle to a lightweight entity.  An entity is used in conjunction with the UMassEntitySubsystem
-// for the current world and can contain lightweight components.
+// for the current world and can contain lightweight fragments.
 USTRUCT()
 struct FMassEntityHandle
 {
@@ -96,58 +96,58 @@ template struct MASSENTITY_API TScriptStructTypeBitSet<FMassChunkFragment>;
 using FMassChunkFragmentBitSet = TScriptStructTypeBitSet<FMassChunkFragment>;
 
 /** The type summarily describing a composition of an entity or an archetype. It contains information on both the
- *  components as well as tags */
+ *  fragments as well as tags */
 struct FMassCompositionDescriptor
 {
 	FMassCompositionDescriptor() = default;
-	FMassCompositionDescriptor(const FMassFragmentBitSet& InComponents, const FMassTagBitSet& InTags, const FMassChunkFragmentBitSet& InChunkComponents)
-		: Components(InComponents), Tags(InTags), ChunkComponents(InChunkComponents)
+	FMassCompositionDescriptor(const FMassFragmentBitSet& InFragments, const FMassTagBitSet& InTags, const FMassChunkFragmentBitSet& InChunkFragments)
+		: Fragments(InFragments), Tags(InTags), ChunkFragments(InChunkFragments)
 	{}
 
-	FMassCompositionDescriptor(TConstArrayView<const UScriptStruct*> InComponents, const FMassTagBitSet& InTags, const FMassChunkFragmentBitSet& InChunkComponents)
-		: FMassCompositionDescriptor(FMassFragmentBitSet(InComponents), InTags, InChunkComponents)
+	FMassCompositionDescriptor(TConstArrayView<const UScriptStruct*> InFragments, const FMassTagBitSet& InTags, const FMassChunkFragmentBitSet& InChunkFragments)
+		: FMassCompositionDescriptor(FMassFragmentBitSet(InFragments), InTags, InChunkFragments)
 	{}
 
-	FMassCompositionDescriptor(TConstArrayView<FInstancedStruct> InComponentInstances, const FMassTagBitSet& InTags, const FMassChunkFragmentBitSet& InChunkComponents)
-		: FMassCompositionDescriptor(FMassFragmentBitSet(InComponentInstances), InTags, InChunkComponents)
+	FMassCompositionDescriptor(TConstArrayView<FInstancedStruct> InFragmentInstances, const FMassTagBitSet& InTags, const FMassChunkFragmentBitSet& InChunkFragments)
+		: FMassCompositionDescriptor(FMassFragmentBitSet(InFragmentInstances), InTags, InChunkFragments)
 	{}
 
-	FMassCompositionDescriptor(FMassFragmentBitSet&& InComponents, FMassTagBitSet&& InTags, FMassChunkFragmentBitSet&& InChunkComponents)
-		: Components(MoveTemp(InComponents)), Tags(MoveTemp(InTags)), ChunkComponents(MoveTemp(InChunkComponents))
+	FMassCompositionDescriptor(FMassFragmentBitSet&& InFragments, FMassTagBitSet&& InTags, FMassChunkFragmentBitSet&& InChunkFragments)
+		: Fragments(MoveTemp(InFragments)), Tags(MoveTemp(InTags)), ChunkFragments(MoveTemp(InChunkFragments))
 	{}
 
 	void Reset()
 	{
-		Components.Reset();
+		Fragments.Reset();
 		Tags.Reset();
 	}
 
-	bool IsEquivalent(const FMassFragmentBitSet& InComponentBitSet, const FMassTagBitSet& InTagBitSet, const FMassChunkFragmentBitSet& InChunkComponentsBitSet) const
+	bool IsEquivalent(const FMassFragmentBitSet& InFragmentBitSet, const FMassTagBitSet& InTagBitSet, const FMassChunkFragmentBitSet& InChunkFragmentsBitSet) const
 	{
-		return Components.IsEquivalent(InComponentBitSet) && Tags.IsEquivalent(InTagBitSet) && ChunkComponents.IsEquivalent(InChunkComponentsBitSet);
+		return Fragments.IsEquivalent(InFragmentBitSet) && Tags.IsEquivalent(InTagBitSet) && ChunkFragments.IsEquivalent(InChunkFragmentsBitSet);
 	}
 
 	bool IsEquivalent(const FMassCompositionDescriptor& OtherDescriptor) const
 	{
-		return IsEquivalent(OtherDescriptor.Components, OtherDescriptor.Tags, OtherDescriptor.ChunkComponents);
+		return IsEquivalent(OtherDescriptor.Fragments, OtherDescriptor.Tags, OtherDescriptor.ChunkFragments);
 	}
 
-	bool IsEmpty() const { return Components.IsEmpty() && Tags.IsEmpty() && ChunkComponents.IsEmpty(); }
+	bool IsEmpty() const { return Fragments.IsEmpty() && Tags.IsEmpty() && ChunkFragments.IsEmpty(); }
 
-	static uint32 CalculateHash(const FMassFragmentBitSet& InComponents, const FMassTagBitSet& InTags, const FMassChunkFragmentBitSet& InChunkComponents)
+	static uint32 CalculateHash(const FMassFragmentBitSet& InFragments, const FMassTagBitSet& InTags, const FMassChunkFragmentBitSet& InChunkFragments)
 	{
-		const uint32 ComponentsHash = GetTypeHash(InComponents);
+		const uint32 FragmentsHash = GetTypeHash(InFragments);
 		const uint32 TagsHash = GetTypeHash(InTags);
-		const uint32 ChunkComponentsHash = GetTypeHash(InChunkComponents);
-		return HashCombine(HashCombine(ComponentsHash, TagsHash), ChunkComponentsHash);
+		const uint32 ChunkFragmentsHash = GetTypeHash(InChunkFragments);
+		return HashCombine(HashCombine(FragmentsHash, TagsHash), ChunkFragmentsHash);
 	}	
 
 	uint32 CalculateHash() const 
 	{
-		return CalculateHash(Components, Tags, ChunkComponents);
+		return CalculateHash(Fragments, Tags, ChunkFragments);
 	}
 
-	FMassFragmentBitSet Components;
+	FMassFragmentBitSet Fragments;
 	FMassTagBitSet Tags;
-	FMassChunkFragmentBitSet ChunkComponents;
+	FMassChunkFragmentBitSet ChunkFragments;
 };
diff --git a/Engine/Plugins/Runtime/MassEntity/Source/MassEntity/Public/MassEntityView.h b/Engine/Plugins/Runtime/MassEntity/Source/MassEntity/Public/MassEntityView.h
--- a/Engine/Plugins/Runtime/MassEntity/Source/MassEntity/Public/MassEntityView.h
+++ b/Engine/Plugins/Runtime/MassEntity/Source/MassEntity/Public/MassEntityView.h
@@ -39,34 +39,34 @@ struct MASSENTITY_API FMassEntityView
 
 	FMassEntityHandle GetEntity() const	{ return Entity; }
 
-	/** will fail a check if the viewed entity doesn't have the given component */	
+	/** will fail a check if the viewed entity doesn't have the given fragment */	
 	template<typename T>
-	T& GetComponentData() const
+	T& GetFragmentData() const
 	{
 		static_assert(!TIsDerivedFrom<T, FMassTag>::IsDerived,
 			"Given struct doesn't represent a valid fragment type but a tag. Use HasTag instead.");
 		static_assert(TIsDerivedFrom<T, FMassTag>::IsDerived || TIsDerivedFrom<T, FMassFragment>::IsDerived,
 			"Given struct doesn't represent a valid fragment type. Make sure to inherit from FMassFragment or one of its child-types.");
 
-		return *((T*)GetComponentPtrChecked(*T::StaticStruct()));
+		return *((T*)GetFragmentPtrChecked(*T::StaticStruct()));
 	}
 		
-	/** if the viewed entity doesn't have the given component the function will return null */
+	/** if the viewed entity doesn't have the given fragment the function will return null */
 	template<typename T>
-	T* GetComponentDataPtr() const
+	T* GetFragmentDataPtr() const
 	{
 		static_assert(!TIsDerivedFrom<T, FMassTag>::IsDerived,
 			"Given struct doesn't represent a valid fragment type but a tag. Use HasTag instead.");
 		static_assert(TIsDerivedFrom<T, FMassTag>::IsDerived || TIsDerivedFrom<T, FMassFragment>::IsDerived,
 			"Given struct doesn't represent a valid fragment type. Make sure to inherit from FMassFragment or one of its child-types.");
 
-		return (T*)GetComponentPtr(*T::StaticStruct());
+		return (T*)GetFragmentPtr(*T::StaticStruct());
 	}
 
-	FStructView GetComponentDataStruct(const UScriptStruct* ComponentType) const
+	FStructView GetFragmentDataStruct(const UScriptStruct* FragmentType) const
 	{
-		check(ComponentType);
-		return FStructView(ComponentType, static_cast<uint8*>(GetComponentPtr(*ComponentType)));
+		check(FragmentType);
+		return FStructView(FragmentType, static_cast<uint8*>(GetFragmentPtr(*FragmentType)));
 	}
 
 	template<typename T>
@@ -80,8 +80,8 @@ struct MASSENTITY_API FMassEntityView
 	bool operator==(const FMassEntityView& Other) const { return Archetype == Other.Archetype && EntityHandle == Other.EntityHandle; }
 
 protected:
-	void* GetComponentPtr(const UScriptStruct& ComponentType) const;
-	void* GetComponentPtrChecked(const UScriptStruct& ComponentType) const;
+	void* GetFragmentPtr(const UScriptStruct& FragmentType) const;
+	void* GetFragmentPtrChecked(const UScriptStruct& FragmentType) const;
 	bool HasTag(const UScriptStruct& TagType) const;
 
 private:
diff --git a/Engine/Plugins/Runtime/MassEntity/Source/MassEntityTestSuite/Private/LWComponentTestFarmPlot.cpp b/Engine/Plugins/Runtime/MassEntity/Source/MassEntityTestSuite/Private/LWComponentTestFarmPlot.cpp
--- a/Engine/Plugins/Runtime/MassEntity/Source/MassEntityTestSuite/Private/LWComponentTestFarmPlot.cpp
+++ b/Engine/Plugins/Runtime/MassEntity/Source/MassEntityTestSuite/Private/LWComponentTestFarmPlot.cpp
@@ -5,13 +5,13 @@
 #include "Components/HierarchicalInstancedStaticMeshComponent.h"
 #include "Engine/CollisionProfile.h"
 
-//@TODO: Can add a ReadyToHarvest tag component on when things are ready to harvest, to stop them ticking and signal that we need to create an icon
+//@TODO: Can add a ReadyToHarvest tag Fragment on when things are ready to harvest, to stop them ticking and signal that we need to create an icon
 
 
 //////////////////////////////////////////////////////////////////////////
-// UComponentUpdateSystem
+// UFragmentUpdateSystem
 
-void UComponentUpdateSystem::PostInitProperties()
+void UFragmentUpdateSystem::PostInitProperties()
 {
 	Super::PostInitProperties();
 	ConfigureQueries();
@@ -27,8 +27,8 @@ void UFarmHarvestTimerSetIcon::Execute(UMassEntitySubsystem& EntitySubsystem, FM
 	EntityQuery.ForEachEntityChunk(EntitySubsystem, Context, [this](FMassExecutionContext& Context) {
 
 		const int32 NumEntities = Context.GetNumEntities();
-		TConstArrayView<FFarmGridCellData> GridCoordList = Context.GetComponentView<FFarmGridCellData>();
-		TArrayView<FFarmVisualComponent> VisualList = Context.GetMutableComponentView<FFarmVisualComponent>();
+		TConstArrayView<FFarmGridCellData> GridCoordList = Context.GetFragmentView<FFarmGridCellData>();
+		TArrayView<FFarmVisualFragment> VisualList = Context.GetMutableFragmentView<FFarmVisualFragment>();
 
 		for (int32 i = 0; i < NumEntities; ++i)
 		{
@@ -47,9 +47,9 @@ void UFarmHarvestTimerSetIcon::Execute(UMassEntitySubsystem& EntitySubsystem, FM
 }
 
 //////////////////////////////////////////////////////////////////////////
-// ALWComponentTestFarmPlot
+// ALWFragmentTestFarmPlot
 
-ALWComponentTestFarmPlot::ALWComponentTestFarmPlot()
+ALWFragmentTestFarmPlot::ALWFragmentTestFarmPlot()
 {
 	PrimaryActorTick.bCanEverTick = true;
 	PrimaryActorTick.bStartWithTickEnabled = true;
@@ -64,37 +64,37 @@ ALWComponentTestFarmPlot::ALWComponentTestFarmPlot()
 	HarvestIconISMC->SetCollisionProfileName(UCollisionProfile::NoCollision_ProfileName);
 }
 
-void ALWComponentTestFarmPlot::AddItemToGrid(UMassEntitySubsystem* EntitySystem, int32 X, int32 Y, FArchetypeHandle Archetype, int32 VisualIndex)
+void ALWFragmentTestFarmPlot::AddItemToGrid(UMassEntitySubsystem* EntitySystem, int32 X, int32 Y, FArchetypeHandle Archetype, int32 VisualIndex)
 {
 	FMassEntityHandle NewItem = EntitySystem->CreateEntity(Archetype);
 	PlantedSquares[X + Y * GridWidth] = NewItem;
 
-	EntitySystem->GetComponentDataChecked<FFarmWaterComponent>(NewItem).DeltaWaterPerSecond = FMath::FRandRange(-0.01f, -0.001f);
-	EntitySystem->GetComponentDataChecked<FHarvestTimerComponent>(NewItem).NumSecondsLeft = 5 + (FMath::Rand() % 100);
+	EntitySystem->GetFragmentDataChecked<FFarmWaterFragment>(NewItem).DeltaWaterPerSecond = FMath::FRandRange(-0.01f, -0.001f);
+	EntitySystem->GetFragmentDataChecked<FHarvestTimerFragment>(NewItem).NumSecondsLeft = 5 + (FMath::Rand() % 100);
 
 	FFarmGridCellData GridCoords;
 	GridCoords.CellX = X;
 	GridCoords.CellY = Y;
-	EntitySystem->GetComponentDataChecked<FFarmGridCellData>(NewItem) = GridCoords;
+	EntitySystem->GetFragmentDataChecked<FFarmGridCellData>(NewItem) = GridCoords;
 
 	const FVector MeshPosition(X*GridCellWidth, Y*GridCellHeight, 0.0f);
 	const FRotator MeshRotation(0.0f, FMath::FRand()*360.0f, 0.0f);
 	const FVector MeshScale(1.0f, 1.0f, 1.0f); //@TODO: plumb in scale param?
 	const FTransform MeshTransform(MeshRotation, MeshPosition, MeshScale);
 
-	FFarmVisualComponent& VisualComp = EntitySystem->GetComponentDataChecked<FFarmVisualComponent>(NewItem);
+	FFarmVisualFragment& VisualComp = EntitySystem->GetFragmentDataChecked<FFarmVisualFragment>(NewItem);
 	VisualComp.VisualType = VisualIndex;
 	VisualComp.InstanceIndex = VisualDataISMCs[VisualComp.VisualType]->AddInstance(MeshTransform);
 }
 
-void ALWComponentTestFarmPlot::BeginPlay()
+void ALWFragmentTestFarmPlot::BeginPlay()
 {
 	Super::BeginPlay();
 
 	UMassEntitySubsystem* EntitySystem = UWorld::GetSubsystem<UMassEntitySubsystem>(GetWorld());
 
-	FArchetypeHandle CropArchetype = EntitySystem->CreateArchetype(TArray<const UScriptStruct*>{ FFarmWaterComponent::StaticStruct(), FFarmCropComponent::StaticStruct(), FHarvestTimerComponent::StaticStruct(), FFarmVisualComponent::StaticStruct(), FFarmGridCellData::StaticStruct() });
-	FArchetypeHandle FlowerArchetype = EntitySystem->CreateArchetype(TArray<const UScriptStruct*>{ FFarmWaterComponent::StaticStruct(), FFarmFlowerComponent::StaticStruct(), FHarvestTimerComponent::StaticStruct(), FFarmVisualComponent::StaticStruct(), FFarmGridCellData::StaticStruct() });
+	FArchetypeHandle CropArchetype = EntitySystem->CreateArchetype(TArray<const UScriptStruct*>{ FFarmWaterFragment::StaticStruct(), FFarmCropFragment::StaticStruct(), FHarvestTimerFragment::StaticStruct(), FFarmVisualFragment::StaticStruct(), FFarmGridCellData::StaticStruct() });
+	FArchetypeHandle FlowerArchetype = EntitySystem->CreateArchetype(TArray<const UScriptStruct*>{ FFarmWaterFragment::StaticStruct(), FFarmFlowerFragment::StaticStruct(), FHarvestTimerFragment::StaticStruct(), FFarmVisualFragment::StaticStruct(), FFarmGridCellData::StaticStruct() });
 
 	PerFrameSystems.Add(NewObject<UFarmWaterUpdateSystem>(this));
 
@@ -144,7 +144,7 @@ void ALWComponentTestFarmPlot::BeginPlay()
 	}
 }
 
-void ALWComponentTestFarmPlot::TickActor(float DeltaTime, enum ELevelTick TickType, FActorTickFunction& ThisTickFunction)
+void ALWFragmentTestFarmPlot::TickActor(float DeltaTime, enum ELevelTick TickType, FActorTickFunction& ThisTickFunction)
 {
 	Super::TickActor(DeltaTime, TickType, ThisTickFunction);
 
@@ -158,7 +158,7 @@ void ALWComponentTestFarmPlot::TickActor(float DeltaTime, enum ELevelTick TickTy
 		TSharedPtr<FMassCommandBuffer> DeferredCommandBuffer = MakeShareable(new FMassCommandBuffer());
 		ExecContext.SetDeferredCommandBuffer(DeferredCommandBuffer);
 
-		for (UComponentUpdateSystem* System : PerFrameSystems)
+		for (UFragmentUpdateSystem* System : PerFrameSystems)
 		{
 			check(System);
 			System->Execute(*EntitySystem, ExecContext);
@@ -174,7 +174,7 @@ void ALWComponentTestFarmPlot::TickActor(float DeltaTime, enum ELevelTick TickTy
 		ExecContext.SetDeferredCommandBuffer(DeferredCommandBuffer);
 
 		NextSecondTimer += 1.0f;
-		for (UComponentUpdateSystem* System : PerSecondSystems)
+		for (UFragmentUpdateSystem* System : PerSecondSystems)
 		{
 			check(System);
 			System->Execute(*EntitySystem, ExecContext);
@@ -193,7 +193,7 @@ void ALWComponentTestFarmPlot::TickActor(float DeltaTime, enum ELevelTick TickTy
 				FMassEntityHandle GridEntity = PlantedSquares[X + Y * GridWidth];
 				if (EntitySystem->IsValidEntity(GridEntity))
 				{
-					FFarmVisualComponent& VisualComp = EntitySystem->GetComponentDataChecked<FFarmVisualComponent>(GridEntity);
+					FFarmVisualFragment& VisualComp = EntitySystem->GetFragmentDataChecked<FFarmVisualFragment>(GridEntity);
 					if (VisualComp.InstanceIndex < 0)
 					{
 						const FVector MeshPosition(X*GridCellWidth, Y*GridCellHeight, 0.0f);
@@ -208,7 +208,7 @@ void ALWComponentTestFarmPlot::TickActor(float DeltaTime, enum ELevelTick TickTy
 		}
 	}
 #endif
-	//ForEachComponentChunk
+	//ForEachFragmentChunk
 	// Update visuals
 
 #if 0
@@ -219,10 +219,10 @@ void ALWComponentTestFarmPlot::TickActor(float DeltaTime, enum ELevelTick TickTy
 	if (QQQZZZ == 1)
 	{
 		QUICK_SCOPE_CYCLE_COUNTER(FARM_HARVEST_ICONS);
-		EntitySystem->ForEachComponentChunk<FHarvestTimerComponent>(
-			[&](TArrayView<FHarvestTimerComponent> Timers)
+		EntitySystem->ForEachFragmentChunk<FHarvestTimerFragment>(
+			[&](TArrayView<FHarvestTimerFragment> Timers)
 		{
-			for (FHarvestTimerComponent& Timer : Timers)
+			for (FHarvestTimerFragment& Timer : Timers)
 			{
 				if ((Timer.NumSecondsLeft == 0) && (Timer.HarvestIconIndex < 0))
 				{
@@ -245,7 +245,7 @@ void ALWComponentTestFarmPlot::TickActor(float DeltaTime, enum ELevelTick TickTy
 				FMassEntityHandle GridEntity = PlantedSquares[X + Y * GridWidth];
 				if (EntitySystem->IsValidEntity(GridEntity))
 				{
-					FHarvestTimerComponent& HarvestTimerData = EntitySystem->GetComponentDataChecked<FHarvestTimerComponent>(GridEntity);
+					FHarvestTimerFragment& HarvestTimerData = EntitySystem->GetFragmentDataChecked<FHarvestTimerFragment>(GridEntity);
 					if ((HarvestTimerData.NumSecondsLeft == 0) && (HarvestTimerData.HarvestIconIndex < 0))
 					{
 						const FVector IconPosition(X*GridCellWidth, Y*GridCellHeight, HarvestIconHeight);
