Renamed all Component-containing functions and member variable names in Mass, replaced with "Fragment". Updated a lot of comments and logged text as well.

diff --git a/Engine/Plugins/Runtime/MassEntity/Source/MassEntityTestSuite/Private/MassEntityTest.cpp b/Engine/Plugins/Runtime/MassEntity/Source/MassEntityTestSuite/Private/MassEntityTest.cpp
--- a/Engine/Plugins/Runtime/MassEntity/Source/MassEntityTestSuite/Private/MassEntityTest.cpp
+++ b/Engine/Plugins/Runtime/MassEntity/Source/MassEntityTestSuite/Private/MassEntityTest.cpp
@@ -25,17 +25,17 @@ struct FEntityTest_ArchetypeCreation : FEntityTestBase
 		AITEST_TRUE("Ints archetype should have been created", IntsArchetype.IsValid());
 
 		TArray<const UScriptStruct*> FragmentsList;
-		EntitySubsystem->DebugGetArchetypeComponentTypes(FloatsArchetype, FragmentsList);
+		EntitySubsystem->DebugGetArchetypeFragmentTypes(FloatsArchetype, FragmentsList);
 		AITEST_EQUAL("Floats archetype should contain just a single fragment", FragmentsList.Num(), 1);
 		AITEST_EQUAL("Floats archetype\'s lone fragment should be of Float fragment type", FragmentsList[0], FTestFragment_Float::StaticStruct());
 
 		FragmentsList.Reset();
-		EntitySubsystem->DebugGetArchetypeComponentTypes(IntsArchetype, FragmentsList);
+		EntitySubsystem->DebugGetArchetypeFragmentTypes(IntsArchetype, FragmentsList);
 		AITEST_EQUAL("Ints archetype should contain just a single fragment", FragmentsList.Num(), 1);
 		AITEST_EQUAL("Ints archetype\'s lone fragment should be of Ints fragment type", FragmentsList[0], FTestFragment_Int::StaticStruct());
 
 		FragmentsList.Reset();
-		EntitySubsystem->DebugGetArchetypeComponentTypes(FloatsIntsArchetype, FragmentsList);
+		EntitySubsystem->DebugGetArchetypeFragmentTypes(FloatsIntsArchetype, FragmentsList);
 		AITEST_EQUAL("FloatsInts archetype should contain exactly two fragments", FragmentsList.Num(), 2);
 		AITEST_TRUE("FloatsInts archetype\'s should contain both expected fragment types", FragmentsList.Find(FTestFragment_Int::StaticStruct()) != INDEX_NONE && FragmentsList.Find(FTestFragment_Float::StaticStruct()) != INDEX_NONE);
 		
@@ -148,21 +148,21 @@ struct FEntityTest_EntityCreationFromInstances : FEntityTestBase
 		AITEST_EQUAL("There should be one entity across the whole system", EntitySubsystem->DebugGetEntityCount(), 1);
 		AITEST_EQUAL("Entity\'s archetype should be the Ints one", EntitySubsystem->GetArchetypeForEntity(Entity), IntsArchetype);
 		AITEST_EQUAL("The created entity should have been added to the Ints archetype", EntitySubsystem->DebugGetArchetypeEntitiesCount(IntsArchetype), 1);
-		AITEST_EQUAL("The entity should have the new component with the correct value set", EntitySubsystem->GetComponentDataChecked<FTestFragment_Int>(Entity).Value, TestIntValue);
+		AITEST_EQUAL("The entity should have the new component with the correct value set", EntitySubsystem->GetFragmentDataChecked<FTestFragment_Int>(Entity).Value, TestIntValue);
 
 		return true;
 	}
 };
 IMPLEMENT_AI_INSTANT_TEST(FEntityTest_EntityCreationFromInstances, "System.Mass.Entity.EntityCreationFromInstances");
 
-#if 0 // this test is compiled out since AddComponentToEntity will fails an ensure if a redundant fragment gets added
+#if 0 // this test is compiled out since AddFragmentToEntity will fails an ensure if a redundant fragment gets added
 struct FEntityTest_AddingRedundantFragment : FEntityTestBase
 {
 	virtual bool InstantTest() override
 	{
 		CA_ASSUME(EntitySubsystem);
 		const FLWEntity Entity = EntitySubsystem->CreateEntity(FloatsArchetype);
-		EntitySubsystem->AddComponentToEntity(Entity, FTestFragment_Float::StaticStruct());		
+		EntitySubsystem->AddFragmentToEntity(Entity, FTestFragment_Float::StaticStruct());		
 		AITEST_EQUAL("Adding a fragment that a given entity\s archetype already has should do nothing", EntitySubsystem->GetArchetypeForEntity(Entity), FloatsArchetype);
 
 		return true;
@@ -178,7 +178,7 @@ struct FEntityTest_AddingFragmentType : FEntityTestBase
 	{
 		CA_ASSUME(EntitySubsystem);
 		const FMassEntityHandle Entity = EntitySubsystem->CreateEntity(FloatsArchetype);
-		EntitySubsystem->AddComponentToEntity(Entity, FTestFragment_Int::StaticStruct());
+		EntitySubsystem->AddFragmentToEntity(Entity, FTestFragment_Int::StaticStruct());
 		AITEST_EQUAL("There should be one entity across the whole system", EntitySubsystem->DebugGetEntityCount(), 1);
 		AITEST_EQUAL("The original archetype should now have no entities", EntitySubsystem->DebugGetArchetypeEntitiesCount(FloatsArchetype), 0);
 		AITEST_EQUAL("The destination archetype should now store a single entity", EntitySubsystem->DebugGetArchetypeEntitiesCount(FloatsIntsArchetype), 1);
@@ -197,14 +197,14 @@ struct FEntityTest_AddingFragmentInstance : FEntityTestBase
 	{
 		CA_ASSUME(EntitySubsystem);
 		const FMassEntityHandle Entity = EntitySubsystem->CreateEntity(FloatsArchetype);
-		EntitySubsystem->AddComponentInstanceListToEntity(Entity, MakeArrayView(&InstanceInt, 1));
+		EntitySubsystem->AddFragmentInstanceListToEntity(Entity, MakeArrayView(&InstanceInt, 1));
 		AITEST_EQUAL("There should be one entity across the whole system", EntitySubsystem->DebugGetEntityCount(), 1);
 		AITEST_EQUAL("The original archetype should now have no entities", EntitySubsystem->DebugGetArchetypeEntitiesCount(FloatsArchetype), 0);
 		AITEST_EQUAL("The destination archetype should now store a single entity", EntitySubsystem->DebugGetArchetypeEntitiesCount(FloatsIntsArchetype), 1);
 		AITEST_EQUAL("The archetype containing just the new fragment should not be affected", EntitySubsystem->DebugGetArchetypeEntitiesCount(IntsArchetype), 0);
 		// this test was originally failing due to FEntityData.CurrentArchetype not getting updated during entity moving between archetypes
 		AITEST_EQUAL("The entity should get associated with the new archetype", EntitySubsystem->GetArchetypeForEntity(Entity), FloatsIntsArchetype);
-		AITEST_EQUAL("The entity should have the new component with the correct value set", EntitySubsystem->GetComponentDataChecked<FTestFragment_Int>(Entity).Value, TestIntValue);
+		AITEST_EQUAL("The entity should have the new component with the correct value set", EntitySubsystem->GetFragmentDataChecked<FTestFragment_Int>(Entity).Value, TestIntValue);
 
 		return true;
 	}
@@ -218,7 +218,7 @@ struct FEntityTest_RemovingFragment : FEntityTestBase
 	{
 		CA_ASSUME(EntitySubsystem);
 		const FMassEntityHandle Entity = EntitySubsystem->CreateEntity(FloatsIntsArchetype);
-		EntitySubsystem->RemoveComponentFromEntity(Entity, FTestFragment_Float::StaticStruct());
+		EntitySubsystem->RemoveFragmentFromEntity(Entity, FTestFragment_Float::StaticStruct());
 		AITEST_EQUAL("There should be just one entity across the whole system", EntitySubsystem->DebugGetEntityCount(), 1);
 		AITEST_EQUAL("The original archetype should now have no entities", EntitySubsystem->DebugGetArchetypeEntitiesCount(FloatsIntsArchetype), 0);
 		AITEST_EQUAL("The destination archetype should now store a single entity", EntitySubsystem->DebugGetArchetypeEntitiesCount(IntsArchetype), 1);
@@ -237,7 +237,7 @@ struct FEntityTest_RemovingLastFragment : FEntityTestBase
 	{
 		CA_ASSUME(EntitySubsystem);
 		const FMassEntityHandle Entity = EntitySubsystem->CreateEntity(FloatsArchetype);
-		EntitySubsystem->RemoveComponentFromEntity(Entity, FTestFragment_Float::StaticStruct());
+		EntitySubsystem->RemoveFragmentFromEntity(Entity, FTestFragment_Float::StaticStruct());
 		AITEST_EQUAL("There should be one entity across the whole system", EntitySubsystem->DebugGetEntityCount(), 1);
 		AITEST_EQUAL("The original archetype should now have no entities", EntitySubsystem->DebugGetArchetypeEntitiesCount(FloatsArchetype), 0);
 		// this test was originally failing due to FEntityData.CurrentArchetype not getting updated during entity moving between archetypes
@@ -299,7 +299,7 @@ struct FEntityTest_EntityReservationAndBuildingFromInstances : FEntityTestBase
 		AITEST_EQUAL("There should be one entity across the whole system", EntitySubsystem->DebugGetEntityCount(), 1);
 		AITEST_EQUAL("Entity\'s archetype should be the Ints one", EntitySubsystem->GetArchetypeForEntity(ReservedEntity), IntsArchetype);
 		AITEST_EQUAL("The created entity should have been added to the Ints archetype", EntitySubsystem->DebugGetArchetypeEntitiesCount(IntsArchetype), 1);
-		AITEST_EQUAL("The entity should have the new component with the correct value set", EntitySubsystem->GetComponentDataChecked<FTestFragment_Int>(ReservedEntity).Value, TestIntValue);
+		AITEST_EQUAL("The entity should have the new component with the correct value set", EntitySubsystem->GetFragmentDataChecked<FTestFragment_Int>(ReservedEntity).Value, TestIntValue);
 		EntitySubsystem->DestroyEntity(ReservedEntity);
 		AITEST_EQUAL("There should not be any entity across the whole system", EntitySubsystem->DebugGetEntityCount(), 0);
 		AITEST_EQUAL("The original archetype should now have no entities", EntitySubsystem->DebugGetArchetypeEntitiesCount(FloatsArchetype), 0);
diff --git a/Engine/Plugins/Runtime/MassEntity/Source/MassEntityTestSuite/Private/MassQueryTest.cpp b/Engine/Plugins/Runtime/MassEntity/Source/MassEntityTestSuite/Private/MassQueryTest.cpp
--- a/Engine/Plugins/Runtime/MassEntity/Source/MassEntityTestSuite/Private/MassQueryTest.cpp
+++ b/Engine/Plugins/Runtime/MassEntity/Source/MassEntityTestSuite/Private/MassQueryTest.cpp
@@ -72,8 +72,8 @@ struct FQueryTest_FragmentViewBinding : FEntityTestBase
 		CA_ASSUME(EntitySubsystem);
 
 		FMassEntityHandle Entity = EntitySubsystem->CreateEntity(FloatsArchetype);
-		FTestFragment_Float& TestedFragment = EntitySubsystem->GetComponentDataChecked<FTestFragment_Float>(Entity);
-		AITEST_TRUE("Initial value of the component should match expectations", TestedFragment.Value == 0.f);
+		FTestFragment_Float& TestedFragment = EntitySubsystem->GetFragmentDataChecked<FTestFragment_Float>(Entity);
+		AITEST_TRUE("Initial value of the fragment should match expectations", TestedFragment.Value == 0.f);
 
 		UPipeTestProcessor_Floats* Processor = NewObject<UPipeTestProcessor_Floats>(EntitySubsystem);
 		Processor->ExecutionFunction = [Processor](UMassEntitySubsystem& InEntitySubsystem, FMassExecutionContext& Context) {
@@ -82,7 +82,7 @@ struct FQueryTest_FragmentViewBinding : FEntityTestBase
 
 			Processor->TestGetQuery().ForEachEntityChunk(InEntitySubsystem, Context, [](FMassExecutionContext& Context)
 				{
-					TArrayView<FTestFragment_Float> Floats = Context.GetMutableComponentView<FTestFragment_Float>();
+					TArrayView<FTestFragment_Float> Floats = Context.GetMutableFragmentView<FTestFragment_Float>();
 
 					for (int32 i = 0; i < Context.GetNumEntities(); ++i)
 					{
@@ -118,7 +118,7 @@ struct FQueryTest_ExecuteSingleArchetype : FEntityTestBase
 		Query.ForEachEntityChunk(*EntitySubsystem, ExecContext, [&TotalProcessed](FMassExecutionContext& Context)
 			{
 				TotalProcessed += Context.GetNumEntities();
-				TArrayView<FTestFragment_Float> Floats = Context.GetMutableComponentView<FTestFragment_Float>();
+				TArrayView<FTestFragment_Float> Floats = Context.GetMutableFragmentView<FTestFragment_Float>();
 				
 				for (int32 i = 0; i < Context.GetNumEntities(); ++i)
 				{
@@ -130,7 +130,7 @@ struct FQueryTest_ExecuteSingleArchetype : FEntityTestBase
 
 		for (FMassEntityHandle& Entity : EntitiesCreated)
 		{
-			const FTestFragment_Float& TestedFragment = EntitySubsystem->GetComponentDataChecked<FTestFragment_Float>(Entity);
+			const FTestFragment_Float& TestedFragment = EntitySubsystem->GetFragmentDataChecked<FTestFragment_Float>(Entity);
 			AITEST_EQUAL("Every fragment value should have changed to the expected value", TestedFragment.Value, 13.f);
 		}
 
@@ -162,7 +162,7 @@ struct FQueryTest_ExecuteMultipleArchetypes : FEntityTestBase
 		Query.ForEachEntityChunk(*EntitySubsystem, ExecContext, [&TotalProcessed](FMassExecutionContext& Context)
 			{
 				TotalProcessed += Context.GetNumEntities();
-				TArrayView<FTestFragment_Float> Floats = Context.GetMutableComponentView<FTestFragment_Float>();
+				TArrayView<FTestFragment_Float> Floats = Context.GetMutableFragmentView<FTestFragment_Float>();
 
 				for (int32 i = 0; i < Context.GetNumEntities(); ++i)
 				{
@@ -174,7 +174,7 @@ struct FQueryTest_ExecuteMultipleArchetypes : FEntityTestBase
 
 		for (FMassEntityHandle& Entity : EntitiesCreated)
 		{
-			const FTestFragment_Float& TestedFragment = EntitySubsystem->GetComponentDataChecked<FTestFragment_Float>(Entity);
+			const FTestFragment_Float& TestedFragment = EntitySubsystem->GetFragmentDataChecked<FTestFragment_Float>(Entity);
 			AITEST_EQUAL("Every fragment value should have changed to the expected value", TestedFragment.Value, 13.f);
 		}
 
@@ -218,7 +218,7 @@ struct FQueryTest_ExecuteSparse : FEntityTestBase
 			.ForEachEntityChunk(FArchetypeChunkCollection(FloatsArchetype, EntitiesToProcess), *EntitySubsystem, ExecContext, [&TotalProcessed](FMassExecutionContext& Context)
 			{
 				TotalProcessed += Context.GetNumEntities();
-				TArrayView<FTestFragment_Float> Floats = Context.GetMutableComponentView<FTestFragment_Float>();
+				TArrayView<FTestFragment_Float> Floats = Context.GetMutableFragmentView<FTestFragment_Float>();
 
 				for (int32 i = 0; i < Context.GetNumEntities(); ++i)
 				{
@@ -230,13 +230,13 @@ struct FQueryTest_ExecuteSparse : FEntityTestBase
 
 		for (FMassEntityHandle& Entity : EntitiesToProcess)
 		{
-			const FTestFragment_Float& TestedFragment = EntitySubsystem->GetComponentDataChecked<FTestFragment_Float>(Entity);
+			const FTestFragment_Float& TestedFragment = EntitySubsystem->GetFragmentDataChecked<FTestFragment_Float>(Entity);
 			AITEST_EQUAL("Every fragment value should have changed to the expected value", TestedFragment.Value, 13.f);
 		}
 		
 		for (FMassEntityHandle& Entity : EntitiesToIgnore)
 		{
-			const FTestFragment_Float& TestedFragment = EntitySubsystem->GetComponentDataChecked<FTestFragment_Float>(Entity);
+			const FTestFragment_Float& TestedFragment = EntitySubsystem->GetFragmentDataChecked<FTestFragment_Float>(Entity);
 			AITEST_EQUAL("Untouched entites should retain default fragment value ", TestedFragment.Value, 0.f);
 		}
 
@@ -442,7 +442,7 @@ struct FQueryTest_UsingOptionalFragment : FEntityTestBase
 		FMassExecutionContext ExecContext;
 		Query.ForEachEntityChunk(*EntitySubsystem, ExecContext, [&TotalProcessed, &EmptyIntsViewCount, IntValueSet](FMassExecutionContext& Context) {
 			++TotalProcessed;
-			TArrayView<FTestFragment_Int> Ints = Context.GetMutableComponentView<FTestFragment_Int>();
+			TArrayView<FTestFragment_Int> Ints = Context.GetMutableFragmentView<FTestFragment_Int>();
 			if (Ints.Num() == 0)
 			{
 				++EmptyIntsViewCount;
@@ -459,7 +459,7 @@ struct FQueryTest_UsingOptionalFragment : FEntityTestBase
 		AITEST_EQUAL("Two archetypes total should get processed", TotalProcessed, 2);
 		AITEST_EQUAL("Only one of these archetypes should get an empty Ints array view", EmptyIntsViewCount, 1);
 
-		const FTestFragment_Int& TestFragment = EntitySubsystem->GetComponentDataChecked<FTestFragment_Int>(EntityWithFloatsInts);
+		const FTestFragment_Int& TestFragment = EntitySubsystem->GetFragmentDataChecked<FTestFragment_Int>(EntityWithFloatsInts);
 		AITEST_TRUE("The optional fragment\'s value should get modified where present", TestFragment.Value == IntValueSet);
 
 		return true;
@@ -499,8 +499,8 @@ struct FQueryTest_AnyFragment : FEntityTestBase
 		FMassExecutionContext TestContext;
 		Query.ForEachEntityChunk(*EntitySubsystem, TestContext, [this](FMassExecutionContext& Context)
 			{
-				TArrayView<FTestFragment_Bool> BoolView = Context.GetMutableComponentView<FTestFragment_Bool>();
-				TArrayView<FTestFragment_Int> IntView = Context.GetMutableComponentView<FTestFragment_Int>();
+				TArrayView<FTestFragment_Bool> BoolView = Context.GetMutableFragmentView<FTestFragment_Bool>();
+				TArrayView<FTestFragment_Int> IntView = Context.GetMutableFragmentView<FTestFragment_Int>();
 				
 				GetTestRunner().TestTrue("Every matching archetype needs to host Bool or Int fragments", BoolView.Num() || IntView.Num());
 			});
diff --git a/Engine/Plugins/Runtime/MassEntity/Source/MassEntityTestSuite/Public/LWComponentTestFarmPlot.h b/Engine/Plugins/Runtime/MassEntity/Source/MassEntityTestSuite/Public/LWComponentTestFarmPlot.h
--- a/Engine/Plugins/Runtime/MassEntity/Source/MassEntityTestSuite/Public/LWComponentTestFarmPlot.h
+++ b/Engine/Plugins/Runtime/MassEntity/Source/MassEntityTestSuite/Public/LWComponentTestFarmPlot.h
@@ -50,7 +50,7 @@ struct FFarmGridCellData : public FMassFragment
 };
 
 USTRUCT()
-struct FFarmWaterComponent : public FMassFragment
+struct FFarmWaterFragment : public FMassFragment
 {
 	GENERATED_BODY()
 
@@ -59,7 +59,7 @@ struct FFarmWaterComponent : public FMassFragment
 };
 
 USTRUCT()
-struct FFarmFlowerComponent : public FMassFragment
+struct FFarmFlowerFragment : public FMassFragment
 {
 	GENERATED_BODY()
 		
@@ -68,7 +68,7 @@ struct FFarmFlowerComponent : public FMassFragment
 };
 
 USTRUCT()
-struct FFarmCropComponent : public FMassFragment
+struct FFarmCropFragment : public FMassFragment
 {
 	GENERATED_BODY()
 
@@ -77,7 +77,7 @@ struct FFarmCropComponent : public FMassFragment
 
 
 USTRUCT()
-struct FFarmVisualComponent : public FMassFragment
+struct FFarmVisualFragment : public FMassFragment
 {
 	GENERATED_BODY()
 
@@ -87,7 +87,7 @@ struct FFarmVisualComponent : public FMassFragment
 };
 
 USTRUCT()
-struct FHarvestTimerComponent : public FMassFragment
+struct FHarvestTimerFragment : public FMassFragment
 {
 	GENERATED_BODY()
 
@@ -96,7 +96,7 @@ struct FHarvestTimerComponent : public FMassFragment
 
 
 UCLASS()
-class UComponentUpdateSystem : public UObject
+class UFragmentUpdateSystem : public UObject
 {
 	GENERATED_BODY()
 public:
@@ -110,14 +110,14 @@ protected:
 
 
 UCLASS()
-class UFarmWaterUpdateSystem : public UComponentUpdateSystem
+class UFarmWaterUpdateSystem : public UFragmentUpdateSystem
 {
 	GENERATED_BODY()
 
 public:
 	virtual void ConfigureQueries() override
 	{
-		EntityQuery.AddRequirement<FFarmWaterComponent>(EMassFragmentAccess::ReadWrite);
+		EntityQuery.AddRequirement<FFarmWaterFragment>(EMassFragmentAccess::ReadWrite);
 	}
 
 	virtual void Execute(UMassEntitySubsystem& EntitySubsystem, FMassExecutionContext& Context) override
@@ -125,11 +125,11 @@ public:
 		QUICK_SCOPE_CYCLE_COUNTER(UFarmWaterUpdateSystem_Run);
 		EntityQuery.ForEachEntityChunk(EntitySubsystem, Context, [this](FMassExecutionContext& Context) {
 
-			TArrayView<FFarmWaterComponent> WaterList = Context.GetMutableComponentView<FFarmWaterComponent>();
+			TArrayView<FFarmWaterFragment> WaterList = Context.GetMutableFragmentView<FFarmWaterFragment>();
 
-			for (FFarmWaterComponent& WaterComponent : WaterList)
+			for (FFarmWaterFragment& WaterFragment : WaterList)
 			{
-				WaterComponent.CurrentWater = FMath::Clamp(WaterComponent.CurrentWater + WaterComponent.DeltaWaterPerSecond * 0.01f, 0.0f, 1.0f);
+				WaterFragment.CurrentWater = FMath::Clamp(WaterFragment.CurrentWater + WaterFragment.DeltaWaterPerSecond * 0.01f, 0.0f, 1.0f);
 			}
 		});
 	}
@@ -137,16 +137,16 @@ public:
 
 
 UCLASS()
-class UFarmHarvestTimerSystem_Flowers : public UComponentUpdateSystem
+class UFarmHarvestTimerSystem_Flowers : public UFragmentUpdateSystem
 {
 	GENERATED_BODY()
 
 public:
 	virtual void ConfigureQueries() override
 	{
-		EntityQuery.AddRequirement<FHarvestTimerComponent>(EMassFragmentAccess::ReadWrite);
-		EntityQuery.AddRequirement<FFarmWaterComponent>(EMassFragmentAccess::ReadOnly);
-		EntityQuery.AddRequirement<FFarmFlowerComponent>(EMassFragmentAccess::ReadWrite);
+		EntityQuery.AddRequirement<FHarvestTimerFragment>(EMassFragmentAccess::ReadWrite);
+		EntityQuery.AddRequirement<FFarmWaterFragment>(EMassFragmentAccess::ReadOnly);
+		EntityQuery.AddRequirement<FFarmFlowerFragment>(EMassFragmentAccess::ReadWrite);
 	}
 
 	virtual void Execute(UMassEntitySubsystem& EntitySubsystem, FMassExecutionContext& Context) override
@@ -156,9 +156,9 @@ public:
 		EntityQuery.ForEachEntityChunk(EntitySubsystem, Context, [this](FMassExecutionContext& Context) {
 			const int32 NumEntities = Context.GetNumEntities();
 			const float WellWateredThreshold = 0.25f;
-			TArrayView<FHarvestTimerComponent> TimerList = Context.GetMutableComponentView<FHarvestTimerComponent>();
-			TConstArrayView<FFarmWaterComponent> WaterList = Context.GetComponentView<FFarmWaterComponent>();
-			TArrayView<FFarmFlowerComponent> FlowerList = Context.GetMutableComponentView<FFarmFlowerComponent>();
+			TArrayView<FHarvestTimerFragment> TimerList = Context.GetMutableFragmentView<FHarvestTimerFragment>();
+			TConstArrayView<FFarmWaterFragment> WaterList = Context.GetFragmentView<FFarmWaterFragment>();
+			TArrayView<FFarmFlowerFragment> FlowerList = Context.GetMutableFragmentView<FFarmFlowerFragment>();
 			
 			for (int32 i = 0; i < NumEntities; ++i)
 			{
@@ -180,16 +180,16 @@ public:
 
 
 UCLASS()
-class UFarmHarvestTimerSystem_Crops : public UComponentUpdateSystem
+class UFarmHarvestTimerSystem_Crops : public UFragmentUpdateSystem
 {
 	GENERATED_BODY()
 
 public:
 	virtual void ConfigureQueries() override
 	{
-		EntityQuery.AddRequirement<FHarvestTimerComponent>(EMassFragmentAccess::ReadWrite);
-		EntityQuery.AddRequirement<FFarmWaterComponent>(EMassFragmentAccess::ReadOnly);
-		EntityQuery.AddRequirement<FFarmCropComponent>(EMassFragmentAccess::ReadOnly);
+		EntityQuery.AddRequirement<FHarvestTimerFragment>(EMassFragmentAccess::ReadWrite);
+		EntityQuery.AddRequirement<FFarmWaterFragment>(EMassFragmentAccess::ReadOnly);
+		EntityQuery.AddRequirement<FFarmCropFragment>(EMassFragmentAccess::ReadOnly);
 	}
 
 	virtual void Execute(UMassEntitySubsystem& EntitySubsystem, FMassExecutionContext& Context) override
@@ -200,8 +200,8 @@ public:
 			
 			const int32 NumEntities = Context.GetNumEntities();
 			const float WellWateredThreshold = 0.25f;
-			TArrayView<FHarvestTimerComponent> TimerList = Context.GetMutableComponentView<FHarvestTimerComponent>();
-			TConstArrayView<FFarmWaterComponent> WaterList = Context.GetMutableComponentView<FFarmWaterComponent>();
+			TArrayView<FHarvestTimerFragment> TimerList = Context.GetMutableFragmentView<FHarvestTimerFragment>();
+			TConstArrayView<FFarmWaterFragment> WaterList = Context.GetMutableFragmentView<FFarmWaterFragment>();
 				
 			for (int32 i = 0; i < NumEntities; ++i)
 			{
@@ -215,14 +215,14 @@ public:
 //////////////////////////////////////////////////////////////////////
 
 UCLASS()
-class UFarmHarvestTimerExpired : public UComponentUpdateSystem
+class UFarmHarvestTimerExpired : public UFragmentUpdateSystem
 {
 	GENERATED_BODY()
 
 public:
 	virtual void ConfigureQueries() override
 	{
-		EntityQuery.AddRequirement<FHarvestTimerComponent>(EMassFragmentAccess::ReadOnly);
+		EntityQuery.AddRequirement<FHarvestTimerFragment>(EMassFragmentAccess::ReadOnly);
 		EntityQuery.AddTagRequirement<FFarmJustBecameReadyToHarvest>(EMassFragmentPresence::None);
 		EntityQuery.AddTagRequirement<FFarmReadyToHarvest>(EMassFragmentPresence::None);
 	}
@@ -233,7 +233,7 @@ public:
 		
 		EntityQuery.ForEachEntityChunk(EntitySubsystem, Context, [this](FMassExecutionContext& Context) {
 			const int32 NumEntities = Context.GetNumEntities();
-			TConstArrayView<FHarvestTimerComponent> TimerList = Context.GetComponentView<FHarvestTimerComponent>();
+			TConstArrayView<FHarvestTimerFragment> TimerList = Context.GetFragmentView<FHarvestTimerFragment>();
 
 			for (int32 i = 0; i < NumEntities; ++i)
 			{
@@ -250,7 +250,7 @@ public:
 //////////////////////////////////////////////////////////////////////
 
 UCLASS()
-class UFarmHarvestTimerSetIcon : public UComponentUpdateSystem
+class UFarmHarvestTimerSetIcon : public UFragmentUpdateSystem
 {
 	GENERATED_BODY()
 
@@ -265,7 +265,7 @@ public:
 	virtual void ConfigureQueries() override
 	{
 		EntityQuery.AddRequirement<FFarmGridCellData>(EMassFragmentAccess::ReadOnly);
-		EntityQuery.AddRequirement<FFarmVisualComponent>(EMassFragmentAccess::ReadWrite);
+		EntityQuery.AddRequirement<FFarmVisualFragment>(EMassFragmentAccess::ReadWrite);
 		EntityQuery.AddTagRequirement<FFarmJustBecameReadyToHarvest>(EMassFragmentPresence::All);
 	}
 
@@ -276,7 +276,7 @@ public:
 
 
 UCLASS(config=Game)
-class ALWComponentTestFarmPlot : public AActor
+class ALWFragmentTestFarmPlot : public AActor
 {
 	GENERATED_BODY()
 
@@ -305,10 +305,10 @@ public:
 	float NextSecondTimer = 0.0f;
 
 	UPROPERTY(Transient)
-	TArray<UComponentUpdateSystem*> PerFrameSystems;
+	TArray<UFragmentUpdateSystem*> PerFrameSystems;
 
 	UPROPERTY(Transient)
-	TArray<UComponentUpdateSystem*> PerSecondSystems;
+	TArray<UFragmentUpdateSystem*> PerSecondSystems;
 
 	// Indicies into VisualDataTable for flowers
 	UPROPERTY(EditAnywhere, Category=Farm)
@@ -340,7 +340,7 @@ private:
 	void AddItemToGrid(UMassEntitySubsystem* EntitySystem, int32 X, int32 Y, FArchetypeHandle Archetype, int32 VisualIndex);
 
 public:
-	ALWComponentTestFarmPlot();
+	ALWFragmentTestFarmPlot();
 
 	virtual void BeginPlay() override;
 	virtual void TickActor(float DeltaTime, enum ELevelTick TickType, FActorTickFunction& ThisTickFunction) override;
diff --git a/Engine/Plugins/Runtime/MassGameplay/Source/MassActors/Private/MassActorSubsystem.cpp b/Engine/Plugins/Runtime/MassGameplay/Source/MassActors/Private/MassActorSubsystem.cpp
--- a/Engine/Plugins/Runtime/MassGameplay/Source/MassActors/Private/MassActorSubsystem.cpp
+++ b/Engine/Plugins/Runtime/MassGameplay/Source/MassActors/Private/MassActorSubsystem.cpp
@@ -81,7 +81,7 @@ FMassEntityHandle UMassActorSubsystem::GetEntityHandleFromActor(const TObjectKey
 AActor* UMassActorSubsystem::GetActorFromHandle(const FMassEntityHandle Handle) const
 {
 	check(EntitySystem);
-	FDataFragment_Actor* Data = EntitySystem->GetComponentDataPtr<FDataFragment_Actor>(Handle);
+	FDataFragment_Actor* Data = EntitySystem->GetFragmentDataPtr<FDataFragment_Actor>(Handle);
 	return Data != nullptr ? Data->GetMutable() : nullptr;
 }
 
@@ -119,7 +119,7 @@ void UMassActorSubsystem::DisconnectActor(const TObjectKey<const AActor> Actor,
 	if (FoundEntity == Handle)
 	{
 		check(EntitySystem);
-		if (FDataFragment_Actor* Data = EntitySystem->GetComponentDataPtr<FDataFragment_Actor>(Handle))
+		if (FDataFragment_Actor* Data = EntitySystem->GetFragmentDataPtr<FDataFragment_Actor>(Handle))
 		{
 			Data->ResetAndUpdateHandleMap();
 		}
diff --git a/Engine/Plugins/Runtime/MassGameplay/Source/MassActors/Private/MassAgentComponent.cpp b/Engine/Plugins/Runtime/MassGameplay/Source/MassActors/Private/MassAgentComponent.cpp
--- a/Engine/Plugins/Runtime/MassGameplay/Source/MassActors/Private/MassAgentComponent.cpp
+++ b/Engine/Plugins/Runtime/MassGameplay/Source/MassActors/Private/MassAgentComponent.cpp
@@ -164,7 +164,7 @@ void UMassAgentComponent::SetEntityHandleInternal(const FMassEntityHandle NewHan
 	// Fetch NetID if it exist
 	if (const UMassEntitySubsystem* EntitySubsystem = UWorld::GetSubsystem<UMassEntitySubsystem>(GetWorld()))
 	{
-		if (const FMassNetworkIDFragment* NetIDFragment = EntitySubsystem->GetComponentDataPtr<FMassNetworkIDFragment>(AgentHandle))
+		if (const FMassNetworkIDFragment* NetIDFragment = EntitySubsystem->GetFragmentDataPtr<FMassNetworkIDFragment>(AgentHandle))
 		{
 			if (!IsNetSimulating())
 			{
@@ -193,20 +193,20 @@ void UMassAgentComponent::SetEntityHandleInternal(const FMassEntityHandle NewHan
 			// @todo Find a way to add these initialization into either translator initializer or adding new fragments
 			// Make sure to fetch the fragment after any release, as that action can move the entity around into new archetype and 
 			// by the same fact change the references to the fragments.
-			if (FDataFragment_Actor* ActorInfo = EntityView.GetComponentDataPtr<FDataFragment_Actor>())
+			if (FDataFragment_Actor* ActorInfo = EntityView.GetFragmentDataPtr<FDataFragment_Actor>())
 			{
 				checkf(!ActorInfo->IsValid(), TEXT("Expecting ActorInfo fragment to be null"));
 				ActorInfo->SetAndUpdateHandleMap(AgentHandle, GetOwner(), !IsNetSimulating()/*bIsOwnedByMass*/);
 			}
 
 			// Initialize location of the replicated actor to match the mass replicated one
-			if (const FDataFragment_Transform* TransformFragment = EntityView.GetComponentDataPtr<FDataFragment_Transform>())
+			if (const FDataFragment_Transform* TransformFragment = EntityView.GetFragmentDataPtr<FDataFragment_Transform>())
 			{
 				GetOwner()->SetActorTransform(TransformFragment->GetTransform(), /*bSweep*/false, /*OutSweepHitResult*/nullptr, ETeleportType::TeleportPhysics);
 			}
 
 			// Initialize velocity of the replicated actor to match the mass replicated one
-			if (const FMassVelocityFragment* Velocity = EntityView.GetComponentDataPtr<FMassVelocityFragment>())
+			if (const FMassVelocityFragment* Velocity = EntityView.GetFragmentDataPtr<FMassVelocityFragment>())
 			{
 				if (UCharacterMovementComponent* MovementComp = GetOwner()->FindComponentByClass<UCharacterMovementComponent>())
 				{
@@ -278,9 +278,9 @@ void UMassAgentComponent::ClearEntityHandleInternal()
 		if (IsNetSimulating())
 		{
 			const FMassEntityView EntityView(*EntitySubsystem, AgentHandle);
-			if (FDataFragment_Actor* ActorInfo = EntityView.GetComponentDataPtr<FDataFragment_Actor>())
+			if (FDataFragment_Actor* ActorInfo = EntityView.GetFragmentDataPtr<FDataFragment_Actor>())
 			{
-				checkf(!ActorInfo->IsValid() || ActorInfo->Get() == GetOwner(), TEXT("Expecting actor pointer to be the component owner"));
+				checkf(!ActorInfo->IsValid() || ActorInfo->Get() == GetOwner(), TEXT("Expecting actor pointer to be the Component\'s owner"));
 				ActorInfo->ResetAndUpdateHandleMap();
 			}
 		}
@@ -356,7 +356,7 @@ void UMassAgentComponent::DebugCheckStateConsistency()
 							if (bIsBuiltEntity)
 							{
 								AActor* Owner = GetOwner();
-								const AActor* Actor = EntitySubsystem->GetComponentDataChecked<FDataFragment_Actor>(AgentHandle).Get();
+								const AActor* Actor = EntitySubsystem->GetFragmentDataChecked<FDataFragment_Actor>(AgentHandle).Get();
 								MASSAGENT_CHECK(Actor == nullptr || Actor == Owner, TEXT("Mass Actor and Owner mismatched in state %s"), *UEnum::GetValueAsString(State));
 							}
 						}
diff --git a/Engine/Plugins/Runtime/MassGameplay/Source/MassActors/Private/MassAgentTraits.cpp b/Engine/Plugins/Runtime/MassGameplay/Source/MassActors/Private/MassAgentTraits.cpp
--- a/Engine/Plugins/Runtime/MassGameplay/Source/MassActors/Private/MassAgentTraits.cpp
+++ b/Engine/Plugins/Runtime/MassGameplay/Source/MassActors/Private/MassAgentTraits.cpp
@@ -53,15 +53,15 @@ void UMassAgentCapsuleCollisionSyncTrait::BuildTemplate(FMassEntityTemplateBuild
 		{
 			if (UCapsuleComponent* CapsuleComponent = FMassAgentTraitsHelper::AsComponent<UCapsuleComponent>(Owner))
 			{
-				FDataFragment_CapsuleComponentWrapper& CapsuleFragment = EntityView.GetComponentData<FDataFragment_CapsuleComponentWrapper>();
+				FDataFragment_CapsuleComponentWrapper& CapsuleFragment = EntityView.GetFragmentData<FDataFragment_CapsuleComponentWrapper>();
 				CapsuleFragment.Component = CapsuleComponent;
 
-				FDataFragment_AgentRadius& RadiusFragment = EntityView.GetComponentData<FDataFragment_AgentRadius>();
+				FDataFragment_AgentRadius& RadiusFragment = EntityView.GetFragmentData<FDataFragment_AgentRadius>();
 				RadiusFragment.Radius = CapsuleComponent->GetScaledCapsuleRadius();
 
 				if (bSyncTransform)
 				{
-					FDataFragment_Transform& TransformFragment = EntityView.GetComponentData<FDataFragment_Transform>();
+					FDataFragment_Transform& TransformFragment = EntityView.GetFragmentData<FDataFragment_Transform>();
 					TransformFragment.GetMutableTransform() = CapsuleComponent->GetComponentTransform();
 				}
 			}
@@ -93,10 +93,10 @@ void UMassAgentMovementSyncTrait::BuildTemplate(FMassEntityTemplateBuildContext&
 		{
 			if (UCharacterMovementComponent* MovementComp = FMassAgentTraitsHelper::AsComponent<UCharacterMovementComponent>(Owner))
 			{
-				FDataFragment_CharacterMovementComponentWrapper& ComponentFragment = EntityView.GetComponentData<FDataFragment_CharacterMovementComponentWrapper>();
+				FDataFragment_CharacterMovementComponentWrapper& ComponentFragment = EntityView.GetFragmentData<FDataFragment_CharacterMovementComponentWrapper>();
 				ComponentFragment.Component = MovementComp;
 
-				FMassVelocityFragment& VelocityFragment = EntityView.GetComponentData<FMassVelocityFragment>();
+				FMassVelocityFragment& VelocityFragment = EntityView.GetFragmentData<FMassVelocityFragment>();
 
 				// the entity is the authority
 				if (CurrentDirection ==  EMassTranslationDirection::MassToActor)
@@ -157,10 +157,10 @@ void UMassAgentFeetLocationSyncTrait::BuildTemplate(FMassEntityTemplateBuildCont
 			if (AsActor && AsActor->GetRootComponent())
 			{
 				USceneComponent* Component = AsActor->GetRootComponent();
-				FMassSceneComponentWrapperFragment& ComponentFragment = EntityView.GetComponentData<FMassSceneComponentWrapperFragment>();
+				FMassSceneComponentWrapperFragment& ComponentFragment = EntityView.GetFragmentData<FMassSceneComponentWrapperFragment>();
 				ComponentFragment.Component = Component;
 
-				FDataFragment_Transform& TransformFragment = EntityView.GetComponentData<FDataFragment_Transform>();
+				FDataFragment_Transform& TransformFragment = EntityView.GetFragmentData<FDataFragment_Transform>();
 
 				REDIRECT_OBJECT_TO_VLOG(Component, &Owner);
 				UE_VLOG_LOCATION(&Owner, LogMass, Log, Component->GetComponentLocation(), 30, FColor::Yellow, TEXT("Initial component location"));
diff --git a/Engine/Plugins/Runtime/MassGameplay/Source/MassActors/Private/Translators/MassCapsuleComponentTranslators.cpp b/Engine/Plugins/Runtime/MassGameplay/Source/MassActors/Private/Translators/MassCapsuleComponentTranslators.cpp
--- a/Engine/Plugins/Runtime/MassGameplay/Source/MassActors/Private/Translators/MassCapsuleComponentTranslators.cpp
+++ b/Engine/Plugins/Runtime/MassGameplay/Source/MassActors/Private/Translators/MassCapsuleComponentTranslators.cpp
@@ -28,8 +28,8 @@ void UMassCapsuleTransformToMassTranslator::Execute(UMassEntitySubsystem& Entity
 {
 	EntityQuery.ForEachEntityChunk(EntitySubsystem, Context, [this](FMassExecutionContext& Context)
 		{
-			const TConstArrayView<FDataFragment_CapsuleComponentWrapper> CapsuleComponentList = Context.GetComponentView<FDataFragment_CapsuleComponentWrapper>();
-			const TArrayView<FDataFragment_Transform> LocationList = Context.GetMutableComponentView<FDataFragment_Transform>();
+			const TConstArrayView<FDataFragment_CapsuleComponentWrapper> CapsuleComponentList = Context.GetFragmentView<FDataFragment_CapsuleComponentWrapper>();
+			const TArrayView<FDataFragment_Transform> LocationList = Context.GetMutableFragmentView<FDataFragment_Transform>();
 			for (int i = 0; i < CapsuleComponentList.Num(); ++i)
 			{
 				if (const UCapsuleComponent* CapsuleComp = CapsuleComponentList[i].Component.Get())
@@ -62,8 +62,8 @@ void UMassTransformToActorCapsuleTranslator::Execute(UMassEntitySubsystem& Entit
 {
 	EntityQuery.ForEachEntityChunk(EntitySubsystem, Context, [this](FMassExecutionContext& Context)
 		{
-			const TArrayView<FDataFragment_CapsuleComponentWrapper> CapsuleComponentList = Context.GetMutableComponentView<FDataFragment_CapsuleComponentWrapper>();
-			const TConstArrayView<FDataFragment_Transform> LocationList = Context.GetComponentView<FDataFragment_Transform>();
+			const TArrayView<FDataFragment_CapsuleComponentWrapper> CapsuleComponentList = Context.GetMutableFragmentView<FDataFragment_CapsuleComponentWrapper>();
+			const TConstArrayView<FDataFragment_Transform> LocationList = Context.GetFragmentView<FDataFragment_Transform>();
 			for (int i = 0; i < CapsuleComponentList.Num(); ++i)
 			{
 				if (UCapsuleComponent* CapsuleComp = CapsuleComponentList[i].Component.Get())
diff --git a/Engine/Plugins/Runtime/MassGameplay/Source/MassActors/Private/Translators/MassCharacterMovementTranslators.cpp b/Engine/Plugins/Runtime/MassGameplay/Source/MassActors/Private/Translators/MassCharacterMovementTranslators.cpp
--- a/Engine/Plugins/Runtime/MassGameplay/Source/MassActors/Private/Translators/MassCharacterMovementTranslators.cpp
+++ b/Engine/Plugins/Runtime/MassGameplay/Source/MassActors/Private/Translators/MassCharacterMovementTranslators.cpp
@@ -32,9 +32,9 @@ void UMassCharacterMovementToMassTranslator::Execute(UMassEntitySubsystem& Entit
 {
 	EntityQuery.ForEachEntityChunk(EntitySubsystem, Context, [this](FMassExecutionContext& Context)
 	{
-		const TConstArrayView<FDataFragment_CharacterMovementComponentWrapper> ComponentList = Context.GetComponentView<FDataFragment_CharacterMovementComponentWrapper>();
-		const TArrayView<FDataFragment_Transform> LocationList = Context.GetMutableComponentView<FDataFragment_Transform>();
-		const TArrayView<FMassVelocityFragment> VelocityList = Context.GetMutableComponentView<FMassVelocityFragment>();
+		const TConstArrayView<FDataFragment_CharacterMovementComponentWrapper> ComponentList = Context.GetFragmentView<FDataFragment_CharacterMovementComponentWrapper>();
+		const TArrayView<FDataFragment_Transform> LocationList = Context.GetMutableFragmentView<FDataFragment_Transform>();
+		const TArrayView<FMassVelocityFragment> VelocityList = Context.GetMutableFragmentView<FMassVelocityFragment>();
 
 		const int32 NumEntities = Context.GetNumEntities();
 		
@@ -72,8 +72,8 @@ void UMassCharacterMovementToActorTranslator::Execute(UMassEntitySubsystem& Enti
 {
 	EntityQuery.ForEachEntityChunk(EntitySubsystem, Context, [this](FMassExecutionContext& Context)
 	{
-		const TArrayView<FDataFragment_CharacterMovementComponentWrapper> ComponentList = Context.GetMutableComponentView<FDataFragment_CharacterMovementComponentWrapper>();
-		const TConstArrayView<FMassVelocityFragment> VelocityList = Context.GetComponentView<FMassVelocityFragment>();
+		const TArrayView<FDataFragment_CharacterMovementComponentWrapper> ComponentList = Context.GetMutableFragmentView<FDataFragment_CharacterMovementComponentWrapper>();
+		const TConstArrayView<FMassVelocityFragment> VelocityList = Context.GetFragmentView<FMassVelocityFragment>();
 		
 		const int32 NumEntities = Context.GetNumEntities();
 
@@ -108,8 +108,8 @@ void UMassCharacterOrientationToMassTranslator::Execute(UMassEntitySubsystem& En
 {
 	EntityQuery.ForEachEntityChunk(EntitySubsystem, Context, [this](FMassExecutionContext& Context)
 	{
-		const TConstArrayView<FDataFragment_CharacterMovementComponentWrapper> ComponentList = Context.GetComponentView<FDataFragment_CharacterMovementComponentWrapper>();
-		const TArrayView<FDataFragment_Transform> LocationList = Context.GetMutableComponentView<FDataFragment_Transform>();
+		const TConstArrayView<FDataFragment_CharacterMovementComponentWrapper> ComponentList = Context.GetFragmentView<FDataFragment_CharacterMovementComponentWrapper>();
+		const TArrayView<FDataFragment_Transform> LocationList = Context.GetMutableFragmentView<FDataFragment_Transform>();
 
 		const int32 NumEntities = Context.GetNumEntities();
 		
@@ -148,8 +148,8 @@ void UMassCharacterOrientationToActorTranslator::Execute(UMassEntitySubsystem& E
 {
 	EntityQuery.ForEachEntityChunk(EntitySubsystem, Context, [this](FMassExecutionContext& Context)
 	{
-		const TArrayView<FDataFragment_CharacterMovementComponentWrapper> ComponentList = Context.GetMutableComponentView<FDataFragment_CharacterMovementComponentWrapper>();
-		const TConstArrayView<FDataFragment_Transform> TransformList = Context.GetComponentView<FDataFragment_Transform>();
+		const TArrayView<FDataFragment_CharacterMovementComponentWrapper> ComponentList = Context.GetMutableFragmentView<FDataFragment_CharacterMovementComponentWrapper>();
+		const TConstArrayView<FDataFragment_Transform> TransformList = Context.GetFragmentView<FDataFragment_Transform>();
 		
 		const int32 NumEntities = Context.GetNumEntities();
 
@@ -185,7 +185,7 @@ void UMassFragmentInitializer_NavLocation::Execute(UMassEntitySubsystem& EntityS
 {
 	EntityQuery.ForEachEntityChunk(EntitySubsystem, Context, [this](FMassExecutionContext& Context)
 		{
-			const TArrayView<FDataFragment_NavLocation> NavLocationList = Context.GetMutableComponentView<FDataFragment_NavLocation>();
+			const TArrayView<FDataFragment_NavLocation> NavLocationList = Context.GetMutableFragmentView<FDataFragment_NavLocation>();
 			const int32 NumEntities = Context.GetNumEntities();
 			for (int32 i = 0; i < NumEntities; ++i)
 			{
@@ -256,7 +256,7 @@ void UMassFragmentInitializer_Transform::Execute(UMassEntitySubsystem& EntitySub
 
 		EntityQuery.ForEachEntityChunk(EntitySubsystem, Context, [&Transforms, this](FMassExecutionContext& Context)
 			{
-				const TArrayView<FDataFragment_Transform> LocationList = Context.GetMutableComponentView<FDataFragment_Transform>();
+				const TArrayView<FDataFragment_Transform> LocationList = Context.GetMutableFragmentView<FDataFragment_Transform>();
 				const int32 NumEntities = Context.GetNumEntities();
 				for (int32 i = 0; i < NumEntities; ++i)
 				{
diff --git a/Engine/Plugins/Runtime/MassGameplay/Source/MassActors/Private/Translators/MassSceneComponentLocationTranslator.cpp b/Engine/Plugins/Runtime/MassGameplay/Source/MassActors/Private/Translators/MassSceneComponentLocationTranslator.cpp
--- a/Engine/Plugins/Runtime/MassGameplay/Source/MassActors/Private/Translators/MassSceneComponentLocationTranslator.cpp
+++ b/Engine/Plugins/Runtime/MassGameplay/Source/MassActors/Private/Translators/MassSceneComponentLocationTranslator.cpp
@@ -26,8 +26,8 @@ void UMassSceneComponentLocationToMassTranslator::Execute(UMassEntitySubsystem&
 {
 	EntityQuery.ForEachEntityChunk(EntitySubsystem, Context, [this](FMassExecutionContext& Context)
 	{
-		const TConstArrayView<FMassSceneComponentWrapperFragment> ComponentList = Context.GetComponentView<FMassSceneComponentWrapperFragment>();
-		const TArrayView<FDataFragment_Transform> LocationList = Context.GetMutableComponentView<FDataFragment_Transform>();
+		const TConstArrayView<FMassSceneComponentWrapperFragment> ComponentList = Context.GetFragmentView<FMassSceneComponentWrapperFragment>();
+		const TArrayView<FDataFragment_Transform> LocationList = Context.GetMutableFragmentView<FDataFragment_Transform>();
 
 		const int32 NumEntities = Context.GetNumEntities();
 		for (int32 i = 0; i < NumEntities; ++i)
@@ -62,8 +62,8 @@ void UMassSceneComponentLocationToActorTranslator::Execute(UMassEntitySubsystem&
 {
 	EntityQuery.ForEachEntityChunk(EntitySubsystem, Context, [this](FMassExecutionContext& Context)
 		{
-			const TConstArrayView<FMassSceneComponentWrapperFragment> ComponentList = Context.GetComponentView<FMassSceneComponentWrapperFragment>();
-			const TArrayView<FDataFragment_Transform> LocationList = Context.GetMutableComponentView<FDataFragment_Transform>();
+			const TConstArrayView<FMassSceneComponentWrapperFragment> ComponentList = Context.GetFragmentView<FMassSceneComponentWrapperFragment>();
+			const TArrayView<FDataFragment_Transform> LocationList = Context.GetMutableFragmentView<FDataFragment_Transform>();
 
 			const int32 NumEntities = Context.GetNumEntities();
 			for (int32 i = 0; i < NumEntities; ++i)
diff --git a/Engine/Plugins/Runtime/MassGameplay/Source/MassActors/Public/MassActorSubsystem.h b/Engine/Plugins/Runtime/MassGameplay/Source/MassActors/Public/MassActorSubsystem.h
--- a/Engine/Plugins/Runtime/MassGameplay/Source/MassActors/Public/MassActorSubsystem.h
+++ b/Engine/Plugins/Runtime/MassGameplay/Source/MassActors/Public/MassActorSubsystem.h
@@ -15,7 +15,7 @@ class AActor;
 class UMassEntitySubsystem;
 
 /**
- * Light weight component data to save the actor pointer of a light weight entity if it exist
+ * Fragment to save the actor pointer of a mass entity if it exist
  */
 USTRUCT()
 struct MASSACTORS_API FDataFragment_Actor : public FDataFragment_ObjectWrapper
diff --git a/Engine/Plugins/Runtime/MassGameplay/Source/MassCommon/Private/MassGenericProcessors.cpp b/Engine/Plugins/Runtime/MassGameplay/Source/MassCommon/Private/MassGenericProcessors.cpp
--- a/Engine/Plugins/Runtime/MassGameplay/Source/MassCommon/Private/MassGenericProcessors.cpp
+++ b/Engine/Plugins/Runtime/MassGameplay/Source/MassCommon/Private/MassGenericProcessors.cpp
@@ -8,7 +8,7 @@
 //----------------------------------------------------------------------//
 URandomizeVectorProcessor::URandomizeVectorProcessor()
 {
-	ComponentType = FVectorComponent::StaticStruct();
+	ComponentType = FVectorFragment::StaticStruct();
 	bAutoRegisterWithProcessingPhases = false;
 }
 
@@ -30,7 +30,7 @@ void URandomizeVectorProcessor::Execute(UMassEntitySubsystem& EntitySubsystem, F
 void URandomizeVectorProcessor::PostEditChangeProperty(FPropertyChangedEvent& PropertyChangedEvent)
 {
 	static const FName ComponentTypeName = GET_MEMBER_NAME_CHECKED(URandomizeVectorProcessor, ComponentType);
-	static const int32 SupportedSize = FVectorComponent::StaticStruct()->GetStructureSize();
+	static const int32 SupportedSize = FVectorFragment::StaticStruct()->GetStructureSize();
 
 	Super::PostEditChangeProperty(PropertyChangedEvent);
 
@@ -43,7 +43,7 @@ void URandomizeVectorProcessor::PostEditChangeProperty(FPropertyChangedEvent& Pr
 				|| ComponentType->IsChildOf(FMassFragment::StaticStruct()) == false
 				|| ComponentType->GetStructureSize() != SupportedSize)
 			{
-				ComponentType = FVectorComponent::StaticStruct();
+				ComponentType = FVectorFragment::StaticStruct();
 			}
 		}
 	}
diff --git a/Engine/Plugins/Runtime/MassGameplay/Source/MassCommon/Public/MassCommonFragments.h b/Engine/Plugins/Runtime/MassGameplay/Source/MassCommon/Public/MassCommonFragments.h
--- a/Engine/Plugins/Runtime/MassGameplay/Source/MassCommon/Public/MassCommonFragments.h
+++ b/Engine/Plugins/Runtime/MassGameplay/Source/MassCommon/Public/MassCommonFragments.h
@@ -28,11 +28,11 @@ struct MASSCOMMON_API FDataFragment_NavLocation : public FMassFragment
 	NavNodeRef NodeRef = INVALID_NAVNODEREF;
 };
 
-/** this component is a kind of an experiment. I've made FMassVelocityFragment and FMassAvoidanceFragment extend it and 
+/** this fragment is a kind of an experiment. I've made FMassVelocityFragment and FMassAvoidanceFragment extend it and 
  *  use URandomizeVectorProcessor to initialize it to random variables at spawn time.
  */
 USTRUCT()
-struct MASSCOMMON_API FVectorComponent : public FMassFragment
+struct MASSCOMMON_API FVectorFragment : public FMassFragment
 {
 	GENERATED_BODY()
 	UPROPERTY(EditAnywhere, Category = "")
diff --git a/Engine/Plugins/Runtime/MassGameplay/Source/MassCommon/Public/MassGenericProcessors.h b/Engine/Plugins/Runtime/MassGameplay/Source/MassCommon/Public/MassGenericProcessors.h
--- a/Engine/Plugins/Runtime/MassGameplay/Source/MassCommon/Public/MassGenericProcessors.h
+++ b/Engine/Plugins/Runtime/MassGameplay/Source/MassCommon/Public/MassGenericProcessors.h
@@ -9,7 +9,7 @@
 
 /** 
  *  An experimental processor that can be configured to work on any Fragment type as long as it's size matches
- *  FVectorComponent (effectively the size of FVector). Meant mostly for prototyping.
+ *  FVectorFragment (effectively the size of FVector). Meant mostly for prototyping.
  *  @todo currently the processor is randomizing it as if it was a fixed-max-length 2d velocity vector. If we keep this
  *  idea around we can add more properties to the processor that would configure the behavior.
  */
diff --git a/Engine/Plugins/Runtime/MassGameplay/Source/MassGameplayDebug/Private/AssignDebugVisProcessor.cpp b/Engine/Plugins/Runtime/MassGameplay/Source/MassGameplayDebug/Private/AssignDebugVisProcessor.cpp
--- a/Engine/Plugins/Runtime/MassGameplay/Source/MassGameplayDebug/Private/AssignDebugVisProcessor.cpp
+++ b/Engine/Plugins/Runtime/MassGameplay/Source/MassGameplayDebug/Private/AssignDebugVisProcessor.cpp
@@ -36,7 +36,7 @@ void UAssignDebugVisProcessor::Execute(UMassEntitySubsystem& EntitySubsystem, FM
 
 	EntityQuery.ForEachEntityChunk(EntitySubsystem, Context, [this, Visualizer](FMassExecutionContext& Context)
 	{
-		const TArrayView<FSimDebugVisComponent> DebugVisList = Context.GetMutableComponentView<FSimDebugVisComponent>();
+		const TArrayView<FSimDebugVisComponent> DebugVisList = Context.GetMutableFragmentView<FSimDebugVisComponent>();
 		for (FSimDebugVisComponent& VisualComp : DebugVisList)
 		{
 			// VisualComp.VisualType needs to be assigned by now. Should be performed as part of spawning, copied from the AgentTemplate
diff --git a/Engine/Plugins/Runtime/MassGameplay/Source/MassGameplayDebug/Private/DebugVisLocationProcessor.cpp b/Engine/Plugins/Runtime/MassGameplay/Source/MassGameplayDebug/Private/DebugVisLocationProcessor.cpp
--- a/Engine/Plugins/Runtime/MassGameplay/Source/MassGameplayDebug/Private/DebugVisLocationProcessor.cpp
+++ b/Engine/Plugins/Runtime/MassGameplay/Source/MassGameplayDebug/Private/DebugVisLocationProcessor.cpp
@@ -40,8 +40,8 @@ void UDebugVisLocationProcessor::Execute(UMassEntitySubsystem& EntitySubsystem,
 		EntityQuery.ForEachEntityChunk(EntitySubsystem, Context, [this, &VisualDataISMCs](const FMassExecutionContext& Context)
 		{
 			const int32 NumEntities = Context.GetNumEntities();
-			const TConstArrayView<FDataFragment_Transform> LocationList = Context.GetComponentView<FDataFragment_Transform>();
-			const TConstArrayView<FSimDebugVisComponent> DebugVisList = Context.GetComponentView<FSimDebugVisComponent>();
+			const TConstArrayView<FDataFragment_Transform> LocationList = Context.GetFragmentView<FDataFragment_Transform>();
+			const TConstArrayView<FSimDebugVisComponent> DebugVisList = Context.GetFragmentView<FSimDebugVisComponent>();
 
 			for (int32 i = 0; i < NumEntities; ++i)
 			{
@@ -105,9 +105,9 @@ void UMassProcessor_UpdateDebugVis::Execute(UMassEntitySubsystem& EntitySubsyste
 	EntityQuery.ForEachEntityChunk(EntitySubsystem, Context, [this, Debugger](FMassExecutionContext& Context)
 		{
 			const int32 NumEntities = Context.GetNumEntities();
-			const TConstArrayView<FDataFragment_Transform> LocationList = Context.GetComponentView<FDataFragment_Transform>();
-			const TArrayView<FDataFragment_DebugVis> DebugVisList = Context.GetMutableComponentView<FDataFragment_DebugVis>();
-			const TArrayView<FDataFragment_AgentRadius> RadiiList = Context.GetMutableComponentView<FDataFragment_AgentRadius>();
+			const TConstArrayView<FDataFragment_Transform> LocationList = Context.GetFragmentView<FDataFragment_Transform>();
+			const TArrayView<FDataFragment_DebugVis> DebugVisList = Context.GetMutableFragmentView<FDataFragment_DebugVis>();
+			const TArrayView<FDataFragment_AgentRadius> RadiiList = Context.GetMutableFragmentView<FDataFragment_AgentRadius>();
 
 			for (int32 i = 0; i < NumEntities; ++i)
 			{
diff --git a/Engine/Plugins/Runtime/MassGameplay/Source/MassLOD/Private/MassLODCollectorProcessor.cpp b/Engine/Plugins/Runtime/MassGameplay/Source/MassLOD/Private/MassLODCollectorProcessor.cpp
--- a/Engine/Plugins/Runtime/MassGameplay/Source/MassLOD/Private/MassLODCollectorProcessor.cpp
+++ b/Engine/Plugins/Runtime/MassGameplay/Source/MassLOD/Private/MassLODCollectorProcessor.cpp
@@ -69,8 +69,8 @@ void UMassLODCollectorProcessor::ConfigureQueries()
 		LODConfig.FarAndOffLODEntityQuery_Conditional.AddChunkRequirement<FMassVisualizationChunkFragment>(EMassFragmentAccess::ReadOnly);
 		LODConfig.FarAndOffLODEntityQuery_Conditional.SetChunkFilter([](const FMassExecutionContext& Context)
 			{
-				return Context.GetChunkComponent<FMassVisualizationChunkFragment>().ShouldUpdateVisualization()
-					|| Context.GetChunkComponent<FMassSimulationVariableTickChunkFragment>().ShouldTickThisFrame();
+				return Context.GetChunkFragment<FMassVisualizationChunkFragment>().ShouldUpdateVisualization()
+					|| Context.GetChunkFragment<FMassSimulationVariableTickChunkFragment>().ShouldTickThisFrame();
 			});
 	}
 }
@@ -82,8 +82,8 @@ void CollectLOD(UMassEntitySubsystem& EntitySubsystem, FMassExecutionContext& Co
 
 	auto CollectLODInfo = [&Collector](FMassExecutionContext& Context)
 	{
-		TConstArrayView<FDataFragment_Transform> LocationList = Context.GetComponentView<FDataFragment_Transform>();
-		TArrayView<FMassLODInfoFragment> ViewerInfoList = Context.GetMutableComponentView<FMassLODInfoFragment>();
+		TConstArrayView<FDataFragment_Transform> LocationList = Context.GetFragmentView<FDataFragment_Transform>();
+		TArrayView<FMassLODInfoFragment> ViewerInfoList = Context.GetMutableFragmentView<FMassLODInfoFragment>();
 
 		Collector.CollectLODInfo(Context, LocationList, ViewerInfoList);
 	};
diff --git a/Engine/Plugins/Runtime/MassGameplay/Source/MassLOD/Private/MassRepresentationLODCollectorProcessor.cpp b/Engine/Plugins/Runtime/MassGameplay/Source/MassLOD/Private/MassRepresentationLODCollectorProcessor.cpp
--- a/Engine/Plugins/Runtime/MassGameplay/Source/MassLOD/Private/MassRepresentationLODCollectorProcessor.cpp
+++ b/Engine/Plugins/Runtime/MassGameplay/Source/MassLOD/Private/MassRepresentationLODCollectorProcessor.cpp
@@ -48,8 +48,8 @@ void CollectLODInfo(UMassEntitySubsystem& EntitySubsystem, FMassExecutionContext
 
 	auto InternalCollectLODInfo = [&Collector](FMassExecutionContext& Context)
 	{
-		TConstArrayView<FDataFragment_Transform> LocationList = Context.GetComponentView<FDataFragment_Transform>();
-		TArrayView<FMassLODInfoFragment> ViewerInfoList = Context.GetMutableComponentView<FMassLODInfoFragment>();
+		TConstArrayView<FDataFragment_Transform> LocationList = Context.GetFragmentView<FDataFragment_Transform>();
+		TArrayView<FMassLODInfoFragment> ViewerInfoList = Context.GetMutableFragmentView<FMassLODInfoFragment>();
 
 		Collector.CollectLODInfo(Context, LocationList, ViewerInfoList);
 	};
diff --git a/Engine/Plugins/Runtime/MassGameplay/Source/MassLOD/Private/MassSimulationLOD.cpp b/Engine/Plugins/Runtime/MassGameplay/Source/MassLOD/Private/MassSimulationLOD.cpp
--- a/Engine/Plugins/Runtime/MassGameplay/Source/MassLOD/Private/MassSimulationLOD.cpp
+++ b/Engine/Plugins/Runtime/MassGameplay/Source/MassLOD/Private/MassSimulationLOD.cpp
@@ -36,8 +36,8 @@ void UMassProcessor_MassSimulationLODViewersInfo::Execute(UMassEntitySubsystem&
 
 	EntityQuery.ForEachEntityChunk(EntitySubsystem, Context, [this](FMassExecutionContext& Context)
 	{
-		const TConstArrayView<FDataFragment_Transform> LocationList = Context.GetComponentView<FDataFragment_Transform>();
-		const TArrayView<FDataFragment_MassSimulationLODInfo> ViewersInfoList = Context.GetMutableComponentView<FDataFragment_MassSimulationLODInfo>();
+		const TConstArrayView<FDataFragment_Transform> LocationList = Context.GetFragmentView<FDataFragment_Transform>();
+		const TArrayView<FDataFragment_MassSimulationLODInfo> ViewersInfoList = Context.GetMutableFragmentView<FDataFragment_MassSimulationLODInfo>();
 		LODCollector.CollectLODInfo(Context, LocationList, ViewersInfoList);
 	});
 }
@@ -132,8 +132,8 @@ void UMassSimulationLODProcessor::CalculateLODForConfig(UMassEntitySubsystem& En
 				return;
 			}
 
-			const TConstArrayView<FMassLODInfoFragment> ViewersInfoList = Context.GetComponentView<FMassLODInfoFragment>();
-			const TArrayView<FMassSimulationLODFragment> SimulationLODFragments = Context.GetMutableComponentView<FMassSimulationLODFragment>();
+			const TConstArrayView<FMassLODInfoFragment> ViewersInfoList = Context.GetFragmentView<FMassLODInfoFragment>();
+			const TArrayView<FMassSimulationLODFragment> SimulationLODFragments = Context.GetMutableFragmentView<FMassSimulationLODFragment>();
 			LODConfig.LODCalculator.CalculateLOD(Context, SimulationLODFragments, ViewersInfoList);
 		});
 	}
@@ -150,8 +150,8 @@ void UMassSimulationLODProcessor::CalculateLODForConfig(UMassEntitySubsystem& En
 					return;
 				}
 
-				const TConstArrayView<FMassLODInfoFragment> ViewersInfoList = Context.GetComponentView<FMassLODInfoFragment>();
-				const TArrayView<FMassSimulationLODFragment> SimulationLODFragments = Context.GetMutableComponentView<FMassSimulationLODFragment>();
+				const TConstArrayView<FMassLODInfoFragment> ViewersInfoList = Context.GetFragmentView<FMassLODInfoFragment>();
+				const TArrayView<FMassSimulationLODFragment> SimulationLODFragments = Context.GetMutableFragmentView<FMassSimulationLODFragment>();
 				LODConfig.LODCalculator.AdjustLODFromCount(Context, SimulationLODFragments, ViewersInfoList);
 			});
 		}
@@ -163,7 +163,7 @@ void UMassSimulationLODProcessor::CalculateLODForConfig(UMassEntitySubsystem& En
 		const float Time = World->GetTimeSeconds();
 		LODConfig.EntityQuery.ForEachEntityChunk(EntitySubsystem, Context, [&LODConfig, Time](FMassExecutionContext& Context)
 		{
-			TArrayView<FMassSimulationLODFragment> SimulationLODFragments = Context.GetMutableComponentView<FMassSimulationLODFragment>();
+			TArrayView<FMassSimulationLODFragment> SimulationLODFragments = Context.GetMutableFragmentView<FMassSimulationLODFragment>();
 			LODConfig.LODTickRateController.UpdateTickRateFromLOD(Context, SimulationLODFragments, Time);
 		});
 	}
@@ -175,8 +175,8 @@ void UMassSimulationLODProcessor::CalculateLODForConfig(UMassEntitySubsystem& En
 
 		LODConfig.EntityQuery.ForEachEntityChunk(EntitySubsystem, Context, [this, &LODConfig](FMassExecutionContext& Context)
 		{
-			const TConstArrayView<FDataFragment_Transform> LocationList = Context.GetComponentView<FDataFragment_Transform>();
-			const TConstArrayView<FMassSimulationLODFragment> SimulationLODList = Context.GetComponentView<FMassSimulationLODFragment>();
+			const TConstArrayView<FDataFragment_Transform> LocationList = Context.GetFragmentView<FDataFragment_Transform>();
+			const TConstArrayView<FMassSimulationLODFragment> SimulationLODList = Context.GetFragmentView<FMassSimulationLODFragment>();
 			LODConfig.LODCalculator.DebugDisplayLOD(Context, SimulationLODList, LocationList, World);
 		});
 	}
diff --git a/Engine/Plugins/Runtime/MassGameplay/Source/MassLOD/Public/MassLODTickRateController.h b/Engine/Plugins/Runtime/MassGameplay/Source/MassLOD/Public/MassLODTickRateController.h
--- a/Engine/Plugins/Runtime/MassGameplay/Source/MassLOD/Public/MassLODTickRateController.h
+++ b/Engine/Plugins/Runtime/MassGameplay/Source/MassLOD/Public/MassLODTickRateController.h
@@ -68,7 +68,7 @@ template <typename FVariableTickChunkFragment, typename FLODLogic>
 bool TMassLODTickRateController<FVariableTickChunkFragment, FLODLogic>::ShouldCalculateLODForChunk(const FMassExecutionContext& Context) const
 {
 	// EMassLOD::Off does not need to handle max count, so we can use ticking rate for them if available
-	const FMassVariableTickChunkFragment& ChunkData = Context.GetChunkComponent<FVariableTickChunkFragment>();
+	const FMassVariableTickChunkFragment& ChunkData = Context.GetChunkFragment<FVariableTickChunkFragment>();
 	return ChunkData.GetLOD() != EMassLOD::Off || ChunkData.ShouldTickThisFrame();
 }
 
@@ -76,7 +76,7 @@ template <typename FVariableTickChunkFragment, typename FLODLogic>
 bool TMassLODTickRateController<FVariableTickChunkFragment, FLODLogic>::ShouldAdjustLODFromCountForChunk(const FMassExecutionContext& Context) const
 {
 	// EMassLOD::Off does not need to handle max count, so we can skip it
-	const FMassVariableTickChunkFragment& ChunkData = Context.GetChunkComponent<FVariableTickChunkFragment>();
+	const FMassVariableTickChunkFragment& ChunkData = Context.GetChunkFragment<FVariableTickChunkFragment>();
 	return ChunkData.GetLOD() != EMassLOD::Off;
 }
 
@@ -89,11 +89,11 @@ bool TMassLODTickRateController<FVariableTickChunkFragment, FLODLogic>::UpdateTi
 	bool bWasChunkTicked = true;
 	const float DeltaTime = Context.GetDeltaTimeSeconds();
 
-	FMassVariableTickChunkFragment& ChunkData = Context.GetMutableChunkComponent<FVariableTickChunkFragment>();
+	FMassVariableTickChunkFragment& ChunkData = Context.GetMutableChunkFragment<FVariableTickChunkFragment>();
 	EMassLOD::Type ChunkLOD = ChunkData.GetLOD();
 	if (ChunkLOD == EMassLOD::Max)
 	{
-		// The LOD on the chunk component data isn't set yet, let see if the Archetype has an LOD tag and set it on the ChunkData
+		// The LOD on the chunk fragment data isn't set yet, let see if the Archetype has an LOD tag and set it on the ChunkData
 		ChunkLOD = UE::MassLOD::GetLODFromArchetype(Context);
 		ChunkData.SetLOD(ChunkLOD);
 	}
diff --git a/Engine/Plugins/Runtime/MassGameplay/Source/MassLOD/Public/MassLODTypes.h b/Engine/Plugins/Runtime/MassGameplay/Source/MassLOD/Public/MassLODTypes.h
--- a/Engine/Plugins/Runtime/MassGameplay/Source/MassLOD/Public/MassLODTypes.h
+++ b/Engine/Plugins/Runtime/MassGameplay/Source/MassLOD/Public/MassLODTypes.h
@@ -172,7 +172,7 @@ struct MASSLOD_API FMassVisualizationChunkFragment : public FMassChunkFragment
 	 */
 	static bool AreAnyEntitiesVisibleInChunk(const FMassExecutionContext& Context)
 	{
-		return Context.GetChunkComponent<FMassVisualizationChunkFragment>().AreAnyEntitiesVisible();
+		return Context.GetChunkFragment<FMassVisualizationChunkFragment>().AreAnyEntitiesVisible();
 	}
 
 	/**
@@ -193,7 +193,7 @@ struct MASSLOD_API FMassVisualizationChunkFragment : public FMassChunkFragment
 	 */
 	static bool ShouldUpdateVisualizationForChunk(const FMassExecutionContext& Context)
 	{
-		return Context.GetChunkComponent<FMassVisualizationChunkFragment>().ShouldUpdateVisualization();
+		return Context.GetChunkFragment<FMassVisualizationChunkFragment>().ShouldUpdateVisualization();
 	}
 
 	/**
diff --git a/Engine/Plugins/Runtime/MassGameplay/Source/MassLOD/Public/MassSimulationLOD.h b/Engine/Plugins/Runtime/MassGameplay/Source/MassLOD/Public/MassSimulationLOD.h
--- a/Engine/Plugins/Runtime/MassGameplay/Source/MassLOD/Public/MassSimulationLOD.h
+++ b/Engine/Plugins/Runtime/MassGameplay/Source/MassLOD/Public/MassSimulationLOD.h
@@ -75,13 +75,13 @@ struct FMassSimulationVariableTickChunkFragment : public FMassVariableTickChunkF
 
 	static bool ShouldTickChunkThisFrame(const FMassExecutionContext& Context)
 	{
-		const FMassSimulationVariableTickChunkFragment& ChunkFragment = Context.GetChunkComponent<FMassSimulationVariableTickChunkFragment>();
+		const FMassSimulationVariableTickChunkFragment& ChunkFragment = Context.GetChunkFragment<FMassSimulationVariableTickChunkFragment>();
 		return ChunkFragment.ShouldTickThisFrame();
 	}
 
 	static EMassLOD::Type GetChunkLOD(const FMassExecutionContext& Context)
 	{
-		const FMassSimulationVariableTickChunkFragment& ChunkFragment = Context.GetChunkComponent<FMassSimulationVariableTickChunkFragment>();
+		const FMassSimulationVariableTickChunkFragment& ChunkFragment = Context.GetChunkFragment<FMassSimulationVariableTickChunkFragment>();
 		return ChunkFragment.GetLOD();
 	}
 
diff --git a/Engine/Plugins/Runtime/MassGameplay/Source/MassReplication/Private/MassReplicationManager.cpp b/Engine/Plugins/Runtime/MassGameplay/Source/MassReplication/Private/MassReplicationManager.cpp
--- a/Engine/Plugins/Runtime/MassGameplay/Source/MassReplication/Private/MassReplicationManager.cpp
+++ b/Engine/Plugins/Runtime/MassGameplay/Source/MassReplication/Private/MassReplicationManager.cpp
@@ -53,7 +53,7 @@ void UMassReplicationManager::Deinitialize()
 FMassNetworkID UMassReplicationManager::GetNetIDFromHandle(const FMassEntityHandle Handle) const
 {
 	check(EntitySystem);
-	const FMassNetworkIDFragment& Data = EntitySystem->GetComponentDataChecked<FMassNetworkIDFragment>(Handle);
+	const FMassNetworkIDFragment& Data = EntitySystem->GetFragmentDataChecked<FMassNetworkIDFragment>(Handle);
 	return Data.NetID;
 }
 
diff --git a/Engine/Plugins/Runtime/MassGameplay/Source/MassReplication/Private/MassReplicationProcessorBase.cpp b/Engine/Plugins/Runtime/MassGameplay/Source/MassReplication/Private/MassReplicationProcessorBase.cpp
--- a/Engine/Plugins/Runtime/MassGameplay/Source/MassReplication/Private/MassReplicationProcessorBase.cpp
+++ b/Engine/Plugins/Runtime/MassGameplay/Source/MassReplication/Private/MassReplicationProcessorBase.cpp
@@ -138,8 +138,8 @@ void UMassReplicationProcessorBase::Execute(UMassEntitySubsystem& EntitySubsyste
 
 	EntityQuery.ForEachEntityChunk(EntitySubsystem, Context, [this](FMassExecutionContext& Context)
 		{
-			const TConstArrayView<FDataFragment_MassSimulationLODInfo> ViewersInfoList = Context.GetComponentView<FDataFragment_MassSimulationLODInfo>();
-			const TArrayView<FMassReplicationViewerLODFragment> ViewerLODList = Context.GetMutableComponentView<FMassReplicationViewerLODFragment>();
+			const TConstArrayView<FDataFragment_MassSimulationLODInfo> ViewersInfoList = Context.GetFragmentView<FDataFragment_MassSimulationLODInfo>();
+			const TArrayView<FMassReplicationViewerLODFragment> ViewerLODList = Context.GetMutableFragmentView<FMassReplicationViewerLODFragment>();
 			LODCalculator.CalculateLOD(Context, ViewerLODList, ViewersInfoList);
 		});
 
@@ -147,8 +147,8 @@ void UMassReplicationProcessorBase::Execute(UMassEntitySubsystem& EntitySubsyste
 	{
 		EntityQuery.ForEachEntityChunk(EntitySubsystem, Context, [this](FMassExecutionContext& Context)
 			{
-				const TConstArrayView<FDataFragment_MassSimulationLODInfo> ViewersInfoList = Context.GetComponentView<FDataFragment_MassSimulationLODInfo>();
-				const TArrayView<FMassReplicationViewerLODFragment> ViewerLODList = Context.GetMutableComponentView<FMassReplicationViewerLODFragment>();
+				const TConstArrayView<FDataFragment_MassSimulationLODInfo> ViewersInfoList = Context.GetFragmentView<FDataFragment_MassSimulationLODInfo>();
+				const TArrayView<FMassReplicationViewerLODFragment> ViewerLODList = Context.GetMutableFragmentView<FMassReplicationViewerLODFragment>();
 				LODCalculator.AdjustLODFromCount(Context, ViewerLODList, ViewersInfoList);
 			});
 	}
diff --git a/Engine/Plugins/Runtime/MassGameplay/Source/MassReplication/Private/MassReplicationTransformHandlers.cpp b/Engine/Plugins/Runtime/MassGameplay/Source/MassReplication/Private/MassReplicationTransformHandlers.cpp
--- a/Engine/Plugins/Runtime/MassGameplay/Source/MassReplication/Private/MassReplicationTransformHandlers.cpp
+++ b/Engine/Plugins/Runtime/MassGameplay/Source/MassReplication/Private/MassReplicationTransformHandlers.cpp
@@ -10,9 +10,9 @@ void FMassReplicationProcessorTransformHandlerBase::AddRequirements(FMassEntityQ
 	InQuery.AddRequirement<FDataFragment_Transform>(EMassFragmentAccess::ReadWrite);
 }
 
-void FMassReplicationProcessorTransformHandlerBase::CacheComponentViews(FMassExecutionContext& ExecContext)
+void FMassReplicationProcessorTransformHandlerBase::CacheFragmentViews(FMassExecutionContext& ExecContext)
 {
-	TransformList = ExecContext.GetMutableComponentView<FDataFragment_Transform>();
+	TransformList = ExecContext.GetMutableFragmentView<FDataFragment_Transform>();
 }
 
 void FMassReplicationProcessorPositionYawHandler::AddEntity(const int32 EntityIdx, FReplicatedAgentPositionYawData& InOutReplicatedPositionYawData) const
diff --git a/Engine/Plugins/Runtime/MassGameplay/Source/MassReplication/Private/MassReplicationTypes.cpp b/Engine/Plugins/Runtime/MassGameplay/Source/MassReplication/Private/MassReplicationTypes.cpp
--- a/Engine/Plugins/Runtime/MassGameplay/Source/MassReplication/Private/MassReplicationTypes.cpp
+++ b/Engine/Plugins/Runtime/MassGameplay/Source/MassReplication/Private/MassReplicationTypes.cpp
@@ -43,7 +43,7 @@ void UMassNetworkIDFragmentInitializer::Execute(UMassEntitySubsystem& EntitySubs
 
 		EntityQuery.ForEachEntityChunk(EntitySubsystem, Context, [this](FMassExecutionContext& Context)
 			{
-				const TArrayView<FMassNetworkIDFragment> NetworkIDList = Context.GetMutableComponentView<FMassNetworkIDFragment>();
+				const TArrayView<FMassNetworkIDFragment> NetworkIDList = Context.GetMutableFragmentView<FMassNetworkIDFragment>();
 				const int32 NumEntities = Context.GetNumEntities();
 
 				for (int32 Idx = 0; Idx < NumEntities; ++Idx)
diff --git a/Engine/Plugins/Runtime/MassGameplay/Source/MassReplication/Public/MassClientBubbleHandler.h b/Engine/Plugins/Runtime/MassGameplay/Source/MassReplication/Public/MassClientBubbleHandler.h
--- a/Engine/Plugins/Runtime/MassGameplay/Source/MassReplication/Public/MassClientBubbleHandler.h
+++ b/Engine/Plugins/Runtime/MassGameplay/Source/MassReplication/Public/MassClientBubbleHandler.h
@@ -132,7 +132,7 @@ public:
 	friend class TMassClientBubbleTransformHandler;
 
 	typedef TFunctionRef<void(FMassEntityQuery&)> FAddRequirementsForSpawnQueryFunction;
-	typedef TFunctionRef<void(FMassExecutionContext&)> FCacheComponentViewsForSpawnQueryFunction;
+	typedef TFunctionRef<void(FMassExecutionContext&)> FCacheFragmentViewsForSpawnQueryFunction;
 	typedef TFunctionRef<void(const FMassEntityView&, const typename AgentArrayItem::FReplicatedAgentType&, const int32)> FSetSpawnedEntityDataFunction;
 	typedef TFunctionRef<void(const FMassEntityView&, const typename AgentArrayItem::FReplicatedAgentType&)> FSetModifiedEntityDataFunction;
 
@@ -164,11 +164,11 @@ protected:
 
 	/** Called from TClientBubbleHandlerBase derived classes in PostReplicatedAd() */
 	void PostReplicatedAddHelper(const TArrayView<int32> AddedIndices, FAddRequirementsForSpawnQueryFunction AddRequirementsForSpawnQuery
-		, FCacheComponentViewsForSpawnQueryFunction CacheComponentViewsForSpawnQuery, FSetSpawnedEntityDataFunction SetSpawnedEntityData, FSetModifiedEntityDataFunction SetModifiedEntityData);
+		, FCacheFragmentViewsForSpawnQueryFunction CacheFragmentViewsForSpawnQuery, FSetSpawnedEntityDataFunction SetSpawnedEntityData, FSetModifiedEntityDataFunction SetModifiedEntityData);
 
 	/** used by PostReplicatedAddHelper */
 	void PostReplicatedAddEntitiesHelper(const TArrayView<int32> AddedIndices, FAddRequirementsForSpawnQueryFunction AddRequirementsForSpawnQuery
-		, FCacheComponentViewsForSpawnQueryFunction CacheComponentViewsForSpawnQuery, FSetSpawnedEntityDataFunction SetSpawnedEntityData);
+		, FCacheFragmentViewsForSpawnQueryFunction CacheFragmentViewsForSpawnQuery, FSetSpawnedEntityDataFunction SetSpawnedEntityData);
 
 
 	/** Called from TClientBubbleHandlerBase derived classes in PostReplicatedChange() */
@@ -423,7 +423,7 @@ void TClientBubbleHandlerBase<AgentArrayItem>::PreReplicatedRemove(const TArrayV
 #if UE_REPLICATION_COMPILE_CLIENT_CODE
 template<typename AgentArrayItem>
 void TClientBubbleHandlerBase<AgentArrayItem>::PostReplicatedAddHelper(const TArrayView<int32> AddedIndices, FAddRequirementsForSpawnQueryFunction AddRequirementsForSpawnQuery
-	, FCacheComponentViewsForSpawnQueryFunction CacheComponentViewsForSpawnQuery, FSetSpawnedEntityDataFunction SetSpawnedEntityData, FSetModifiedEntityDataFunction SetModifiedEntityData)
+	, FCacheFragmentViewsForSpawnQueryFunction CacheFragmentViewsForSpawnQuery, FSetSpawnedEntityDataFunction SetSpawnedEntityData, FSetModifiedEntityDataFunction SetModifiedEntityData)
 {
 	TArray<FMassEntityHandle> EntitiesDestroy;
 
@@ -495,14 +495,14 @@ void TClientBubbleHandlerBase<AgentArrayItem>::PostReplicatedAddHelper(const TAr
 		}
 	}
 
-	PostReplicatedAddEntitiesHelper(AgentsToAddArray, AddRequirementsForSpawnQuery, CacheComponentViewsForSpawnQuery, SetSpawnedEntityData);
+	PostReplicatedAddEntitiesHelper(AgentsToAddArray, AddRequirementsForSpawnQuery, CacheFragmentViewsForSpawnQuery, SetSpawnedEntityData);
 }
 #endif // UE_REPLICATION_COMPILE_CLIENT_CODE
 
 #if UE_REPLICATION_COMPILE_CLIENT_CODE
 template<typename AgentArrayItem>
 void TClientBubbleHandlerBase<AgentArrayItem>::PostReplicatedAddEntitiesHelper(const TArrayView<int32> AddedIndices, FAddRequirementsForSpawnQueryFunction AddRequirementsForSpawnQuery
-	, FCacheComponentViewsForSpawnQueryFunction CacheComponentViewsForSpawnQuery, FSetSpawnedEntityDataFunction SetSpawnedEntityData)
+	, FCacheFragmentViewsForSpawnQueryFunction CacheFragmentViewsForSpawnQuery, FSetSpawnedEntityDataFunction SetSpawnedEntityData)
 {
 	check(Serializer);
 
@@ -554,11 +554,11 @@ void TClientBubbleHandlerBase<AgentArrayItem>::PostReplicatedAddEntitiesHelper(c
 
 		int32 AgentsSpawnIdx = 0;
 
-		Query.ForEachEntityChunk(FArchetypeChunkCollection(ArchetypeHandle, Entities), *EntitySystem, ExecContext, [&AgentsSpawn, &AgentsSpawnIdx, this, ReplicationManager, &ExecContext, &CacheComponentViewsForSpawnQuery, &SetSpawnedEntityData, &EntitySystem](FMassExecutionContext& Context)
+		Query.ForEachEntityChunk(FArchetypeChunkCollection(ArchetypeHandle, Entities), *EntitySystem, ExecContext, [&AgentsSpawn, &AgentsSpawnIdx, this, ReplicationManager, &ExecContext, &CacheFragmentViewsForSpawnQuery, &SetSpawnedEntityData, &EntitySystem](FMassExecutionContext& Context)
 			{
-				CacheComponentViewsForSpawnQuery(ExecContext);
+				CacheFragmentViewsForSpawnQuery(ExecContext);
 
-				const TArrayView<FMassNetworkIDFragment> NetworkIDList = Context.GetMutableComponentView<FMassNetworkIDFragment>();
+				const TArrayView<FMassNetworkIDFragment> NetworkIDList = Context.GetMutableFragmentView<FMassNetworkIDFragment>();
 
 				for (int32 i = 0; i < Context.GetNumEntities(); ++i)
 				{
@@ -678,7 +678,7 @@ void TClientBubbleHandlerBase<AgentArrayItem>::DebugValidateBubbleOnClient()
 
 				if (bIsEntityValid)
 				{
-					const FMassNetworkIDFragment& FragmentNetID = EntitySystem->GetComponentDataChecked<FMassNetworkIDFragment>(EntityInfo->Entity);
+					const FMassNetworkIDFragment& FragmentNetID = EntitySystem->GetFragmentDataChecked<FMassNetworkIDFragment>(EntityInfo->Entity);
 
 					checkf(FragmentNetID.NetID == Agent.GetNetID(), TEXT("Fragment and Agent NetID do not match!"));
 				}
@@ -741,12 +741,12 @@ void TClientBubbleHandlerBase<AgentArrayItem>::DebugValidateBubbleOnServer()
 		checkf(LookupData.AgentsIdx == OuterIdx, TEXT("Agent index must match lookup data!"));
 		checkf(EntitySystem->IsEntityValid(LookupData.Entity), TEXT("Must be valid entity"));
 
-		const FMassNetworkIDFragment& FragmentNetID = EntitySystem->GetComponentDataChecked<FMassNetworkIDFragment>(LookupData.Entity);
+		const FMassNetworkIDFragment& FragmentNetID = EntitySystem->GetFragmentDataChecked<FMassNetworkIDFragment>(LookupData.Entity);
 
 		checkf(FragmentNetID.NetID == OuterItem.Agent.GetNetID(), TEXT("Fragment and Agent NetID do not match!"));
 		checkf(LookupData.NetID == OuterItem.Agent.GetNetID(), TEXT("LookupData and Agent NetID do not match!"));
 
-		const FDataFragment_ReplicationTemplateID& FragmentTemplateID = EntitySystem->GetComponentDataChecked<FDataFragment_ReplicationTemplateID>(LookupData.Entity);
+		const FDataFragment_ReplicationTemplateID& FragmentTemplateID = EntitySystem->GetFragmentDataChecked<FDataFragment_ReplicationTemplateID>(LookupData.Entity);
 		checkf(FragmentTemplateID.ID == OuterItem.Agent.GetTemplateID(), TEXT("Agent TemplateID different to Fragment!"));
 	}
 
diff --git a/Engine/Plugins/Runtime/MassGameplay/Source/MassReplication/Public/MassReplicationProcessorBase.h b/Engine/Plugins/Runtime/MassGameplay/Source/MassReplication/Public/MassReplicationProcessorBase.h
--- a/Engine/Plugins/Runtime/MassGameplay/Source/MassReplication/Public/MassReplicationProcessorBase.h
+++ b/Engine/Plugins/Runtime/MassGameplay/Source/MassReplication/Public/MassReplicationProcessorBase.h
@@ -97,10 +97,10 @@ void UMassReplicationProcessorBase::CalculateClientReplication(FMassExecutionCon
 
 	const int32 NumEntities = Context.GetNumEntities();
 
-	TConstArrayView<FMassNetworkIDFragment> NetworkIDList = Context.GetComponentView<FMassNetworkIDFragment>();
-	TArrayView<FMassReplicationViewerLODFragment> ViewerLODList = Context.GetMutableComponentView<FMassReplicationViewerLODFragment>();
-	TArrayView<FMassReplicatedAgentFragment> ReplicatedAgentList = Context.GetMutableComponentView<FMassReplicatedAgentFragment>();
-	TConstArrayView<FDataFragment_ReplicationTemplateID> TemplateIDList = Context.GetComponentView<FDataFragment_ReplicationTemplateID>();
+	TConstArrayView<FMassNetworkIDFragment> NetworkIDList = Context.GetFragmentView<FMassNetworkIDFragment>();
+	TArrayView<FMassReplicationViewerLODFragment> ViewerLODList = Context.GetMutableFragmentView<FMassReplicationViewerLODFragment>();
+	TArrayView<FMassReplicatedAgentFragment> ReplicatedAgentList = Context.GetMutableFragmentView<FMassReplicatedAgentFragment>();
+	TConstArrayView<FDataFragment_ReplicationTemplateID> TemplateIDList = Context.GetFragmentView<FDataFragment_ReplicationTemplateID>();
 
 	CacheViews();
 
diff --git a/Engine/Plugins/Runtime/MassGameplay/Source/MassReplication/Public/MassReplicationTransformHandlers.h b/Engine/Plugins/Runtime/MassGameplay/Source/MassReplication/Public/MassReplicationTransformHandlers.h
--- a/Engine/Plugins/Runtime/MassGameplay/Source/MassReplication/Public/MassReplicationTransformHandlers.h
+++ b/Engine/Plugins/Runtime/MassGameplay/Source/MassReplication/Public/MassReplicationTransformHandlers.h
@@ -77,8 +77,8 @@ public:
 	 * The following functions are used to configure the query and then set the position and yaw data.
 	 */
 	static void AddRequirementsForSpawnQuery(FMassEntityQuery& InQuery);
-	void CacheComponentViewsForSpawnQuery(FMassExecutionContext& InExecContext);
-	void ClearComponentViewsForSpawnQuery();
+	void CacheFragmentViewsForSpawnQuery(FMassExecutionContext& InExecContext);
+	void ClearFragmentViewsForSpawnQuery();
 
 	void SetSpawnedEntityData(const int32 EntityIdx, const FReplicatedAgentPositionYawData& ReplicatedPathData) const;
 
@@ -151,15 +151,15 @@ void TMassClientBubbleTransformHandler<AgentArrayItem>::AddRequirementsForSpawnQ
 
 #if UE_REPLICATION_COMPILE_CLIENT_CODE
 template<typename AgentArrayItem>
-void TMassClientBubbleTransformHandler<AgentArrayItem>::CacheComponentViewsForSpawnQuery(FMassExecutionContext& InExecContext)
+void TMassClientBubbleTransformHandler<AgentArrayItem>::CacheFragmentViewsForSpawnQuery(FMassExecutionContext& InExecContext)
 {
-	TransformList = InExecContext.GetMutableComponentView<FDataFragment_Transform>();
+	TransformList = InExecContext.GetMutableFragmentView<FDataFragment_Transform>();
 }
 #endif // UE_REPLICATION_COMPILE_CLIENT_CODE
 
 #if UE_REPLICATION_COMPILE_CLIENT_CODE
 template<typename AgentArrayItem>
-void TMassClientBubbleTransformHandler<AgentArrayItem>::ClearComponentViewsForSpawnQuery()
+void TMassClientBubbleTransformHandler<AgentArrayItem>::ClearFragmentViewsForSpawnQuery()
 {
 	TransformList = TArrayView<FDataFragment_Transform>();
 }
@@ -179,7 +179,7 @@ void TMassClientBubbleTransformHandler<AgentArrayItem>::SetSpawnedEntityData(con
 template<typename AgentArrayItem>
 void TMassClientBubbleTransformHandler<AgentArrayItem>::SetModifiedEntityData(const FMassEntityView& EntityView, const FReplicatedAgentPositionYawData& ReplicatedPositionYawData)
 {
-	FDataFragment_Transform& TransformFragment = EntityView.GetComponentData<FDataFragment_Transform>();
+	FDataFragment_Transform& TransformFragment = EntityView.GetFragmentData<FDataFragment_Transform>();
 
 	SetEntityData(TransformFragment, ReplicatedPositionYawData);
 }
@@ -200,7 +200,7 @@ class MASSREPLICATION_API FMassReplicationProcessorTransformHandlerBase
 {
 public:
 	static void AddRequirements(FMassEntityQuery& InQuery);
-	void CacheComponentViews(FMassExecutionContext& ExecContext);
+	void CacheFragmentViews(FMassExecutionContext& ExecContext);
 
 protected:
 	TArrayView<FDataFragment_Transform> TransformList;
diff --git a/Engine/Plugins/Runtime/MassGameplay/Source/MassReplication/Public/MassReplicationTypes.h b/Engine/Plugins/Runtime/MassGameplay/Source/MassReplication/Public/MassReplicationTypes.h
--- a/Engine/Plugins/Runtime/MassGameplay/Source/MassReplication/Public/MassReplicationTypes.h
+++ b/Engine/Plugins/Runtime/MassGameplay/Source/MassReplication/Public/MassReplicationTypes.h
@@ -118,7 +118,7 @@ struct MASSREPLICATION_API FMassBubbleInfoClassHandle : public FSimpleIndexedHan
 };
 
 /**
- *  Light weight component data for the mass network id of a light weight entity
+ *  Fragment type for the mass network id of a mass entity
  */
 USTRUCT()
 struct MASSREPLICATION_API FMassNetworkIDFragment : public FMassFragment
diff --git a/Engine/Plugins/Runtime/MassGameplay/Source/MassRepresentation/Private/MassRepresentationProcessor.cpp b/Engine/Plugins/Runtime/MassGameplay/Source/MassRepresentation/Private/MassRepresentationProcessor.cpp
--- a/Engine/Plugins/Runtime/MassGameplay/Source/MassRepresentation/Private/MassRepresentationProcessor.cpp
+++ b/Engine/Plugins/Runtime/MassGameplay/Source/MassRepresentation/Private/MassRepresentationProcessor.cpp
@@ -61,10 +61,10 @@ void UMassRepresentationProcessor::UpdateRepresentation(FMassExecutionContext& C
 {
 	check(RepresentationSubsystem);
 
-	const TConstArrayView<FDataFragment_Transform> TransformList = Context.GetComponentView<FDataFragment_Transform>();
-	const TArrayView<FMassRepresentationFragment> RepresentationList = Context.GetMutableComponentView<FMassRepresentationFragment>();
-	const TConstArrayView<FMassRepresentationLODFragment> RepresentationLODList = Context.GetComponentView<FMassRepresentationLODFragment>();
-	const TArrayView<FDataFragment_Actor> ActorList = Context.GetMutableComponentView<FDataFragment_Actor>();
+	const TConstArrayView<FDataFragment_Transform> TransformList = Context.GetFragmentView<FDataFragment_Transform>();
+	const TArrayView<FMassRepresentationFragment> RepresentationList = Context.GetMutableFragmentView<FMassRepresentationFragment>();
+	const TConstArrayView<FMassRepresentationLODFragment> RepresentationLODList = Context.GetFragmentView<FMassRepresentationLODFragment>();
+	const TArrayView<FDataFragment_Actor> ActorList = Context.GetMutableFragmentView<FDataFragment_Actor>();
 
 	const int32 NumEntities = Context.GetNumEntities();
 	for (int32 EntityIdx = 0; EntityIdx < NumEntities; EntityIdx++)
@@ -289,9 +289,9 @@ void UMassRepresentationProcessor::ReleaseAnyActorOrCancelAnySpawning(UMassRepre
 {
 	if (const UMassEntitySubsystem* EntitySubsystem = UWorld::GetSubsystem<UMassEntitySubsystem>(RepresentationSubsystem.GetWorld()))
 	{
-		if (FDataFragment_Actor* ActorInfo = EntitySubsystem->GetComponentDataPtr<FDataFragment_Actor>(MassAgent))
+		if (FDataFragment_Actor* ActorInfo = EntitySubsystem->GetFragmentDataPtr<FDataFragment_Actor>(MassAgent))
 		{
-			if (FMassRepresentationFragment* Representation = EntitySubsystem->GetComponentDataPtr<FMassRepresentationFragment>(MassAgent))
+			if (FMassRepresentationFragment* Representation = EntitySubsystem->GetFragmentDataPtr<FMassRepresentationFragment>(MassAgent))
 			{
 				ReleaseAnyActorOrCancelAnySpawning(RepresentationSubsystem, MassAgent, *ActorInfo, *Representation);
 			}
@@ -329,8 +329,8 @@ void UMassRepresentationProcessor::UpdateVisualization(FMassExecutionContext& Co
 	UpdateRepresentation(Context);
 
 	// Update entity visibility
-	const TArrayView<FMassRepresentationFragment> RepresentationList = Context.GetMutableComponentView<FMassRepresentationFragment>();
-	const TConstArrayView<FMassRepresentationLODFragment> RepresentationLODList = Context.GetComponentView<FMassRepresentationLODFragment>();
+	const TArrayView<FMassRepresentationFragment> RepresentationList = Context.GetMutableFragmentView<FMassRepresentationFragment>();
+	const TConstArrayView<FMassRepresentationLODFragment> RepresentationLODList = Context.GetFragmentView<FMassRepresentationLODFragment>();
 
 	const int32 NumEntities = Context.GetNumEntities();
 	for (int32 EntityIdx = 0; EntityIdx < NumEntities; EntityIdx++)
@@ -344,12 +344,12 @@ void UMassRepresentationProcessor::UpdateVisualization(FMassExecutionContext& Co
 
 FMassVisualizationChunkFragment& UMassRepresentationProcessor::UpdateChunkVisibility(FMassExecutionContext& Context) const
 {
-	// Setup chunk component data about visibility
-	FMassVisualizationChunkFragment& ChunkData = Context.GetMutableChunkComponent<FMassVisualizationChunkFragment>();
+	// Setup chunk fragment data about visibility
+	FMassVisualizationChunkFragment& ChunkData = Context.GetMutableChunkFragment<FMassVisualizationChunkFragment>();
 	EMassVisibility ChunkVisibility = ChunkData.GetVisibility();
 	if (ChunkVisibility == EMassVisibility::Max)
 	{
-		// The visibility on the chunk component data isn't set yet, let see if the Archetype has an visibility tag and set it on the ChunkData
+		// The visibility on the chunk fragment data isn't set yet, let see if the Archetype has an visibility tag and set it on the ChunkData
 		ChunkVisibility = UE::MassRepresentation::GetVisibilityFromArchetype(Context);
 		ChunkData.SetVisibility(ChunkVisibility);
 	}
@@ -394,7 +394,7 @@ void UMassRepresentationProcessor::OnActorPreSpawn(const FMassActorSpawnRequestH
 	const FMassActorSpawnRequest& MassActorSpawnRequest = SpawnRequest.Get<FMassActorSpawnRequest>();
 	const FMassEntityView View(*PipeEntitySubsystem, MassActorSpawnRequest.MassAgent);
 
-	if (FDataFragment_Actor* ActorInfo = View.GetComponentDataPtr<FDataFragment_Actor>())
+	if (FDataFragment_Actor* ActorInfo = View.GetFragmentDataPtr<FDataFragment_Actor>())
 	{
 		// Release any existing actor
 		if (AActor* Actor = ActorInfo->GetMutable())
@@ -406,7 +406,7 @@ void UMassRepresentationProcessor::OnActorPreSpawn(const FMassActorSpawnRequestH
 			// so the Fragment passed in parameters would not be valid anymore.
 			ActorInfo->ResetAndUpdateHandleMap();
 
-			if (const FMassRepresentationFragment* Representation = View.GetComponentDataPtr<FMassRepresentationFragment>())
+			if (const FMassRepresentationFragment* Representation = View.GetFragmentDataPtr<FMassRepresentationFragment>())
 			{
 				if (!RepresentationSubsystem->ReleaseTemplateActor(MassActorSpawnRequest.MassAgent, Representation->HighResTemplateActorIndex, Actor, /*bImmediate*/ true))
 				{
@@ -428,7 +428,7 @@ EMassActorSpawnRequestAction UMassRepresentationProcessor::OnActorPostSpawn(cons
 	checkf(MassActorSpawnRequest.SpawnedActor, TEXT("Expecting valid spawned actor"));
 
 	// Might be already done if the actor has a MassAgentComponent via the callback OnMassAgentComponentEntityAssociated on the MassRepresentationSubsystem
-	FDataFragment_Actor& ActorInfo = PipeEntitySubsystem->GetComponentDataChecked<FDataFragment_Actor>(MassActorSpawnRequest.MassAgent);
+	FDataFragment_Actor& ActorInfo = PipeEntitySubsystem->GetFragmentDataChecked<FDataFragment_Actor>(MassActorSpawnRequest.MassAgent);
 	if (ActorInfo.IsValid())
 	{
 		// If already set, make sure it is pointing to the same actor.
@@ -474,8 +474,8 @@ void UMassRepresentationFragmentDestructor::Execute(UMassEntitySubsystem& Entity
 
 	EntityQuery.ForEachEntityChunk(EntitySubsystem, Context, [this](FMassExecutionContext& Context)
 	{
-		const TArrayView<FMassRepresentationFragment> RepresentationList = Context.GetMutableComponentView<FMassRepresentationFragment>();
-		const TArrayView<FDataFragment_Actor> ActorList = Context.GetMutableComponentView<FDataFragment_Actor>();
+		const TArrayView<FMassRepresentationFragment> RepresentationList = Context.GetMutableFragmentView<FMassRepresentationFragment>();
+		const TArrayView<FDataFragment_Actor> ActorList = Context.GetMutableFragmentView<FDataFragment_Actor>();
 
 		const int32 NumEntities = Context.GetNumEntities();
 		for (int32 i = 0; i < NumEntities; ++i)
diff --git a/Engine/Plugins/Runtime/MassGameplay/Source/MassRepresentation/Private/MassRepresentationSubsystem.cpp b/Engine/Plugins/Runtime/MassGameplay/Source/MassRepresentation/Private/MassRepresentationSubsystem.cpp
--- a/Engine/Plugins/Runtime/MassGameplay/Source/MassRepresentation/Private/MassRepresentationSubsystem.cpp
+++ b/Engine/Plugins/Runtime/MassGameplay/Source/MassRepresentation/Private/MassRepresentationSubsystem.cpp
@@ -396,7 +396,7 @@ void UMassRepresentationSubsystem::OnMassAgentComponentEntityDetaching(const UMa
 	if (EntitySubsystem->IsEntityValid(MassAgent) && AgentComponent.IsNetSimulating())
 	{
 		const FMassEntityView EntityView(*EntitySubsystem, MassAgent);
-		if (FMassRepresentationFragment* Representation = EntityView.GetComponentDataPtr<FMassRepresentationFragment>())
+		if (FMassRepresentationFragment* Representation = EntityView.GetFragmentDataPtr<FMassRepresentationFragment>())
 		{
 			// Force a reevaluate of the current representation
 			Representation->CurrentRepresentation = ERepresentationType::None;
diff --git a/Engine/Plugins/Runtime/MassGameplay/Source/MassRepresentation/Private/MassUpdateISMProcessor.cpp b/Engine/Plugins/Runtime/MassGameplay/Source/MassRepresentation/Private/MassUpdateISMProcessor.cpp
--- a/Engine/Plugins/Runtime/MassGameplay/Source/MassRepresentation/Private/MassUpdateISMProcessor.cpp
+++ b/Engine/Plugins/Runtime/MassGameplay/Source/MassRepresentation/Private/MassUpdateISMProcessor.cpp
@@ -49,9 +49,9 @@ void UMassUpdateISMProcessor::Execute(UMassEntitySubsystem& EntitySubsystem, FMa
 	{
 		EntityQuery.ForEachEntityChunk(EntitySubsystem, Context, [&ISMInfo](FMassExecutionContext& Context)
 		{
-			const TConstArrayView<FDataFragment_Transform> TransformList = Context.GetComponentView<FDataFragment_Transform>();
-			const TArrayView<FMassRepresentationFragment> RepresentationList = Context.GetMutableComponentView<FMassRepresentationFragment>();
-			const TConstArrayView<FMassRepresentationLODFragment> RepresentationLODList = Context.GetComponentView<FMassRepresentationLODFragment>();
+			const TConstArrayView<FDataFragment_Transform> TransformList = Context.GetFragmentView<FDataFragment_Transform>();
+			const TArrayView<FMassRepresentationFragment> RepresentationList = Context.GetMutableFragmentView<FMassRepresentationFragment>();
+			const TConstArrayView<FMassRepresentationLODFragment> RepresentationLODList = Context.GetFragmentView<FMassRepresentationLODFragment>();
 
 			const int32 NumEntities = Context.GetNumEntities();
 			for (int32 EntityIdx = 0; EntityIdx < NumEntities; EntityIdx++)
diff --git a/Engine/Plugins/Runtime/MassGameplay/Source/MassRepresentation/Public/MassRepresentationProcessor.h b/Engine/Plugins/Runtime/MassGameplay/Source/MassRepresentation/Public/MassRepresentationProcessor.h
--- a/Engine/Plugins/Runtime/MassGameplay/Source/MassRepresentation/Public/MassRepresentationProcessor.h
+++ b/Engine/Plugins/Runtime/MassGameplay/Source/MassRepresentation/Public/MassRepresentationProcessor.h
@@ -73,7 +73,7 @@ protected:
 	 * @param EnabledType is the type of enabling to do on this actor
 	 * @param Actor is the actual actor to perform enabling type on
 	 * @param EntityIdx is the entity index currently processing
-	 * @param Context is the current light weight component execution context 
+	 * @param Context is the current Mass execution context 
 	 */
 	virtual void SetActorEnabled(const EActorEnabledType EnabledType, AActor& Actor, const int32 EntityIdx, FMassExecutionContext& Context);
 
@@ -82,7 +82,7 @@ protected:
 	 * The destination will be adjusted to fit an existing capsule.
 	 * @param Transform is the new actor's transform 
 	 * @param Actor is the actual actor to teleport
-	 * @param Context is the current light weight component execution context
+	 * @param Context is the current Mass execution context
 	 */
 	virtual void TeleportActor(const FTransform& Transform, AActor& Actor, FMassExecutionContext& Context);
 
diff --git a/Engine/Plugins/Runtime/MassGameplay/Source/MassRepresentation/Public/MassVisualizationLODProcessor.h b/Engine/Plugins/Runtime/MassGameplay/Source/MassRepresentation/Public/MassVisualizationLODProcessor.h
--- a/Engine/Plugins/Runtime/MassGameplay/Source/MassRepresentation/Public/MassVisualizationLODProcessor.h
+++ b/Engine/Plugins/Runtime/MassGameplay/Source/MassRepresentation/Public/MassVisualizationLODProcessor.h
@@ -79,7 +79,7 @@ void UMassVisualizationLODProcessor::ExecuteInternal(UMassEntitySubsystem& Entit
 	{
 		CloseEntityQuery.ForEachEntityChunk(EntitySubsystem, Context, [this](FMassExecutionContext& Context)
 		{
-			TArrayView<FMassRepresentationLODFragment> RepresentationLODList = Context.GetMutableComponentView<FMassRepresentationLODFragment>();
+			TArrayView<FMassRepresentationLODFragment> RepresentationLODList = Context.GetMutableFragmentView<FMassRepresentationLODFragment>();
 			LODCalculator.ForceOffLOD(Context, RepresentationLODList);
 		});
 		return;
@@ -98,8 +98,8 @@ void UMassVisualizationLODProcessor::ExecuteInternal(UMassEntitySubsystem& Entit
 		// the "debug" section below is non-production, should be separated as well
 		auto CalculateLOD = [this](FMassExecutionContext& Context)
 		{
-			TArrayView<FMassRepresentationLODFragment> RepresentationLODList = Context.GetMutableComponentView<FMassRepresentationLODFragment>();
-			TConstArrayView<TMassViewerLODInfoFragment> ViewerInfoList = Context.GetComponentView<TMassViewerLODInfoFragment>();
+			TArrayView<FMassRepresentationLODFragment> RepresentationLODList = Context.GetMutableFragmentView<FMassRepresentationLODFragment>();
+			TConstArrayView<TMassViewerLODInfoFragment> ViewerInfoList = Context.GetFragmentView<TMassViewerLODInfoFragment>();
 			LODCalculator.CalculateLOD(Context, RepresentationLODList, ViewerInfoList);
 		};
 		CloseEntityQuery.ForEachEntityChunk(EntitySubsystem, Context, CalculateLOD);
@@ -120,8 +120,8 @@ void UMassVisualizationLODProcessor::ExecuteInternal(UMassEntitySubsystem& Entit
 		{
 			CloseEntityQuery.ForEachEntityChunk(EntitySubsystem, Context, [this](FMassExecutionContext& Context)
 			{
-				TConstArrayView<TMassViewerLODInfoFragment> ViewerInfoList = Context.GetComponentView<TMassViewerLODInfoFragment>();
-				TArrayView<FMassRepresentationLODFragment> RepresentationLODList = Context.GetMutableComponentView<FMassRepresentationLODFragment>();
+				TConstArrayView<TMassViewerLODInfoFragment> ViewerInfoList = Context.GetFragmentView<TMassViewerLODInfoFragment>();
+				TArrayView<FMassRepresentationLODFragment> RepresentationLODList = Context.GetMutableFragmentView<FMassRepresentationLODFragment>();
 				LODCalculator.AdjustLODFromCount(Context, RepresentationLODList, ViewerInfoList);
 			});
 			// Far entities do not need to maximice count
@@ -135,8 +135,8 @@ void UMassVisualizationLODProcessor::ExecuteInternal(UMassEntitySubsystem& Entit
 
 		auto DebugDisplayLOD = [this](FMassExecutionContext& Context)
 		{
-			TConstArrayView<FMassRepresentationLODFragment> RepresentationLODList = Context.GetComponentView<FMassRepresentationLODFragment>();
-			TConstArrayView<FDataFragment_Transform> TransformList = Context.GetComponentView<FDataFragment_Transform>();
+			TConstArrayView<FMassRepresentationLODFragment> RepresentationLODList = Context.GetFragmentView<FMassRepresentationLODFragment>();
+			TConstArrayView<FDataFragment_Transform> TransformList = Context.GetFragmentView<FDataFragment_Transform>();
 			LODCalculator.DebugDisplayLOD(Context, RepresentationLODList, TransformList, World);
 		};
 
diff --git a/Engine/Plugins/Runtime/MassGameplay/Source/MassSmartObjects/Private/MassSmartObjectHandler.cpp b/Engine/Plugins/Runtime/MassGameplay/Source/MassSmartObjects/Private/MassSmartObjectHandler.cpp
--- a/Engine/Plugins/Runtime/MassGameplay/Source/MassSmartObjects/Private/MassSmartObjectHandler.cpp
+++ b/Engine/Plugins/Runtime/MassGameplay/Source/MassSmartObjects/Private/MassSmartObjectHandler.cpp
@@ -18,7 +18,7 @@ struct FPayload
 
 void OnSlotInvalidated(const FSmartObjectClaimHandle& ClaimHandle, const ESmartObjectSlotState State, FPayload Payload)
 {
-	FDataFragment_SmartObjectUser& User = Payload.EntitySubsystem->GetComponentDataChecked<FDataFragment_SmartObjectUser>(Payload.Entity);
+	FDataFragment_SmartObjectUser& User = Payload.EntitySubsystem->GetFragmentDataChecked<FDataFragment_SmartObjectUser>(Payload.Entity);
 	if (USmartObjectSubsystem* SmartObjectSubsystem = Payload.SmartObjectSubsystem.Get())
 	{
 		SmartObjectSubsystem->Release(User.ClaimHandle);
@@ -43,7 +43,7 @@ FMassSmartObjectRequestID FMassSmartObjectHandler::FindCandidatesAsync(const FMa
 	FMassSmartObjectRequestResultFragment ResultFragment;
 
 	ExecutionContext.Defer().PushCommand(
-		FBuildEntityFromComponentInstances(ReservedEntity,
+		FBuildEntityFromFragmentInstances(ReservedEntity,
 			{
 				FStructView::Make(RequestFragment),
 				FStructView::Make(ResultFragment)
@@ -63,7 +63,7 @@ FMassSmartObjectRequestID FMassSmartObjectHandler::FindCandidatesAsync(const FMa
 	FMassSmartObjectRequestResultFragment ResultFragment;
 
 	ExecutionContext.Defer().PushCommand(
-		FBuildEntityFromComponentInstances(ReservedEntity,
+		FBuildEntityFromFragmentInstances(ReservedEntity,
 			{
 				FStructView::Make(RequestFragment),
 				FStructView::Make(ResultFragment)
@@ -86,7 +86,7 @@ FMassSmartObjectRequestResult FMassSmartObjectHandler::GetRequestResult(const FM
 		return {};
 	}
 
-	const FMassSmartObjectRequestResultFragment& RequestFragment = EntitySubsystem.GetComponentDataChecked<FMassSmartObjectRequestResultFragment>(RequestEntity);
+	const FMassSmartObjectRequestResultFragment& RequestFragment = EntitySubsystem.GetFragmentDataChecked<FMassSmartObjectRequestResultFragment>(RequestEntity);
 	return RequestFragment.Result;
 }
 
@@ -104,7 +104,7 @@ EMassSmartObjectClaimResult FMassSmartObjectHandler::ClaimCandidate(const FMassE
 	}
 
 	const FMassEntityHandle RequestEntity = static_cast<FMassEntityHandle>(RequestID);
-	const FMassSmartObjectRequestResultFragment& RequestFragment = EntitySubsystem.GetComponentDataChecked<FMassSmartObjectRequestResultFragment>(RequestEntity);
+	const FMassSmartObjectRequestResultFragment& RequestFragment = EntitySubsystem.GetFragmentDataChecked<FMassSmartObjectRequestResultFragment>(RequestEntity);
 	if (!ensureMsgf(RequestFragment.Result.bProcessed, TEXT("Trying to claim using a request that is not processed.")))
 	{
 		return EMassSmartObjectClaimResult::Failed_UnprocessedRequest;
@@ -156,7 +156,7 @@ bool FMassSmartObjectHandler::ClaimSmartObject(const FMassEntityHandle Entity, F
 		User.InteractionStatus = EMassSmartObjectInteractionStatus::Unset;
 		User.TargetLocation = SmartObjectSubsystem.GetSlotLocation(User.ClaimHandle).Get(FVector::ZeroVector);
 
-		// @todo SO: need to unregister callback if entity or component gets removed. Or at least support getting called with a no longer valid LWEntity.
+		// @todo SO: need to unregister callback if entity or fragment gets removed. Or at least support getting called with a no longer valid LWEntity.
 		UE::Mass::SmartObject::FPayload Payload;
 		Payload.Entity = Entity;
 		Payload.EntitySubsystem = &EntitySubsystem;
