Removing some "LWEntity" name remains.

diff --git a/Engine/Plugins/AI/MassAI/Source/MassAIDebug/Private/GameplayDebuggerCategory_Mass.cpp b/Engine/Plugins/AI/MassAI/Source/MassAIDebug/Private/GameplayDebuggerCategory_Mass.cpp
--- a/Engine/Plugins/AI/MassAI/Source/MassAIDebug/Private/GameplayDebuggerCategory_Mass.cpp
+++ b/Engine/Plugins/AI/MassAI/Source/MassAIDebug/Private/GameplayDebuggerCategory_Mass.cpp
@@ -135,11 +135,11 @@ void FGameplayDebuggerCategory_Mass::CollectData(APlayerController* OwnerPC, AAc
 
 	auto GetEntityFromActorFunc = [](const AActor& Actor, const UMassAgentComponent** OutMassAgentComponent = nullptr)
 	{
-		FMassEntityHandle LWEntity;
+		FMassEntityHandle EntityHandle;
 		UMassAgentComponent* AgentComp = Actor.FindComponentByClass<UMassAgentComponent>();
 		if (AgentComp)
 		{
-			LWEntity = AgentComp->GetEntityHandle();
+			EntityHandle = AgentComp->GetEntityHandle();
 			if (OutMassAgentComponent)
 			{
 				*OutMassAgentComponent = AgentComp;
@@ -150,10 +150,10 @@ void FGameplayDebuggerCategory_Mass::CollectData(APlayerController* OwnerPC, AAc
 			UMassActorSubsystem* ActorManager = UWorld::GetSubsystem<UMassActorSubsystem>(Actor.GetWorld());
 			if (ActorManager != nullptr)
 			{
-				LWEntity = ActorManager->GetEntityHandleFromActor(&Actor);
+				EntityHandle = ActorManager->GetEntityHandleFromActor(&Actor);
 			}
 		}
-		return LWEntity;
+		return EntityHandle;
 	};
 
 	if (CachedDebugActor != DebugActor)
diff --git a/Engine/Plugins/Runtime/MassEntity/Source/MassEntity/Public/MassEntitySubsystem.h b/Engine/Plugins/Runtime/MassEntity/Source/MassEntity/Public/MassEntitySubsystem.h
--- a/Engine/Plugins/Runtime/MassEntity/Source/MassEntity/Public/MassEntitySubsystem.h
+++ b/Engine/Plugins/Runtime/MassEntity/Source/MassEntity/Public/MassEntitySubsystem.h
@@ -114,13 +114,13 @@ public:
 	/**
 	 * Creates fully built entity ready to be used by the subsystem
 	 * @param Archetype you want this entity to be
-	 * @return FLWEntity id of the newly created entity */
+	 * @return FMassEntityHandle id of the newly created entity */
 	FMassEntityHandle CreateEntity(const FArchetypeHandle Archetype);
 
 	/**
 	 * Creates fully built entity ready to be used by the subsystem
 	 * @param FragmentInstanceList is the fragments to create the entity from and initialize values
-	 * @return FLWEntity id of the newly created entity */
+	 * @return FMassEntityHandle id of the newly created entity */
 	FMassEntityHandle CreateEntity(TConstArrayView<FInstancedStruct> FragmentInstanceList);
 
 	/** A version of CreateEntity that's creating a number of entities (Count) at one go
@@ -137,7 +137,7 @@ public:
 
 	/**
 	 * Reserves an entity in the subsystem, the entity is still not ready to be used by the subsystem, need to call BuildEntity()
-	 * @return FLWEntity id of the reserved entity */
+	 * @return FMassEntityHandle id of the reserved entity */
 	FMassEntityHandle ReserveEntity();
 
 	/**
diff --git a/Engine/Plugins/Runtime/MassEntity/Source/MassEntityTestSuite/Private/MassEntityTest.cpp b/Engine/Plugins/Runtime/MassEntity/Source/MassEntityTestSuite/Private/MassEntityTest.cpp
--- a/Engine/Plugins/Runtime/MassEntity/Source/MassEntityTestSuite/Private/MassEntityTest.cpp
+++ b/Engine/Plugins/Runtime/MassEntity/Source/MassEntityTestSuite/Private/MassEntityTest.cpp
@@ -161,7 +161,7 @@ struct FEntityTest_AddingRedundantFragment : FEntityTestBase
 	virtual bool InstantTest() override
 	{
 		CA_ASSUME(EntitySubsystem);
-		const FLWEntity Entity = EntitySubsystem->CreateEntity(FloatsArchetype);
+		const FMassEntityHandle Entity = EntitySubsystem->CreateEntity(FloatsArchetype);
 		EntitySubsystem->AddFragmentToEntity(Entity, FTestFragment_Float::StaticStruct());		
 		AITEST_EQUAL("Adding a fragment that a given entity\s archetype already has should do nothing", EntitySubsystem->GetArchetypeForEntity(Entity), FloatsArchetype);
 
diff --git a/Engine/Plugins/Runtime/MassGameplay/Source/MassActors/Public/MassAgentSubsystem.h b/Engine/Plugins/Runtime/MassGameplay/Source/MassActors/Public/MassAgentSubsystem.h
--- a/Engine/Plugins/Runtime/MassGameplay/Source/MassActors/Public/MassAgentSubsystem.h
+++ b/Engine/Plugins/Runtime/MassGameplay/Source/MassActors/Public/MassAgentSubsystem.h
@@ -45,7 +45,7 @@ public:
 	virtual void Initialize(FSubsystemCollectionBase& Collection) override;
 	// USubsystem END
 
-	/** Registers given AgentComp with the Mass Simulation, including creation of a LWEntity using 
+	/** Registers given AgentComp with the Mass Simulation, including creation of a FMassEntityHandle using 
 	 *  UMassAgentComponent.LWComponentList to determine the Archetype to use. */
 	FMassEntityTemplateID RegisterAgentComponent(UMassAgentComponent& AgentComp);
 
diff --git a/Engine/Plugins/Runtime/MassGameplay/Source/MassSmartObjects/Private/MassSmartObjectHandler.cpp b/Engine/Plugins/Runtime/MassGameplay/Source/MassSmartObjects/Private/MassSmartObjectHandler.cpp
--- a/Engine/Plugins/Runtime/MassGameplay/Source/MassSmartObjects/Private/MassSmartObjectHandler.cpp
+++ b/Engine/Plugins/Runtime/MassGameplay/Source/MassSmartObjects/Private/MassSmartObjectHandler.cpp
@@ -156,7 +156,7 @@ bool FMassSmartObjectHandler::ClaimSmartObject(const FMassEntityHandle Entity, F
 		User.InteractionStatus = EMassSmartObjectInteractionStatus::Unset;
 		User.TargetLocation = SmartObjectSubsystem.GetSlotLocation(User.ClaimHandle).Get(FVector::ZeroVector);
 
-		// @todo SO: need to unregister callback if entity or fragment gets removed. Or at least support getting called with a no longer valid LWEntity.
+		// @todo SO: need to unregister callback if entity or fragment gets removed. Or at least support getting called with a no longer valid FMassEntityHandle.
 		UE::Mass::SmartObject::FPayload Payload;
 		Payload.Entity = Entity;
 		Payload.EntitySubsystem = &EntitySubsystem;
diff --git a/Engine/Plugins/Runtime/MassGameplay/Source/MassSmartObjects/Public/MassSmartObjectTypes.h b/Engine/Plugins/Runtime/MassGameplay/Source/MassSmartObjects/Public/MassSmartObjectTypes.h
--- a/Engine/Plugins/Runtime/MassGameplay/Source/MassSmartObjects/Public/MassSmartObjectTypes.h
+++ b/Engine/Plugins/Runtime/MassGameplay/Source/MassSmartObjects/Public/MassSmartObjectTypes.h
@@ -14,7 +14,7 @@ namespace UE::Mass::Signals
 
 /**
  * Identifier associated to a request for smart object candidates. We use a 1:1 match
- * with an LWEntity since all requests are batched together using the EntitySubsystem.
+ * with an FMassEntityHandle since all requests are batched together using the EntitySubsystem.
  */
 USTRUCT()
 struct MASSSMARTOBJECTS_API FMassSmartObjectRequestID
