Reimplemented the fix to MassSettings' issue that was resulting in never utilizing the processor that was first in the global processors list.

diff --git a/Engine/Plugins/Runtime/MassEntity/Source/MassEntity/Private/MassEntitySettings.cpp b/Engine/Plugins/Runtime/MassEntity/Source/MassEntity/Private/MassEntitySettings.cpp
--- a/Engine/Plugins/Runtime/MassEntity/Source/MassEntity/Private/MassEntitySettings.cpp
+++ b/Engine/Plugins/Runtime/MassEntity/Source/MassEntity/Private/MassEntitySettings.cpp
@@ -73,7 +73,7 @@ void UMassSettings::BuildProcessorList()
 	TArray<UClass*> SubClassess;
 	GetDerivedClasses(UMassProcessor::StaticClass(), SubClassess);
 
-	for (int i = SubClassess.Num() - 1; i; --i)
+	for (int i = SubClassess.Num() - 1; i >= 0; --i)
 	{
 		if (SubClassess[i]->HasAnyClassFlags(CLASS_Abstract))
 		{
