Fixed FMassEntityQuery::DoesArchetypeMatchRequirements not checking tags nor chunk fragments. Also added a function that can print out the archetype's composition bits that don't match.

diff --git a/Engine/Plugins/Runtime/MassEntity/Source/MassEntity/Private/MassEntityQuery.cpp b/Engine/Plugins/Runtime/MassEntity/Source/MassEntity/Private/MassEntityQuery.cpp
--- a/Engine/Plugins/Runtime/MassEntity/Source/MassEntity/Private/MassEntityQuery.cpp
+++ b/Engine/Plugins/Runtime/MassEntity/Source/MassEntity/Private/MassEntityQuery.cpp
@@ -7,6 +7,7 @@
 #include "MassCommandBuffer.h"
 #include "VisualLogger/VisualLogger.h"
 #include "Async/ParallelFor.h"
+#include "Containers/UnrealString.h"
 
 //////////////////////////////////////////////////////////////////////
 // FMassEntityQuery
@@ -98,24 +99,17 @@ bool FMassEntityQuery::DoesArchetypeMatchRequirements(const FArchetypeHandle& Ar
 	check(ArchetypeHandle.IsValid());
 	const FMassArchetypeData* Archetype = ArchetypeHandle.DataPtr.Get();
 	CA_ASSUME(Archetype);
-
-	for (const FMassFragmentRequirement& Requirement : GetRequirements())
-	{
-		check(Requirement.StructType);
-		if (Requirement.IsOptional())
-		{
-			// at this stage we don't care if the optional is present
-			continue;
-		}
-
-		const bool bNeedsFragment = (Requirement.Presence == EMassFragmentPresence::All);
-		const bool bHasFragment = Archetype->HasFragmentType(Requirement.StructType);
-		if (bNeedsFragment != bHasFragment)
-		{
-			return false;
-		}
-	}
-	return true;
+	
+	const FMassCompositionDescriptor& ArchetypeComposition = Archetype->GetCompositionDescriptor();
+
+	return ArchetypeComposition.Fragments.HasAll(RequiredAllFragments)
+		&& (RequiredAnyFragments.IsEmpty() || ArchetypeComposition.Fragments.HasAny(RequiredAnyFragments))
+		&& ArchetypeComposition.Fragments.HasNone(RequiredNoneFragments)
+		&& ArchetypeComposition.Tags.HasAll(RequiredAllTags)
+		&& (RequiredAnyTags.IsEmpty() || ArchetypeComposition.Tags.HasAny(RequiredAnyTags))
+		&& ArchetypeComposition.Tags.HasNone(RequiredNoneTags)
+		&& ArchetypeComposition.ChunkFragments.HasAll(RequiredAllChunkFragments)
+		&& ArchetypeComposition.ChunkFragments.HasNone(RequiredNoneChunkFragments);
 }
 
 void FMassEntityQuery::ForEachEntityChunk(const FArchetypeChunkCollection& Chunks, UMassEntitySubsystem& EntitySubsystem, FMassExecutionContext& ExecutionContext, const FMassExecuteFunction& ExecuteFunction)
@@ -139,7 +133,8 @@ void FMassEntityQuery::ForEachEntityChunk(UMassEntitySubsystem& EntitySubsystem,
 		if (DoesArchetypeMatchRequirements(ExecutionContext.GetChunkCollection().GetArchetype()) == false)
 		{
 			// mz@todo add a unit test for this message
-			UE_VLOG_UELOG(&EntitySubsystem, LogMass, Error, TEXT("Attempted to execute FMassEntityQuery with an incompatible Archetype"));
+			UE_VLOG_UELOG(&EntitySubsystem, LogMass, Error, TEXT("Attempted to execute FMassEntityQuery with an incompatible Archetype: %s")
+				, *DebugGetArchetypeCompatibilityDescription(ExecutionContext.GetChunkCollection().GetArchetype()));
 			return;
 		}
 		ExecutionContext.SetRequirements(Requirements, ChunkRequirements);
@@ -200,7 +195,8 @@ void FMassEntityQuery::ParallelForEachEntityChunk(UMassEntitySubsystem& EntitySu
 		if (DoesArchetypeMatchRequirements(ArchetypeHandle) == false)
 		{
 			// mz@todo add a unit test for this message
-			UE_VLOG_UELOG(&EntitySubsystem, LogMass, Error, TEXT("Attempted to execute FMassEntityQuery with an incompatible Archetype"));
+			UE_VLOG_UELOG(&EntitySubsystem, LogMass, Error, TEXT("Attempted to execute FMassEntityQuery with an incompatible Archetype: %s")
+				, *DebugGetArchetypeCompatibilityDescription(ExecutionContext.GetChunkCollection().GetArchetype()));
 			return;
 		}
 
@@ -295,6 +291,79 @@ FString FMassEntityQuery::DebugGetDescription() const
 #endif
 }
 
+FString FMassEntityQuery::DebugGetArchetypeCompatibilityDescription(const FArchetypeHandle& ArchetypeHandle) const
+{
+	if (ArchetypeHandle.IsValid() == false)
+	{
+		return TEXT("Invalid");
+	}
+	
+	const FMassArchetypeData& Archetype = *ArchetypeHandle.DataPtr;
+	FStringOutputDevice OutDescription;
+#if WITH_MASSENTITY_DEBUG	
+	const FMassCompositionDescriptor& ArchetypeComposition = Archetype.GetCompositionDescriptor();
+
+	if (ArchetypeComposition.Fragments.HasAll(RequiredAllFragments) == false)
+	{
+		// missing one of the strictly required fragments
+		OutDescription += TEXT("\nMissing required fragments: ");
+		(RequiredAllFragments - ArchetypeComposition.Fragments).DebugGetStringDesc(OutDescription);
+	}
+
+	if (RequiredAnyFragments.IsEmpty() == false && ArchetypeComposition.Fragments.HasAny(RequiredAnyFragments) == false)
+	{
+		// missing all of the "any" fragments
+		OutDescription += TEXT("\nMissing all \'any\' fragments: ");
+		RequiredAnyFragments.DebugGetStringDesc(OutDescription);
+	}
+
+	if (ArchetypeComposition.Fragments.HasNone(RequiredNoneFragments) == false)
+	{
+		// has some of the fragments required absent
+		OutDescription += TEXT("\nHas fragments required absent: ");
+		RequiredNoneFragments.DebugGetStringDesc(OutDescription);
+	}
+
+	if (ArchetypeComposition.Tags.HasAll(RequiredAllTags) == false)
+	{
+		// missing one of the strictly required tags
+		OutDescription += TEXT("\nMissing required tags: ");
+		(RequiredAllTags - ArchetypeComposition.Tags).DebugGetStringDesc(OutDescription);
+	}
+
+	if (RequiredAnyTags.IsEmpty() == false && ArchetypeComposition.Tags.HasAny(RequiredAnyTags) == false)
+	{
+		// missing all of the "any" tags
+		OutDescription += TEXT("\nMissing all \'any\' tags: ");
+		RequiredAnyTags.DebugGetStringDesc(OutDescription);
+	}
+
+	if (ArchetypeComposition.Tags.HasNone(RequiredNoneTags) == false)
+	{
+		// has some of the tags required absent
+		OutDescription += TEXT("\nHas tags required absent: ");
+		RequiredNoneTags.DebugGetStringDesc(OutDescription);
+	}
+	
+	if (ArchetypeComposition.ChunkFragments.HasAll(RequiredAllChunkFragments) == false)
+    {
+    	// missing one of the strictly required chunk fragments
+    	OutDescription += TEXT("\nMissing required chunk fragments: ");
+		(RequiredAllChunkFragments - ArchetypeComposition.ChunkFragments).DebugGetStringDesc(OutDescription);
+    }
+
+    if (ArchetypeComposition.ChunkFragments.HasNone(RequiredNoneChunkFragments) == false)
+	{
+		// has some of the chunk fragments required absent
+		OutDescription += TEXT("\nHas chunk fragments required absent: ");
+		RequiredNoneChunkFragments.DebugGetStringDesc(OutDescription);
+	}
+#endif // WITH_MASSENTITY_DEBUG
+	
+	return OutDescription.Len() > 0 ? static_cast<FString>(OutDescription) : TEXT("Match");
+}
+
+
 //////////////////////////////////////////////////////////////////////
 // FMassFragmentRequirement
 
diff --git a/Engine/Plugins/Runtime/MassEntity/Source/MassEntity/Public/MassEntityQuery.h b/Engine/Plugins/Runtime/MassEntity/Source/MassEntity/Public/MassEntityQuery.h
--- a/Engine/Plugins/Runtime/MassEntity/Source/MassEntity/Public/MassEntityQuery.h
+++ b/Engine/Plugins/Runtime/MassEntity/Source/MassEntity/Public/MassEntityQuery.h
@@ -271,6 +271,7 @@ public:
 	bool CheckValidity() const;
 
 	FString DebugGetDescription() const;
+	FString DebugGetArchetypeCompatibilityDescription(const FArchetypeHandle& ArchetypeHandle) const;
 
 	TConstArrayView<FMassFragmentRequirement> GetRequirements() const { return Requirements; }
 	const FMassFragmentBitSet& GetRequiredAllFragments() const { return RequiredAllFragments; }
