Some more "Pipe" name removal

diff --git a/Engine/Plugins/Runtime/MassEntity/Source/MassEntity/Private/MassEntityDebug.cpp b/Engine/Plugins/Runtime/MassEntity/Source/MassEntity/Private/MassEntityDebug.cpp
--- a/Engine/Plugins/Runtime/MassEntity/Source/MassEntity/Private/MassEntityDebug.cpp
+++ b/Engine/Plugins/Runtime/MassEntity/Source/MassEntity/Private/MassEntityDebug.cpp
@@ -50,7 +50,7 @@ void DebugOutputDescription(TConstArrayView<UMassProcessor*> Processors, FOutput
 
 #if WITH_MASSENTITY_DEBUG && WITH_MASSENTITY_DEBUG
 FAutoConsoleCommandWithWorldArgsAndOutputDevice PrintEntityFragmentsCmd(
-	TEXT("pipe.PrintEntityFragments"),
+	TEXT("mass.PrintEntityFragments"),
 	TEXT("Prints all fragment types and values (uproperties) for the specified Entity index"),
 	FConsoleCommandWithWorldArgsAndOutputDeviceDelegate::CreateLambda(
 		[](const TArray<FString>& Params, UWorld* World, FOutputDevice& Ar)
@@ -70,13 +70,13 @@ FAutoConsoleCommandWithWorldArgsAndOutputDevice PrintEntityFragmentsCmd(
 			}
 			else
 			{
-				Ar.Logf(ELogVerbosity::Error, TEXT("Failed to find PipeEntitySubsystem for world %s"), *GetPathNameSafe(World));
+				Ar.Logf(ELogVerbosity::Error, TEXT("Failed to find MassEntitySubsystem for world %s"), *GetPathNameSafe(World));
 			}
 		})
 );
 
 FAutoConsoleCommandWithWorldArgsAndOutputDevice LogArchetypesCmd(
-	TEXT("pipe.LogArchetypes"),
+	TEXT("mass.LogArchetypes"),
 	TEXT("Dumps description of archetypes to log. Optional parameter controls whether to include or exclude non-occupied archetypes. Defaults to 'include'."),
 	FConsoleCommandWithWorldArgsAndOutputDeviceDelegate::CreateLambda([](const TArray<FString>& Params, UWorld*, FOutputDevice& Ar)
 	{
@@ -106,7 +106,7 @@ FAutoConsoleCommandWithWorldArgsAndOutputDevice LogArchetypesCmd(
 			}
 			else
 			{
-				Ar.Logf(ELogVerbosity::Error, TEXT("Failed to find PipeEntitySubsystem for world: %s (%s - %s)"),
+				Ar.Logf(ELogVerbosity::Error, TEXT("Failed to find MassEntitySubsystem for world: %s (%s - %s)"),
 					*GetPathNameSafe(World),
 					LexToString(World->WorldType),
 					*ToString(World->GetNetMode()));
@@ -117,7 +117,7 @@ FAutoConsoleCommandWithWorldArgsAndOutputDevice LogArchetypesCmd(
 
 // @todo these console commands will be reparented to "massentities" domain once we rename and shuffle the modules around 
 FAutoConsoleCommandWithWorld RecacheQueries(
-	TEXT("pipe.RecacheQueries"),
+	TEXT("mass.RecacheQueries"),
 	TEXT("Forces EntityQueries to recache their valid archetypes"),
 	FConsoleCommandWithWorldDelegate::CreateLambda([](UWorld* InWorld)
 		{
@@ -130,7 +130,7 @@ FAutoConsoleCommandWithWorld RecacheQueries(
 	));
 
 FAutoConsoleCommandWithWorldArgsAndOutputDevice LogFragmentSizes(
-	TEXT("pipe.LogFragmentSizes"),
+	TEXT("mass.LogFragmentSizes"),
 	TEXT("Logs all the fragment types being used along with their sizes."),
 	FConsoleCommandWithWorldArgsAndOutputDeviceDelegate::CreateLambda([](const TArray<FString>& Params, UWorld* World, FOutputDevice& Ar)
 		{
diff --git a/Engine/Plugins/Runtime/MassEntity/Source/MassEntity/Private/MassEntityQuery.cpp b/Engine/Plugins/Runtime/MassEntity/Source/MassEntity/Private/MassEntityQuery.cpp
--- a/Engine/Plugins/Runtime/MassEntity/Source/MassEntity/Private/MassEntityQuery.cpp
+++ b/Engine/Plugins/Runtime/MassEntity/Source/MassEntity/Private/MassEntityQuery.cpp
@@ -69,7 +69,7 @@ void FMassEntityQuery::CacheArchetypes(UMassEntitySubsystem& InEntitySubsystem)
 			InEntitySubsystem.GetValidArchetypes(*this, ValidArchetypes);
 			ArchetypeDataVersion = InEntitySubsystem.GetArchetypeDataVersion();
 
-			TRACE_CPUPROFILER_EVENT_SCOPE_STR("Pipe RequirementsBinding")
+			TRACE_CPUPROFILER_EVENT_SCOPE_STR("Mass RequirementsBinding")
 			const TConstArrayView<FMassFragmentRequirement> LocalRequirements = GetRequirements();
 			ArchetypeFragmentMapping.Reset(ValidArchetypes.Num());
 			ArchetypeFragmentMapping.AddDefaulted(ValidArchetypes.Num());
diff --git a/Engine/Plugins/Runtime/MassEntity/Source/MassEntity/Private/MassExecutor.cpp b/Engine/Plugins/Runtime/MassEntity/Source/MassEntity/Private/MassExecutor.cpp
--- a/Engine/Plugins/Runtime/MassEntity/Source/MassEntity/Private/MassExecutor.cpp
+++ b/Engine/Plugins/Runtime/MassEntity/Source/MassEntity/Private/MassExecutor.cpp
@@ -17,7 +17,7 @@ void Run(FMassRuntimePipeline& RuntimePipeline, FMassProcessingContext& Processi
 		return;
 	}
 
-	TRACE_CPUPROFILER_EVENT_SCOPE_STR("PipeExecutor Run Pipeline")
+	TRACE_CPUPROFILER_EVENT_SCOPE_STR("MassExecutor Run Pipeline")
 	RunProcessorsView(RuntimePipeline.Processors, ProcessingContext);
 }
 
@@ -31,7 +31,7 @@ void RunSparse(FMassRuntimePipeline& RuntimePipeline, FMassProcessingContext& Pr
 		return;
 	}
 
-	TRACE_CPUPROFILER_EVENT_SCOPE_STR("PipeExecutor RunSparseEntities");
+	TRACE_CPUPROFILER_EVENT_SCOPE_STR("MassExecutor RunSparseEntities");
 
 	const FArchetypeChunkCollection ChunkCollection(Archetype, Entities);
 	RunProcessorsView(RuntimePipeline.Processors, ProcessingContext, &ChunkCollection);
@@ -47,7 +47,7 @@ void RunSparse(FMassRuntimePipeline& RuntimePipeline, FMassProcessingContext& Pr
 		return;
 	}
 
-	TRACE_CPUPROFILER_EVENT_SCOPE_STR("PipeExecutor RunSparse");
+	TRACE_CPUPROFILER_EVENT_SCOPE_STR("MassExecutor RunSparse");
 
 	RunProcessorsView(RuntimePipeline.Processors, ProcessingContext, &ChunkCollection);
 }
@@ -59,7 +59,7 @@ void Run(UMassProcessor& Processor, FMassProcessingContext& ProcessingContext)
 		return;
 	}
 
-	TRACE_CPUPROFILER_EVENT_SCOPE_STR("PipeExecutor Run")
+	TRACE_CPUPROFILER_EVENT_SCOPE_STR("MassExecutor Run")
 
 	UMassProcessor* ProcPtr = &Processor;
 	RunProcessorsView(MakeArrayView(&ProcPtr, 1), ProcessingContext);
@@ -82,14 +82,14 @@ void RunProcessorsView(TArrayView<UMassProcessor*> Processors, FMassProcessingCo
 	}
 #endif // WITH_MASSENTITY_DEBUG
 
-	TRACE_CPUPROFILER_EVENT_SCOPE_STR("PipeExecutor RunProcessorsView")
+	TRACE_CPUPROFILER_EVENT_SCOPE_STR("MassExecutor RunProcessorsView")
 
 	FMassExecutionContext ExecutionContext(ProcessingContext.DeltaSeconds);
 	if (ChunkCollection)
 	{
 		ExecutionContext.SetChunkCollection(*ChunkCollection);
 	}
-	// manually creating a new command buffer to let the default one still be used by code unaware of pipe processing
+	// manually creating a new command buffer to let the default one still be used by code unaware of mass processing
 	ExecutionContext.SetDeferredCommandBuffer(MakeShareable(new FMassCommandBuffer()));
 	ExecutionContext.SetFlushDeferredCommands(false);
 	ExecutionContext.SetAuxData(ProcessingContext.AuxData);
@@ -115,9 +115,9 @@ void RunProcessorsView(TArrayView<UMassProcessor*> Processors, FMassProcessingCo
 	}
 }
 
-struct FPipeExecutorDoneTask
+struct FMassExecutorDoneTask
 {
-	FPipeExecutorDoneTask(const FMassExecutionContext& InExecutionContext, UMassEntitySubsystem& InEntitySubsystem, TFunction<void()> InOnDoneNotification, const FString& InDebugName)
+	FMassExecutorDoneTask(const FMassExecutionContext& InExecutionContext, UMassEntitySubsystem& InEntitySubsystem, TFunction<void()> InOnDoneNotification, const FString& InDebugName)
 		: ExecutionContext(InExecutionContext)
 		, EntitySubsystem(InEntitySubsystem)
 		, OnDoneNotification(InOnDoneNotification)
@@ -126,7 +126,7 @@ struct FPipeExecutorDoneTask
 	}
 	static TStatId GetStatId()
 	{
-		RETURN_QUICK_DECLARE_CYCLE_STAT(FPipeExecutorDoneTask, STATGROUP_TaskGraphTasks);
+		RETURN_QUICK_DECLARE_CYCLE_STAT(FMassExecutorDoneTask, STATGROUP_TaskGraphTasks);
 	}
 
 	static ENamedThreads::Type GetDesiredThread() { return ENamedThreads::GameThread; }
@@ -140,7 +140,7 @@ struct FPipeExecutorDoneTask
 			ExecutionContext.Defer().MoveAppend(EntitySubsystem.Defer());
 		}
 
-		UE_LOG(LogMass, Log, TEXT("PipeExecutor %s tasks DONE"), *DebugName);
+		UE_LOG(LogMass, Log, TEXT("MassExecutor %s tasks DONE"), *DebugName);
 		ExecutionContext.SetFlushDeferredCommands(true);
 		ExecutionContext.FlushDeferred(EntitySubsystem);
 
@@ -163,7 +163,7 @@ FGraphEventRef TriggerParallelTasks(UMassProcessor& Processor, FMassProcessingCo
 		return FGraphEventRef();
 	}
 
-	TRACE_CPUPROFILER_EVENT_SCOPE_STR("PipeExecutor RunParallel")
+	TRACE_CPUPROFILER_EVENT_SCOPE_STR("MassExecutor RunParallel")
 
 	// not going through UMassEntitySubsystem::CreateExecutionContext on purpose - we do need a separate command buffer
 	FMassExecutionContext ExecutionContext(ProcessingContext.DeltaSeconds);
@@ -180,7 +180,7 @@ FGraphEventRef TriggerParallelTasks(UMassProcessor& Processor, FMassProcessingCo
 	if (CompletionEvent.IsValid())
 	{
 		const FGraphEventArray Prerequisites = { CompletionEvent };
-		CompletionEvent = TGraphTask<FPipeExecutorDoneTask>::CreateTask(&Prerequisites)
+		CompletionEvent = TGraphTask<FMassExecutorDoneTask>::CreateTask(&Prerequisites)
 			.ConstructAndDispatchWhenReady(ExecutionContext, *ProcessingContext.EntitySubsystem, OnDoneNotification, Processor.GetName());
 	}
 
diff --git a/Engine/Plugins/Runtime/MassEntity/Source/MassEntity/Private/MassProcessingPhase.cpp b/Engine/Plugins/Runtime/MassEntity/Source/MassEntity/Private/MassProcessingPhase.cpp
--- a/Engine/Plugins/Runtime/MassEntity/Source/MassEntity/Private/MassProcessingPhase.cpp
+++ b/Engine/Plugins/Runtime/MassEntity/Source/MassEntity/Private/MassProcessingPhase.cpp
@@ -15,12 +15,12 @@
 
 DECLARE_CYCLE_STAT(TEXT("Mass Phase Done"), STAT_MassPhaseDone, STATGROUP_TaskGraphTasks);
 
-namespace FPipeTweakables
+namespace FMassTweakables
 {
 	bool bFullyParallel = false;
 
 	FAutoConsoleVariableRef CVars[] = {
-		{TEXT("pipe.FullyParallel"), bFullyParallel, TEXT("Enables pipe processing distribution to all available thread (via the task graph)")},
+		{TEXT("mass.FullyParallel"), bFullyParallel, TEXT("Enables mass processing distribution to all available thread (via the task graph)")},
 	};
 }
 
@@ -52,7 +52,7 @@ void FMassProcessingPhase::ExecuteTick(float DeltaTime, ELevelTick TickType, ENa
 
 	if (bRunInParallelMode)
 	{
-		bIsDuringPipeProcessing = true;
+		bIsDuringMassProcessing = true;
 		const FGraphEventRef PipelineCompletionEvent = UE::Mass::Executor::TriggerParallelTasks(*PhaseProcessor, Context, [this, DeltaTime]()
 			{
 				OnParallelExecutionDone(DeltaTime);
@@ -69,7 +69,7 @@ void FMassProcessingPhase::ExecuteTick(float DeltaTime, ELevelTick TickType, ENa
 	}
 	else
 	{
-		TGuardValue<bool> PipeRunningGuard(bIsDuringPipeProcessing, true);
+		TGuardValue<bool> MassRunningGuard(bIsDuringMassProcessing, true);
 		UE::Mass::Executor::Run(*PhaseProcessor, Context);
 
 		OnPhaseEnd.Broadcast(DeltaTime);
@@ -79,7 +79,7 @@ void FMassProcessingPhase::ExecuteTick(float DeltaTime, ELevelTick TickType, ENa
 
 void FMassProcessingPhase::OnParallelExecutionDone(const float DeltaTime)
 {
-	bIsDuringPipeProcessing = false;
+	bIsDuringMassProcessing = false;
 	OnPhaseEnd.Broadcast(DeltaTime);
 	check(Manager);
 	Manager->OnPhaseEnd(*this);
@@ -87,12 +87,12 @@ void FMassProcessingPhase::OnParallelExecutionDone(const float DeltaTime)
 
 FString FMassProcessingPhase::DiagnosticMessage()
 {
-	return (Manager ? Manager->GetFullName() : TEXT("NULL-PipeProcessingPhaseManager")) + TEXT("[ProcessorTick]");
+	return (Manager ? Manager->GetFullName() : TEXT("NULL-MassProcessingPhaseManager")) + TEXT("[ProcessorTick]");
 }
 
 FName FMassProcessingPhase::DiagnosticContext(bool bDetailed)
 {
-	return Manager ? Manager->GetClass()->GetFName() : TEXT("NULL-PipeProcessingPhaseManager");
+	return Manager ? Manager->GetClass()->GetFName() : TEXT("NULL-MassProcessingPhaseManager");
 }
 
 //----------------------------------------------------------------------//
@@ -107,9 +107,9 @@ void UMassProcessingPhaseManager::PostInitProperties()
 		UWorld* World = GetWorld();
 		if (World && World->IsGameWorld() == false)
 		{
-			UMassSettings* PipeSettings = GetMutableDefault<UMassSettings>();
-			check(PipeSettings);
-			PipeSettingsChangeHandle = PipeSettings->GetOnSettingsChange().AddUObject(this, &UMassProcessingPhaseManager::OnPipeSettingsChange);
+			UMassSettings* MassSettings = GetMutableDefault<UMassSettings>();
+			check(MassSettings);
+			MassSettingsChangeHandle = MassSettings->GetOnSettingsChange().AddUObject(this, &UMassProcessingPhaseManager::OnMassSettingsChange);
 		}
 #endif // WITH_EDITOR
 	}
@@ -122,25 +122,25 @@ void UMassProcessingPhaseManager::BeginDestroy()
 	Super::BeginDestroy();
 
 #if WITH_EDITOR
-	if (UMassSettings* PipeSettings = GetMutableDefault<UMassSettings>())
+	if (UMassSettings* MassSettings = GetMutableDefault<UMassSettings>())
 	{
-		PipeSettings->GetOnSettingsChange().Remove(PipeSettingsChangeHandle);
+		MassSettings->GetOnSettingsChange().Remove(MassSettingsChangeHandle);
 	}
 #endif // WITH_EDITOR
 }
 
 void UMassProcessingPhaseManager::InitializePhases(UObject& InProcessorOwner)
 {
-	const FMassProcessingPhaseConfig* ProcessingPhasesConfig = GET_PIPE_CONFIG_VALUE(GetProcessingPhasesConfig());
+	const FMassProcessingPhaseConfig* ProcessingPhasesConfig = GET_MASS_CONFIG_VALUE(GetProcessingPhasesConfig());
 
 	FString DependencyGraphFileName;
 
 #if WITH_EDITOR
 	const UWorld* World = InProcessorOwner.GetWorld();
-	const UMassSettings* PipeSettings = GetMutableDefault<UMassSettings>();
-	if (World != nullptr && PipeSettings != nullptr && !PipeSettings->DumpDependencyGraphFileName.IsEmpty())
+	const UMassSettings* MassSettings = GetMutableDefault<UMassSettings>();
+	if (World != nullptr && MassSettings != nullptr && !MassSettings->DumpDependencyGraphFileName.IsEmpty())
 	{
-		DependencyGraphFileName = FString::Printf(TEXT("%s_%s"), *PipeSettings->DumpDependencyGraphFileName,*ToString(World->GetNetMode()));
+		DependencyGraphFileName = FString::Printf(TEXT("%s_%s"), *MassSettings->DumpDependencyGraphFileName,*ToString(World->GetNetMode()));
 	}
 #endif // WITH_EDITOR
 
@@ -201,7 +201,7 @@ void UMassProcessingPhaseManager::EnableTickFunctions(const UWorld& World)
 		Phase.RegisterTickFunction(World.PersistentLevel);
 		Phase.SetTickFunctionEnable(true);
 	}
-	UE_VLOG_UELOG(this, LogMass, Log, TEXT("UPipeTickManager %s.%s has been started")
+	UE_VLOG_UELOG(this, LogMass, Log, TEXT("UMassProcessingPhaseManager %s.%s has been started")
 		, *GetNameSafe(GetOuter()), *GetName());
 }
 
@@ -214,7 +214,7 @@ void UMassProcessingPhaseManager::Stop()
 		Phase.SetTickFunctionEnable(false);
 	}
 
-	UE_VLOG_UELOG(this, LogMass, Log, TEXT("UPipeTickManager %s.%s has been stopped")
+	UE_VLOG_UELOG(this, LogMass, Log, TEXT("UMassProcessingPhaseManager %s.%s has been stopped")
 		, *GetNameSafe(GetOuter()), *GetName());
 }
 
@@ -255,9 +255,9 @@ void UMassProcessingPhaseManager::OnPhaseEnd(FMassProcessingPhase& Phase)
 	CurrentPhase = EMassProcessingPhase::MAX;
 
 	// switch between parallel and single-thread versions only after a given batch of processing has been wrapped up	
-	if (Phase.IsConfiguredForParallelMode() != FPipeTweakables::bFullyParallel)
+	if (Phase.IsConfiguredForParallelMode() != FMassTweakables::bFullyParallel)
 	{
-		if (FPipeTweakables::bFullyParallel)
+		if (FMassTweakables::bFullyParallel)
 		{
 			Phase.ConfigureForParallelMode();
 		}
@@ -269,7 +269,7 @@ void UMassProcessingPhaseManager::OnPhaseEnd(FMassProcessingPhase& Phase)
 }
 
 #if WITH_EDITOR
-void UMassProcessingPhaseManager::OnPipeSettingsChange(const FPropertyChangedEvent& PropertyChangedEvent)
+void UMassProcessingPhaseManager::OnMassSettingsChange(const FPropertyChangedEvent& PropertyChangedEvent)
 {
 	check(GetOuter());
 	InitializePhases(*GetOuter());
diff --git a/Engine/Plugins/Runtime/MassEntity/Source/MassEntity/Private/MassProcessingTypes.cpp b/Engine/Plugins/Runtime/MassEntity/Source/MassEntity/Private/MassProcessingTypes.cpp
--- a/Engine/Plugins/Runtime/MassEntity/Source/MassEntity/Private/MassProcessingTypes.cpp
+++ b/Engine/Plugins/Runtime/MassEntity/Source/MassEntity/Private/MassProcessingTypes.cpp
@@ -51,10 +51,10 @@ void FMassRuntimePipeline::InitializeFromSchematics(TConstArrayView<TSoftObjectP
 	// @todo we'll sometimes end up with duplicated MassProcessors in the resulting array. We need to come up with a consistent policy for handling that 
 	for (const TSoftObjectPtr<UMassSchematic>& Schematic : Schematics)
 	{
-		UMassSchematic* Pipe = Schematic.LoadSynchronous();
-		if (Pipe)
+		UMassSchematic* SchematicInstance = Schematic.LoadSynchronous();
+		if (SchematicInstance)
 		{
-			AppendOrOverrideRuntimeProcessorCopies(Pipe->GetProcessors(), InOwner);
+			AppendOrOverrideRuntimeProcessorCopies(SchematicInstance->GetProcessors(), InOwner);
 		}
 		else
 		{
diff --git a/Engine/Plugins/Runtime/MassEntity/Source/MassEntity/Private/MassProcessor.cpp b/Engine/Plugins/Runtime/MassEntity/Source/MassEntity/Private/MassProcessor.cpp
--- a/Engine/Plugins/Runtime/MassEntity/Source/MassEntity/Private/MassProcessor.cpp
+++ b/Engine/Plugins/Runtime/MassEntity/Source/MassEntity/Private/MassProcessor.cpp
@@ -7,7 +7,7 @@
 #include "Engine/World.h"
 #include "MassCommandBuffer.h"
 
-DECLARE_CYCLE_STAT(TEXT("MassProcessor Group Completed"), Pipe_GroupCompletedTask, STATGROUP_TaskGraphTasks);
+DECLARE_CYCLE_STAT(TEXT("MassProcessor Group Completed"), Mass_GroupCompletedTask, STATGROUP_TaskGraphTasks);
 
 #define PARALLELIZED_TRAFFIC_HACK 1
 
@@ -23,7 +23,7 @@ namespace UE::MassTraffic
 namespace UE::Mass::Debug
 {
 	bool bLogProcessingGraph = false;
-	FAutoConsoleVariableRef CVarLogProcessingGraph(TEXT("pipe.LogProcessingGraph"), bLogProcessingGraph
+	FAutoConsoleVariableRef CVarLogProcessingGraph(TEXT("mass.LogProcessingGraph"), bLogProcessingGraph
 		, TEXT("When enabled will log task graph tasks created while dispatching processors to other threads, along with their dependencies"), ECVF_Cheat);
 }
 #endif // WITH_MASSENTITY_DEBUG
@@ -35,14 +35,14 @@ namespace UE::Mass::Debug
 #define PROCESSOR_LOG(...) 
 #endif // WITH_MASSENTITY_DEBUG
 
-namespace FPipeTweakables
+namespace FMassTweakables
 {
 	bool bParallelGroups = false;
 	float PostponedTaskWaitTimeWarningLevel = 0.002f;
 
 	FAutoConsoleVariableRef CVarsMassProcessor[] = {
-		{TEXT("pipe.ParallelGroups"), bParallelGroups, TEXT("Enables pipe processing groups distribution to all available threads (via the task graph)")},
-		{TEXT("pipe.PostponedTaskWaitTimeWarningLevel"), PostponedTaskWaitTimeWarningLevel, TEXT("if waiting for postponed task\'s dependencies exceeds this number an error will be logged")},
+		{TEXT("mass.ParallelGroups"), bParallelGroups, TEXT("Enables mass processing groups distribution to all available threads (via the task graph)")},
+		{TEXT("mass.PostponedTaskWaitTimeWarningLevel"), PostponedTaskWaitTimeWarningLevel, TEXT("if waiting for postponed task\'s dependencies exceeds this number an error will be logged")},
 	};
 }
 
@@ -77,7 +77,7 @@ public:
 
 		UMassEntitySubsystem::FScopedProcessing ProcessingScope = EntitySubsystem->NewProcessingScope();
 
-		TRACE_CPUPROFILER_EVENT_SCOPE_STR("Pipe Processor Task");
+		TRACE_CPUPROFILER_EVENT_SCOPE_STR("Mass Processor Task");
 		
 		if (bManageCommandBuffer)
 		{
@@ -185,7 +185,7 @@ void UMassProcessor::PostEditChangeProperty(FPropertyChangedEvent& PropertyChang
 
 	if (HasAnyFlags(RF_ClassDefaultObject))
 	{
-		// this is here to make sure all the changes to CDOs we do via PipeSettings gets serialized to the ini file
+		// this is here to make sure all the changes to CDOs we do via MassSettings gets serialized to the ini file
 		SaveConfig(CPF_Config, *GetDefaultConfigFilename());
 	}
 }
@@ -236,7 +236,7 @@ FGraphEventRef UMassCompositeProcessor::DispatchProcessorTasks(UMassEntitySubsys
 		else
 		{
 			Events.Add(FFunctionGraphTask::CreateAndDispatchWhenReady([=](){}
-				, GET_STATID(Pipe_GroupCompletedTask), &Prerequisites, ENamedThreads::AnyHiPriThreadHiPriTask));
+				, GET_STATID(Mass_GroupCompletedTask), &Prerequisites, ENamedThreads::AnyHiPriThreadHiPriTask));
 		}
 	}
 
@@ -268,7 +268,7 @@ FGraphEventRef UMassCompositeProcessor::DispatchProcessorTasks(UMassEntitySubsys
 #endif // WITH_MASSENTITY_DEBUG
 
 	FGraphEventRef CompletionEvent = FFunctionGraphTask::CreateAndDispatchWhenReady([this](){}
-		, GET_STATID(Pipe_GroupCompletedTask), &Events, ENamedThreads::AnyHiPriThreadHiPriTask);
+		, GET_STATID(Mass_GroupCompletedTask), &Events, ENamedThreads::AnyHiPriThreadHiPriTask);
 
 	return CompletionEvent;
 }
@@ -276,7 +276,7 @@ FGraphEventRef UMassCompositeProcessor::DispatchProcessorTasks(UMassEntitySubsys
 void UMassCompositeProcessor::Execute(UMassEntitySubsystem& EntitySubsystem, FMassExecutionContext& Context)
 {
 #if PARALLELIZED_TRAFFIC_HACK
-	if (FPipeTweakables::bParallelGroups == false
+	if (FMassTweakables::bParallelGroups == false
 		&& UE::MassTraffic::bParallelizeTraffic && GetProcessingPhase() == EMassProcessingPhase::PrePhysics)
 	{
 		static FName TrafficGroup(TEXT("Traffic"));
@@ -314,7 +314,7 @@ void UMassCompositeProcessor::Execute(UMassEntitySubsystem& EntitySubsystem, FMa
 	}
 #endif // PARALLELIZED_TRAFFIC_HACK
 
-	if (FPipeTweakables::bParallelGroups && bHasOffThreadSubGroups)
+	if (FMassTweakables::bParallelGroups && bHasOffThreadSubGroups)
 	{
 		CompletionStatus.Reset();
 		CompletionStatus.AddDefaulted(ChildPipeline.Processors.Num());
@@ -492,7 +492,7 @@ int32 UMassCompositeProcessor::Populate(TArray<FProcessorDependencySolver::FOrde
 	int32 Index = StartIndex;
 	TMap<FName, int32> NameToIndexMap;
 	bool bOffThreadGroupsFound = false;
-	const FMassProcessingPhaseConfig& PhaseConfig = GET_PIPE_CONFIG_VALUE(GetProcessingPhaseConfig(ProcessingPhase));
+	const FMassProcessingPhaseConfig& PhaseConfig = GET_MASS_CONFIG_VALUE(GetProcessingPhaseConfig(ProcessingPhase));
 
 	while (ProcessorsAndGroups.IsValidIndex(Index))
 	{
diff --git a/Engine/Plugins/Runtime/MassEntity/Source/MassEntity/Public/MassEntitySettings.h b/Engine/Plugins/Runtime/MassEntity/Source/MassEntity/Public/MassEntitySettings.h
--- a/Engine/Plugins/Runtime/MassEntity/Source/MassEntity/Public/MassEntitySettings.h
+++ b/Engine/Plugins/Runtime/MassEntity/Source/MassEntity/Public/MassEntitySettings.h
@@ -8,7 +8,7 @@
 #include "InstancedStruct.h"
 #include "MassEntitySettings.generated.h"
 
-#define GET_PIPE_CONFIG_VALUE(a) (GetMutableDefault<UMassSettings>()->a)
+#define GET_MASS_CONFIG_VALUE(a) (GetMutableDefault<UMassSettings>()->a)
 
 class UMassProcessingPhaseManager;
 struct FMassProcessingPhaseConfig;
@@ -20,13 +20,13 @@ struct FMassProcessingPhaseConfig
 {
 	GENERATED_BODY()
 
-	UPROPERTY(EditAnywhere, Category = Pipe, config)
+	UPROPERTY(EditAnywhere, Category = Mass, config)
 	FName PhaseName;
 
-	UPROPERTY(EditAnywhere, Category = Pipe, config, NoClear)
+	UPROPERTY(EditAnywhere, Category = Mass, config, NoClear)
 	TSubclassOf<UMassCompositeProcessor> PhaseGroupClass = UMassCompositeProcessor::StaticClass();
 
-	UPROPERTY(EditAnywhere, Category = Pipe, config, NoClear)
+	UPROPERTY(EditAnywhere, Category = Mass, config, NoClear)
 	TArray<FName> OffGameThreadGroupNames;
 
 	UPROPERTY(Transient)
@@ -38,7 +38,7 @@ struct FMassProcessingPhaseConfig
 	UPROPERTY(Transient)
 	UMassCompositeProcessor* PhaseProcessor = nullptr;
 
-	UPROPERTY(VisibleAnywhere, Category = Pipe, Transient)
+	UPROPERTY(VisibleAnywhere, Category = Mass, Transient)
 	FString Description;
 #endif //WITH_EDITORONLY_DATA
 };
@@ -82,15 +82,15 @@ public:
 	 * The dot file will be put in the project log folder.
 	 * To generate a svg out of that file, simply run dot executable with following parameters: -Tsvg -O filename.dot 
 	 */
-	UPROPERTY(EditDefaultsOnly, Category = Pipe, Transient)
+	UPROPERTY(EditDefaultsOnly, Category = Mass, Transient)
 	FString DumpDependencyGraphFileName;
 
 	/** Lets users configure processing phases including the composite processor class to be used as a container for the phases' processors. */
-	UPROPERTY(EditDefaultsOnly, Category = Pipe, config)
+	UPROPERTY(EditDefaultsOnly, Category = Mass, config)
 	FMassProcessingPhaseConfig ProcessingPhasesConfig[(uint8)EMassProcessingPhase::MAX];
 
 	/** This list contains all the processors available in the given binary (including plugins). The contents are sorted by display name.*/
-	UPROPERTY(VisibleAnywhere, Category = Pipe, Transient, Instanced)
+	UPROPERTY(VisibleAnywhere, Category = Mass, Transient, Instanced)
 	TArray<UMassProcessor*> ProcessorCDOs;
 
 #if WITH_EDITORONLY_DATA
diff --git a/Engine/Plugins/Runtime/MassEntity/Source/MassEntity/Public/MassProcessingPhase.h b/Engine/Plugins/Runtime/MassEntity/Source/MassEntity/Public/MassProcessingPhase.h
--- a/Engine/Plugins/Runtime/MassEntity/Source/MassEntity/Public/MassProcessingPhase.h
+++ b/Engine/Plugins/Runtime/MassEntity/Source/MassEntity/Public/MassProcessingPhase.h
@@ -39,12 +39,12 @@ protected:
 	void ConfigureForSingleThreadMode() { bRunInParallelMode = false; }
 
 public:
-	bool IsDuringPipeProcessing() const { return bIsDuringPipeProcessing; }
+	bool IsDuringMassProcessing() const { return bIsDuringMassProcessing; }
 
 protected:
 	friend UMassProcessingPhaseManager;
 
-	UPROPERTY(EditAnywhere, Category=Pipe)
+	UPROPERTY(EditAnywhere, Category=Mass)
 	UMassCompositeProcessor* PhaseProcessor = nullptr;
 
 	UPROPERTY()
@@ -56,7 +56,7 @@ protected:
 
 private:
 	bool bRunInParallelMode = false;
-	bool bIsDuringPipeProcessing = false;
+	bool bIsDuringMassProcessing = false;
 };
 
 // It is unsafe to copy FTickFunctions and any subclasses of FTickFunction should specify the type trait WithCopy = false
@@ -69,10 +69,10 @@ struct TStructOpsTypeTraits<FMassProcessingPhase> : public TStructOpsTypeTraitsB
 	};
 };
 
-/** PipeProcessingPhaseManager owns separate FMassProcessingPhase instances for every ETickingGroup. When activated
+/** MassProcessingPhaseManager owns separate FMassProcessingPhase instances for every ETickingGroup. When activated
  *  via Start function it registers and enables the FMassProcessingPhase instances which themselves are tick functions 
  *  that host FMassRuntimePipeline which they trigger as part of their Tick function. 
- *  PipeProcessingPhaseManager serves as an interface to said FMassProcessingPhase instances and allows initialization
+ *  MassProcessingPhaseManager serves as an interface to said FMassProcessingPhase instances and allows initialization
  *  with MassSchematics (via InitializePhases function) as well as registering arbitrary functions to be called 
  *  when a particular phase starts of ends (via GetOnPhaseStart and GetOnPhaseEnd functions). */
 UCLASS(Transient, HideCategories = (Tick))
@@ -88,7 +88,7 @@ public:
 	FMassProcessingPhase::FOnPhaseEvent& GetOnPhaseEnd(const EMassProcessingPhase Phase) { return ProcessingPhases[uint8(Phase)].OnPhaseEnd; }
 
 	/** 
-	 *  Populates hosted FMassProcessingPhase instances with Processors read from PipeSettings configuration.
+	 *  Populates hosted FMassProcessingPhase instances with Processors read from MassSettings configuration.
 	 *  Calling this function overrides previous configuration of Phases.
 	 */
 	void InitializePhases(UObject& InProcessorOwner);
@@ -108,18 +108,18 @@ public:
 	 *  Note that the function will return false while the OnPhaseStart or OnPhaseEnd are being broadcast,
 	 *  the value returned will be `true` only when the entity subsystem is actively engaged 
 	 */
-	bool IsDuringPipeProcessing() const { return CurrentPhase != EMassProcessingPhase::MAX && ProcessingPhases[int(CurrentPhase)].IsDuringPipeProcessing(); }
+	bool IsDuringMassProcessing() const { return CurrentPhase != EMassProcessingPhase::MAX && ProcessingPhases[int(CurrentPhase)].IsDuringMassProcessing(); }
 
 protected:
 	virtual void PostInitProperties() override;
 	virtual void BeginDestroy() override;
 	void EnableTickFunctions(const UWorld& World);
 
-	/** Creates phase processors instances for each declared phase name, based on PipeSettings */
+	/** Creates phase processors instances for each declared phase name, based on MassSettings */
 	virtual void CreatePhases();
 
 #if WITH_EDITOR
-	virtual void OnPipeSettingsChange(const FPropertyChangedEvent& PropertyChangedEvent);
+	virtual void OnMassSettingsChange(const FPropertyChangedEvent& PropertyChangedEvent);
 #endif // WITH_EDITOR
 
 	friend FMassProcessingPhase;
@@ -137,7 +137,7 @@ protected:
 	void OnPhaseEnd(FMassProcessingPhase& Phase);
 
 protected:	
-	UPROPERTY(VisibleAnywhere, Category=Pipe)
+	UPROPERTY(VisibleAnywhere, Category=Mass)
 	FMassProcessingPhase ProcessingPhases[(uint8)EMassProcessingPhase::MAX];
 
 	UPROPERTY()
@@ -146,6 +146,6 @@ protected:
 	EMassProcessingPhase CurrentPhase = EMassProcessingPhase::MAX;
 
 #if WITH_EDITOR
-	FDelegateHandle PipeSettingsChangeHandle;
+	FDelegateHandle MassSettingsChangeHandle;
 #endif // WITH_EDITOR
 };
diff --git a/Engine/Plugins/Runtime/MassEntity/Source/MassEntity/Public/MassProcessor.h b/Engine/Plugins/Runtime/MassEntity/Source/MassEntity/Public/MassProcessor.h
--- a/Engine/Plugins/Runtime/MassEntity/Source/MassEntity/Public/MassProcessor.h
+++ b/Engine/Plugins/Runtime/MassEntity/Source/MassEntity/Public/MassProcessor.h
@@ -76,7 +76,7 @@ public:
 #endif // WITH_EDITORONLY_DATA
 
 	/** Sets bAutoRegisterWithProcessingPhases. Setting it to true will result in this processor class being always 
-	 * instantiated to be automatically evaluated every frame. @see UPipeProcessingPhaseManager
+	 * instantiated to be automatically evaluated every frame. @see UMassProcessingPhaseManager
 	 * Note that calling this function is only valid on CDOs. Calling it on a regular instance will fail an ensure and 
 	 * have no other effect, i.e. CDO's value won't change */
 	void SetShouldAutoRegisterWithGlobalList(const bool bAutoRegister);
@@ -185,7 +185,7 @@ protected:
 	UMassCompositeProcessor* FindOrAddGroupProcessor(FName RequestedGroupName, FString* OutRemainingGroupName = nullptr);
 
 protected:
-	UPROPERTY(VisibleAnywhere, Category=Pipe)
+	UPROPERTY(VisibleAnywhere, Category=Mass)
 	FMassRuntimePipeline ChildPipeline;
 
 	/** Group name that will be used when resolving processor dependencies and grouping */
diff --git a/Engine/Plugins/Runtime/MassEntity/Source/MassEntity/Public/MassSchematic.h b/Engine/Plugins/Runtime/MassEntity/Source/MassEntity/Public/MassSchematic.h
--- a/Engine/Plugins/Runtime/MassEntity/Source/MassEntity/Public/MassSchematic.h
+++ b/Engine/Plugins/Runtime/MassEntity/Source/MassEntity/Public/MassSchematic.h
@@ -11,7 +11,7 @@
 class UMassProcessor;
 
 /**
- * Pipe Schematic asset.
+ * Mass Schematic asset.
  */
 UCLASS(BlueprintType)
 class MASSENTITY_API UMassSchematic : public UDataAsset
@@ -22,6 +22,6 @@ public:
 	TConstArrayView<const UMassProcessor*> GetProcessors() const { return MakeArrayView(Processors);  }
 
 protected:
-	UPROPERTY(EditAnywhere, BlueprintReadOnly, Category = Pipe, Instanced)
+	UPROPERTY(EditAnywhere, BlueprintReadOnly, Category = Mass, Instanced)
 	TArray<UMassProcessor*> Processors;
 };
diff --git a/Engine/Plugins/Runtime/MassEntity/Source/MassEntityEditor/Private/AssetTypeActions_MassSchematic.cpp b/Engine/Plugins/Runtime/MassEntity/Source/MassEntityEditor/Private/AssetTypeActions_MassSchematic.cpp
--- a/Engine/Plugins/Runtime/MassEntity/Source/MassEntityEditor/Private/AssetTypeActions_MassSchematic.cpp
+++ b/Engine/Plugins/Runtime/MassEntity/Source/MassEntityEditor/Private/AssetTypeActions_MassSchematic.cpp
@@ -14,10 +14,10 @@ void FAssetTypeActions_MassSchematic::OpenAssetEditor(const TArray<UObject*>& In
 
 	for (auto ObjIt = InObjects.CreateConstIterator(); ObjIt; ++ObjIt)
 	{
-		if (UMassSchematic* Pipe = Cast<UMassSchematic>(*ObjIt))
+		if (UMassSchematic* Schematic = Cast<UMassSchematic>(*ObjIt))
 		{
 			FMassEntityEditorModule& EditorModule = FModuleManager::LoadModuleChecked<FMassEntityEditorModule>("MassEntityEditor");
-			TSharedRef<IMassEntityEditor> NewEditor = EditorModule.CreateMassEntityEditor(Mode, EditWithinLevelEditor, Pipe);
+			TSharedRef<IMassEntityEditor> NewEditor = EditorModule.CreateMassEntityEditor(Mode, EditWithinLevelEditor, Schematic);
 		}
 	}
 }
diff --git a/Engine/Plugins/Runtime/MassEntity/Source/MassEntityEditor/Private/MassEditorStyle.cpp b/Engine/Plugins/Runtime/MassEntity/Source/MassEntityEditor/Private/MassEditorStyle.cpp
--- a/Engine/Plugins/Runtime/MassEntity/Source/MassEntityEditor/Private/MassEditorStyle.cpp
+++ b/Engine/Plugins/Runtime/MassEntity/Source/MassEntityEditor/Private/MassEditorStyle.cpp
@@ -44,12 +44,12 @@ void FMassEntityEditorStyle::Initialize()
 		FTextBlockStyle StateIcon = FTextBlockStyle(NormalText)
 			.SetFont(FEditorStyle::Get().GetFontStyle("FontAwesome.12"))
 			.SetColorAndOpacity(FLinearColor(230.0f / 255.0f, 230.0f / 255.0f, 230.0f / 255.0f, 0.5f));
-		StyleSet->Set("Pipe.Icon", StateIcon);
+		StyleSet->Set("Mass.Icon", StateIcon);
 
 		FTextBlockStyle StateTitle = FTextBlockStyle(NormalText)
 			.SetFont(TTF_CORE_FONT("Fonts/Roboto-Bold", 12))
 			.SetColorAndOpacity(FLinearColor(230.0f / 255.0f, 230.0f / 255.0f, 230.0f / 255.0f));
-		StyleSet->Set("Pipe.State.Title", StateTitle);
+		StyleSet->Set("Mass.State.Title", StateTitle);
 
 		FEditableTextBoxStyle StateTitleEditableText = FEditableTextBoxStyle()
 			.SetFont(TTF_CORE_FONT("Fonts/Roboto-Bold", 10))
@@ -58,9 +58,9 @@ void FMassEntityEditorStyle::Initialize()
 			.SetBackgroundImageFocused(BOX_BRUSH("Common/TextBox_Hovered", FMargin(4.0f / 16.0f)))
 			.SetBackgroundImageReadOnly(BOX_BRUSH("Common/TextBox_ReadOnly", FMargin(4.0f / 16.0f)))
 			.SetScrollBarStyle(ScrollBar);
-		StyleSet->Set("Pipe.State.TitleEditableText", StateTitleEditableText);
+		StyleSet->Set("Mass.State.TitleEditableText", StateTitleEditableText);
 
-		StyleSet->Set("Pipe.State.TitleInlineEditableText", FInlineEditableTextBlockStyle()
+		StyleSet->Set("Mass.State.TitleInlineEditableText", FInlineEditableTextBlockStyle()
 			.SetTextStyle(StateTitle)
 			.SetEditableTextBoxStyle(StateTitleEditableText));
 	}
@@ -70,7 +70,7 @@ void FMassEntityEditorStyle::Initialize()
 		FTextBlockStyle TaskTitle = FTextBlockStyle(NormalText)
 			.SetFont(TTF_CORE_FONT("Fonts/Roboto-Regular", 11))
 			.SetColorAndOpacity(FLinearColor(230.0f / 255.0f, 230.0f / 255.0f, 230.0f / 255.0f));
-		StyleSet->Set("Pipe.Task.Title", TaskTitle);
+		StyleSet->Set("Mass.Task.Title", TaskTitle);
 
 		FEditableTextBoxStyle TaskTitleEditableText = FEditableTextBoxStyle()
 			.SetFont(TTF_CORE_FONT("Fonts/Roboto-Regular", 9))
@@ -79,9 +79,9 @@ void FMassEntityEditorStyle::Initialize()
 			.SetBackgroundImageFocused(BOX_BRUSH("Common/TextBox_Hovered", FMargin(4.0f / 16.0f)))
 			.SetBackgroundImageReadOnly(BOX_BRUSH("Common/TextBox_ReadOnly", FMargin(4.0f / 16.0f)))
 			.SetScrollBarStyle(ScrollBar);
-		StyleSet->Set("Pipe.Task.TitleEditableText", TaskTitleEditableText);
+		StyleSet->Set("Mass.Task.TitleEditableText", TaskTitleEditableText);
 
-		StyleSet->Set("Pipe.Task.TitleInlineEditableText", FInlineEditableTextBlockStyle()
+		StyleSet->Set("Mass.Task.TitleInlineEditableText", FInlineEditableTextBlockStyle()
 			.SetTextStyle(TaskTitle)
 			.SetEditableTextBoxStyle(TaskTitleEditableText));
 	}
@@ -91,12 +91,12 @@ void FMassEntityEditorStyle::Initialize()
 		FTextBlockStyle StateTitle = FTextBlockStyle(NormalText)
 			.SetFont(TTF_CORE_FONT("Fonts/Roboto-Regular", 10))
 			.SetColorAndOpacity(FLinearColor(230.0f / 255.0f, 230.0f / 255.0f, 230.0f / 255.0f));
-		StyleSet->Set("Pipe.Details", StateTitle);
+		StyleSet->Set("Mass.Details", StateTitle);
 	}
 
 	const FLinearColor SelectionColor = FColor(0, 0, 0, 32);
 	const FTableRowStyle& NormalTableRowStyle = FEditorStyle::Get().GetWidgetStyle<FTableRowStyle>("TableView.Row");
-	StyleSet->Set("Pipe.Selection",
+	StyleSet->Set("Mass.Selection",
 		FTableRowStyle(NormalTableRowStyle)
 		.SetActiveBrush(IMAGE_BRUSH("Common/Selection", Icon8x8, SelectionColor))
 		.SetActiveHoveredBrush(IMAGE_BRUSH("Common/Selection", Icon8x8, SelectionColor))
diff --git a/Engine/Plugins/Runtime/MassEntity/Source/MassEntityEditor/Private/MassEntityEditor.cpp b/Engine/Plugins/Runtime/MassEntity/Source/MassEntityEditor/Private/MassEntityEditor.cpp
--- a/Engine/Plugins/Runtime/MassEntity/Source/MassEntityEditor/Private/MassEntityEditor.cpp
+++ b/Engine/Plugins/Runtime/MassEntity/Source/MassEntityEditor/Private/MassEntityEditor.cpp
@@ -30,7 +30,7 @@ void FMassEntityEditor::AddReferencedObjects(FReferenceCollector& Collector)
 
 void FMassEntityEditor::RegisterTabSpawners(const TSharedRef<class FTabManager>& InTabManager)
 {
-	WorkspaceMenuCategory = InTabManager->AddLocalWorkspaceMenuCategory(LOCTEXT("WorkspaceMenu_MassEntityEditor", "Pipe Editor"));
+	WorkspaceMenuCategory = InTabManager->AddLocalWorkspaceMenuCategory(LOCTEXT("WorkspaceMenu_MassEntityEditor", "Mass Editor"));
 	TSharedRef<FWorkspaceItem> WorkspaceMenuCategoryRef = WorkspaceMenuCategory.ToSharedRef();
 
 	FAssetEditorToolkit::RegisterTabSpawners(InTabManager);
@@ -55,7 +55,7 @@ void FMassEntityEditor::InitEditor(const EToolkitMode::Type Mode, const TSharedP
 	//FGuid Result;
 	//FPlatformMisc::CreateGuid(Result);
 	//TSharedRef<FTabManager::FLayout> StandaloneDefaultLayout = FTabManager::NewLayout(*Result.ToString())
-	TSharedRef<FTabManager::FLayout> StandaloneDefaultLayout = FTabManager::NewLayout("Standalone_Pipe_Layout_v1")
+	TSharedRef<FTabManager::FLayout> StandaloneDefaultLayout = FTabManager::NewLayout("Standalone_Mass_Layout_v1")
 	->AddArea
 	(
 		FTabManager::NewPrimaryArea() ->SetOrientation(Orient_Vertical)
diff --git a/Engine/Plugins/Runtime/MassEntity/Source/MassEntityEditor/Private/MassEntityEditor.h b/Engine/Plugins/Runtime/MassEntity/Source/MassEntityEditor/Private/MassEntityEditor.h
--- a/Engine/Plugins/Runtime/MassEntity/Source/MassEntityEditor/Private/MassEntityEditor.h
+++ b/Engine/Plugins/Runtime/MassEntity/Source/MassEntityEditor/Private/MassEntityEditor.h
@@ -46,7 +46,7 @@ protected:
 private:
 	TSharedRef<SDockTab> SpawnTab_AssetDetails(const FSpawnTabArgs& Args);
 
-	/** Pipe Schematic being edited */
+	/** Mass Schematic being edited */
 	UMassSchematic* MassSchematic;
 
 	/** Asset Property View */
diff --git a/Engine/Plugins/Runtime/MassEntity/Source/MassEntityEditor/Private/MassEntityEditorModule.cpp b/Engine/Plugins/Runtime/MassEntity/Source/MassEntityEditor/Private/MassEntityEditorModule.cpp
--- a/Engine/Plugins/Runtime/MassEntity/Source/MassEntityEditor/Private/MassEntityEditorModule.cpp
+++ b/Engine/Plugins/Runtime/MassEntity/Source/MassEntityEditor/Private/MassEntityEditorModule.cpp
@@ -26,9 +26,9 @@ void FMassEntityEditorModule::StartupModule()
 	// Register asset types
 	IAssetTools& AssetTools = FModuleManager::LoadModuleChecked<FAssetToolsModule>("AssetTools").Get();
 
-	TSharedPtr<FAssetTypeActions_MassSchematic> PipeAssetTypeAction = MakeShareable(new FAssetTypeActions_MassSchematic);
-	ItemDataAssetTypeActions.Add(PipeAssetTypeAction);
-	AssetTools.RegisterAssetTypeActions(PipeAssetTypeAction.ToSharedRef());
+	TSharedPtr<FAssetTypeActions_MassSchematic> MassAssetTypeAction = MakeShareable(new FAssetTypeActions_MassSchematic);
+	ItemDataAssetTypeActions.Add(MassAssetTypeAction);
+	AssetTools.RegisterAssetTypeActions(MassAssetTypeAction.ToSharedRef());
 
 	// Register the details customizers
 	FPropertyEditorModule& PropertyModule = FModuleManager::LoadModuleChecked<FPropertyEditorModule>("PropertyEditor");
diff --git a/Engine/Plugins/Runtime/MassEntity/Source/MassEntityEditor/Public/IMassEntityEditor.h b/Engine/Plugins/Runtime/MassEntity/Source/MassEntityEditor/Public/IMassEntityEditor.h
--- a/Engine/Plugins/Runtime/MassEntity/Source/MassEntityEditor/Public/IMassEntityEditor.h
+++ b/Engine/Plugins/Runtime/MassEntity/Source/MassEntityEditor/Public/IMassEntityEditor.h
@@ -6,7 +6,7 @@
 #include "Toolkits/AssetEditorToolkit.h"
 
 
-/** Pipe Editor public interface */
+/** Mass Editor public interface */
 class MASSENTITYEDITOR_API IMassEntityEditor : public FAssetEditorToolkit
 {
 public:
diff --git a/Engine/Plugins/Runtime/MassEntity/Source/MassEntityEditor/Public/MassEntityEditorModule.h b/Engine/Plugins/Runtime/MassEntity/Source/MassEntityEditor/Public/MassEntityEditorModule.h
--- a/Engine/Plugins/Runtime/MassEntity/Source/MassEntityEditor/Public/MassEntityEditorModule.h
+++ b/Engine/Plugins/Runtime/MassEntity/Source/MassEntityEditor/Public/MassEntityEditorModule.h
@@ -25,7 +25,7 @@ public:
 	virtual void StartupModule() override;
 	virtual void ShutdownModule() override;
 
-	/** Creates an instance of Pipe editor. Only virtual so that it can be called across the DLL boundary. */
+	/** Creates an instance of Mass editor. Only virtual so that it can be called across the DLL boundary. */
 	virtual TSharedRef<IMassEntityEditor> CreateMassEntityEditor(const EToolkitMode::Type Mode, const TSharedPtr< class IToolkitHost >& InitToolkitHost, UMassSchematic* MassSchematic);
 
 	virtual TSharedPtr<FExtensibilityManager> GetMenuExtensibilityManager() override { return MenuExtensibilityManager; }
diff --git a/Engine/Plugins/Runtime/MassEntity/Source/MassEntityTestSuite/Private/MassCompositeProcessorTest.cpp b/Engine/Plugins/Runtime/MassEntity/Source/MassEntityTestSuite/Private/MassCompositeProcessorTest.cpp
--- a/Engine/Plugins/Runtime/MassEntity/Source/MassEntityTestSuite/Private/MassCompositeProcessorTest.cpp
+++ b/Engine/Plugins/Runtime/MassEntity/Source/MassEntityTestSuite/Private/MassCompositeProcessorTest.cpp
@@ -13,7 +13,7 @@
 
 PRAGMA_DISABLE_OPTIMIZATION
 
-namespace FPipeCompositeProcessorTest
+namespace FMassCompositeProcessorTest
 {
 
 struct FCompositeProcessorTest_Empty : FEntityTestBase
@@ -47,7 +47,7 @@ struct FCompositeProcessorTest_MultipleSubProcessors : FEntityTestBase
 			TArray<UMassProcessor*> Processors;
 			for (int i = 0; i < 3; ++i)
 			{
-				UPipeTestProcessorBase* Processor = NewObject<UPipeTestProcessorBase>(EntitySubsystem);
+				UMassTestProcessorBase* Processor = NewObject<UMassTestProcessorBase>(EntitySubsystem);
 				Processor->ExecutionFunction = [Processor, &Result, i](UMassEntitySubsystem& InEntitySubsystem, FMassExecutionContext& Context) {
 						check(Processor);
 						Result += FMath::Pow(10.f, float(i));
@@ -66,7 +66,7 @@ struct FCompositeProcessorTest_MultipleSubProcessors : FEntityTestBase
 };
 IMPLEMENT_AI_INSTANT_TEST(FCompositeProcessorTest_MultipleSubProcessors, "System.Mass.Processor.Composite.MultipleSubProcessors");
 
-} // FPipeCompositeProcessorTest
+} // FMassCompositeProcessorTest
 
 PRAGMA_ENABLE_OPTIMIZATION
 
diff --git a/Engine/Plugins/Runtime/MassEntity/Source/MassEntityTestSuite/Private/MassDependencyTest.cpp b/Engine/Plugins/Runtime/MassEntity/Source/MassEntityTestSuite/Private/MassDependencyTest.cpp
--- a/Engine/Plugins/Runtime/MassEntity/Source/MassEntityTestSuite/Private/MassDependencyTest.cpp
+++ b/Engine/Plugins/Runtime/MassEntity/Source/MassEntityTestSuite/Private/MassDependencyTest.cpp
@@ -10,7 +10,7 @@
 
 PRAGMA_DISABLE_OPTIMIZATION
 
-namespace FPipeDependencySolverTest
+namespace FMassDependencySolverTest
 {
 
 template<typename T>
@@ -21,7 +21,7 @@ static FName GetProcessorName()
 
 struct FDependencySolverBase : FAITestBase
 {
-	TArray<UPipeTestProcessorBase*> Processors;
+	TArray<UMassTestProcessorBase*> Processors;
 	TArray<FProcessorDependencySolver::FOrderInfo> Result;
 	
 	virtual bool SetUp() override
@@ -43,16 +43,16 @@ struct FTrivialDependency : FDependencySolverBase
 	virtual bool SetUp() override
 	{
 		{
-			UPipeTestProcessorBase* Proc = Processors.Add_GetRef(NewObject<UPipeTestProcessor_A>());
-			Proc->GetMutableExecutionOrder().ExecuteAfter.Add(GetProcessorName<UPipeTestProcessor_C>());
+			UMassTestProcessorBase* Proc = Processors.Add_GetRef(NewObject<UMassTestProcessor_A>());
+			Proc->GetMutableExecutionOrder().ExecuteAfter.Add(GetProcessorName<UMassTestProcessor_C>());
 		}
 
 		{
-			UPipeTestProcessorBase* Proc = Processors.Add_GetRef(NewObject<UPipeTestProcessor_B>());
-			Proc->GetMutableExecutionOrder().ExecuteAfter.Add(GetProcessorName<UPipeTestProcessor_A>());
+			UMassTestProcessorBase* Proc = Processors.Add_GetRef(NewObject<UMassTestProcessor_B>());
+			Proc->GetMutableExecutionOrder().ExecuteAfter.Add(GetProcessorName<UMassTestProcessor_A>());
 		}
 
-		Processors.Add(NewObject<UPipeTestProcessor_C>());
+		Processors.Add(NewObject<UMassTestProcessor_C>());
 
 		return true;
 	}
@@ -61,9 +61,9 @@ struct FTrivialDependency : FDependencySolverBase
 	{
 		Solve();
 
-		AITEST_TRUE("C is expected to be first", Result[0].Name == GetProcessorName<UPipeTestProcessor_C>());
-		AITEST_TRUE("A is expected to be second", Result[1].Name == GetProcessorName<UPipeTestProcessor_A>());
-		AITEST_TRUE("B is expected to be third", Result[2].Name == GetProcessorName<UPipeTestProcessor_B>());
+		AITEST_TRUE("C is expected to be first", Result[0].Name == GetProcessorName<UMassTestProcessor_C>());
+		AITEST_TRUE("A is expected to be second", Result[1].Name == GetProcessorName<UMassTestProcessor_A>());
+		AITEST_TRUE("B is expected to be third", Result[2].Name == GetProcessorName<UMassTestProcessor_B>());
 
 		return true;
 	}
@@ -76,13 +76,13 @@ struct FDeepGroup : FDependencySolverBase
 	virtual bool SetUp() override
 	{
 		{
-			UPipeTestProcessorBase* Proc = Processors.Add_GetRef(NewObject<UPipeTestProcessor_A>());
+			UMassTestProcessorBase* Proc = Processors.Add_GetRef(NewObject<UMassTestProcessor_A>());
 			Proc->GetMutableExecutionOrder().ExecuteAfter.Add(TEXT("W.X.Y.Z"));
 			Proc->GetMutableExecutionOrder().ExecuteInGroup = TEXT("P.Q.R");
 		}
 
 		{
-			UPipeTestProcessorBase* Proc = Processors.Add_GetRef(NewObject<UPipeTestProcessor_B>());
+			UMassTestProcessorBase* Proc = Processors.Add_GetRef(NewObject<UMassTestProcessor_B>());
 			Proc->GetMutableExecutionOrder().ExecuteInGroup = TEXT("W.X.Y.Z");
 		}
 
@@ -102,8 +102,8 @@ struct FDeepGroup : FDependencySolverBase
 			}
 		}
 
-		AITEST_TRUE("B is expected to be first", Result[0].Name == GetProcessorName<UPipeTestProcessor_B>());
-		AITEST_TRUE("A is expected to be second", Result[1].Name == GetProcessorName<UPipeTestProcessor_A>());
+		AITEST_TRUE("B is expected to be first", Result[0].Name == GetProcessorName<UMassTestProcessor_B>());
+		AITEST_TRUE("A is expected to be second", Result[1].Name == GetProcessorName<UMassTestProcessor_A>());
 
 		return true;
 	}
@@ -115,35 +115,35 @@ struct FComplexScenario : FDependencySolverBase
 	virtual bool SetUp() override
 	{
 		{
-			UPipeTestProcessorBase* Proc = Processors.Add_GetRef(NewObject<UPipeTestProcessor_A>());
+			UMassTestProcessorBase* Proc = Processors.Add_GetRef(NewObject<UMassTestProcessor_A>());
 			Proc->GetMutableExecutionOrder().ExecuteInGroup = TEXT("X.Z");
 			Proc->GetMutableExecutionOrder().ExecuteAfter.Add(TEXT("X.Y"));
-			Proc->GetMutableExecutionOrder().ExecuteAfter.Add(UPipeTestProcessor_E::StaticClass()->GetFName());
+			Proc->GetMutableExecutionOrder().ExecuteAfter.Add(UMassTestProcessor_E::StaticClass()->GetFName());
 		}
 
 		{
-			UPipeTestProcessorBase* Proc = Processors.Add_GetRef(NewObject<UPipeTestProcessor_B>());
+			UMassTestProcessorBase* Proc = Processors.Add_GetRef(NewObject<UMassTestProcessor_B>());
 			Proc->GetMutableExecutionOrder().ExecuteInGroup = TEXT("X.Y");
 		}
 
 		{
-			UPipeTestProcessorBase* Proc = Processors.Add_GetRef(NewObject<UPipeTestProcessor_C>());
+			UMassTestProcessorBase* Proc = Processors.Add_GetRef(NewObject<UMassTestProcessor_C>());
 			Proc->GetMutableExecutionOrder().ExecuteInGroup = TEXT("X.Y");
 		}
 
 		{
-			UPipeTestProcessorBase* Proc = Processors.Add_GetRef(NewObject<UPipeTestProcessor_D>());
-			Proc->GetMutableExecutionOrder().ExecuteBefore.Add(UPipeTestProcessor_A::StaticClass()->GetFName());
+			UMassTestProcessorBase* Proc = Processors.Add_GetRef(NewObject<UMassTestProcessor_D>());
+			Proc->GetMutableExecutionOrder().ExecuteBefore.Add(UMassTestProcessor_A::StaticClass()->GetFName());
 		}
 
 		{
-			UPipeTestProcessorBase* Proc = Processors.Add_GetRef(NewObject<UPipeTestProcessor_E>());
+			UMassTestProcessorBase* Proc = Processors.Add_GetRef(NewObject<UMassTestProcessor_E>());
 			Proc->GetMutableExecutionOrder().ExecuteInGroup = TEXT("X.Z");
 		}
 
 		{
-			UPipeTestProcessorBase* Proc = Processors.Add_GetRef(NewObject<UPipeTestProcessor_F>());
-			Proc->GetMutableExecutionOrder().ExecuteAfter.Add(UPipeTestProcessor_A::StaticClass()->GetFName());
+			UMassTestProcessorBase* Proc = Processors.Add_GetRef(NewObject<UMassTestProcessor_F>());
+			Proc->GetMutableExecutionOrder().ExecuteAfter.Add(UMassTestProcessor_A::StaticClass()->GetFName());
 		}
 		return true;
 	}
@@ -161,11 +161,11 @@ struct FComplexScenario : FDependencySolverBase
 			}
 		}
 
-		AITEST_TRUE("D is the only fully dependency-less processor so should be first", Result[0].Name == GetProcessorName<UPipeTestProcessor_D>());		
-		AITEST_TRUE("B and C come next", (Result[1].Name == GetProcessorName<UPipeTestProcessor_B>() || Result[2].Name == GetProcessorName<UPipeTestProcessor_B>()) && (Result[1].Name == GetProcessorName<UPipeTestProcessor_C>() || Result[2].Name == GetProcessorName<UPipeTestProcessor_C>()));
-		AITEST_TRUE("Following by E", Result[3].Name == GetProcessorName<UPipeTestProcessor_E>());
-		AITEST_TRUE("Then A", Result[4].Name == GetProcessorName<UPipeTestProcessor_A>());
-		AITEST_TRUE("F is last", Result[5].Name == GetProcessorName<UPipeTestProcessor_F>());
+		AITEST_TRUE("D is the only fully dependency-less processor so should be first", Result[0].Name == GetProcessorName<UMassTestProcessor_D>());		
+		AITEST_TRUE("B and C come next", (Result[1].Name == GetProcessorName<UMassTestProcessor_B>() || Result[2].Name == GetProcessorName<UMassTestProcessor_B>()) && (Result[1].Name == GetProcessorName<UMassTestProcessor_C>() || Result[2].Name == GetProcessorName<UMassTestProcessor_C>()));
+		AITEST_TRUE("Following by E", Result[3].Name == GetProcessorName<UMassTestProcessor_E>());
+		AITEST_TRUE("Then A", Result[4].Name == GetProcessorName<UMassTestProcessor_A>());
+		AITEST_TRUE("F is last", Result[5].Name == GetProcessorName<UMassTestProcessor_F>());
 
 		return true;
 	}
@@ -221,9 +221,9 @@ struct FBadInput_MixedNulls : FDependencySolverBase
 		GetTestRunner().AddExpectedError(TEXT("nullptr found in Processors"), EAutomationExpectedErrorFlags::Contains, 2);
 		Processors.Reset();
 		Processors.Add(nullptr);
-		Processors.Add(NewObject<UPipeTestProcessor_A>());
+		Processors.Add(NewObject<UMassTestProcessor_A>());
 		Processors.Add(nullptr);
-		Processors.Add(NewObject<UPipeTestProcessor_B>());
+		Processors.Add(NewObject<UMassTestProcessor_B>());
 		Solve();
 		AITEST_TRUE("Mixed nullptr and proper inputs should be handled gracefully", Result.Num() == 2);
 		return true;
@@ -237,9 +237,9 @@ struct FBadInput_Duplicates : FDependencySolverBase
 	{
 		GetTestRunner().AddExpectedError(TEXT("already registered. Duplicates are not supported"), EAutomationExpectedErrorFlags::Contains, 2);
 		Processors.Reset();
-		Processors.Add(NewObject<UPipeTestProcessor_A>());
-		Processors.Add(NewObject<UPipeTestProcessor_A>());
-		Processors.Add(NewObject<UPipeTestProcessor_A>());
+		Processors.Add(NewObject<UMassTestProcessor_A>());
+		Processors.Add(NewObject<UMassTestProcessor_A>());
+		Processors.Add(NewObject<UMassTestProcessor_A>());
 		Solve();
 		AITEST_TRUE("Duplicates in input should be handled gracefully", Result.Num() == 1);
 		return true;
@@ -282,7 +282,7 @@ IMPLEMENT_AI_INSTANT_TEST(FGroupNamesGeneration, "System.Mass.Dependencies.Subgr
 // * single non-null processor
 // * circular dependencies  
 
-} // FPipeDependencySolverTest
+} // FMassDependencySolverTest
 
 PRAGMA_ENABLE_OPTIMIZATION
 
diff --git a/Engine/Plugins/Runtime/MassEntity/Source/MassEntityTestSuite/Private/MassEntityTest.cpp b/Engine/Plugins/Runtime/MassEntity/Source/MassEntityTestSuite/Private/MassEntityTest.cpp
--- a/Engine/Plugins/Runtime/MassEntity/Source/MassEntityTestSuite/Private/MassEntityTest.cpp
+++ b/Engine/Plugins/Runtime/MassEntity/Source/MassEntityTestSuite/Private/MassEntityTest.cpp
@@ -13,7 +13,7 @@
 
 PRAGMA_DISABLE_OPTIMIZATION
 
-namespace FPipeEntityTest
+namespace FMassEntityTest
 {
 #if WITH_MASSENTITY_DEBUG
 struct FEntityTest_ArchetypeCreation : FEntityTestBase
@@ -357,7 +357,7 @@ IMPLEMENT_AI_INSTANT_TEST(FEntityTest_ReserveAPreviouslyBuiltEntity, "System.Mas
 
 #endif // WITH_MASSENTITY_DEBUG
 
-} // FPipeEntityTestTest
+} // FMassEntityTestTest
 
 PRAGMA_ENABLE_OPTIMIZATION
 
diff --git a/Engine/Plugins/Runtime/MassEntity/Source/MassEntityTestSuite/Private/MassExecutionTest.cpp b/Engine/Plugins/Runtime/MassEntity/Source/MassEntityTestSuite/Private/MassExecutionTest.cpp
--- a/Engine/Plugins/Runtime/MassEntity/Source/MassEntityTestSuite/Private/MassExecutionTest.cpp
+++ b/Engine/Plugins/Runtime/MassEntity/Source/MassEntityTestSuite/Private/MassExecutionTest.cpp
@@ -14,7 +14,7 @@ PRAGMA_DISABLE_OPTIMIZATION
 //----------------------------------------------------------------------//
 // tests 
 //----------------------------------------------------------------------//
-namespace FPipeExecutionTest
+namespace FMassExecutionTest
 {
 
 struct FExecution_Setup : FExecutionTestBase
@@ -110,7 +110,7 @@ struct FExecution_SingleValidProcessor : FExecutionTestBase
 		CA_ASSUME(EntitySubsystem);
 		const float DeltaSeconds = 0.f;
 		FMassProcessingContext ProcessingContext(*EntitySubsystem, DeltaSeconds);
-		UPipeTestProcessorBase* Processor = NewObject<UPipeTestProcessorBase>(EntitySubsystem);
+		UMassTestProcessorBase* Processor = NewObject<UMassTestProcessorBase>(EntitySubsystem);
 		check(Processor);
 
 		// nothing should break. The actual result of processing is getting tested in MassProcessorTests.cpp
@@ -151,7 +151,7 @@ struct FExecution_Sparse : FEntityTestBase
 		CA_ASSUME(EntitySubsystem);
 		const float DeltaSeconds = 0.f;
 		FMassProcessingContext ProcessingContext(*EntitySubsystem, DeltaSeconds);
-		UPipeTestProcessorBase* Processor = NewObject<UPipeTestProcessorBase>(EntitySubsystem);
+		UMassTestProcessorBase* Processor = NewObject<UMassTestProcessorBase>(EntitySubsystem);
 		check(Processor);
 
 		FMassRuntimePipeline Pipeline;
@@ -170,7 +170,7 @@ struct FExecution_Sparse : FEntityTestBase
 	}
 };
 IMPLEMENT_AI_INSTANT_TEST(FExecution_Sparse, "System.Mass.Execution.Sparse");
-} // FPipeExecutionTest
+} // FMassExecutionTest
 
 PRAGMA_ENABLE_OPTIMIZATION
 
diff --git a/Engine/Plugins/Runtime/MassEntity/Source/MassEntityTestSuite/Private/MassProcessorTest.cpp b/Engine/Plugins/Runtime/MassEntity/Source/MassEntityTestSuite/Private/MassProcessorTest.cpp
--- a/Engine/Plugins/Runtime/MassEntity/Source/MassEntityTestSuite/Private/MassProcessorTest.cpp
+++ b/Engine/Plugins/Runtime/MassEntity/Source/MassEntityTestSuite/Private/MassProcessorTest.cpp
@@ -41,7 +41,7 @@ struct FProcessorTest_NoEntities : FEntityTestBase
 	{
 		CA_ASSUME(EntitySubsystem);
 		
-		int32 EntityProcessedCount = SimpleProcessorRun<UPipeTestProcessor_Floats>(*EntitySubsystem);
+		int32 EntityProcessedCount = SimpleProcessorRun<UMassTestProcessor_Floats>(*EntitySubsystem);
 		AITEST_EQUAL("No entities have been created yet so the processor shouldn't do any work.", EntityProcessedCount, 0);
 				
 		return true;
@@ -58,7 +58,7 @@ struct FProcessorTest_NoMatchingEntities : FEntityTestBase
 
 		TArray<FMassEntityHandle> EntitiesCreated;
 		EntitySubsystem->BatchCreateEntities(IntsArchetype, 100, EntitiesCreated);
-		int32 EntityProcessedCount = SimpleProcessorRun<UPipeTestProcessor_Floats>(*EntitySubsystem);
+		int32 EntityProcessedCount = SimpleProcessorRun<UMassTestProcessor_Floats>(*EntitySubsystem);
 		AITEST_EQUAL("No matching entities have been created yet so the processor shouldn't do any work.", EntityProcessedCount, 0);
 
 		return true;
@@ -75,7 +75,7 @@ struct FProcessorTest_OneMatchingArchetype : FEntityTestBase
 		const int32 EntitiesToCreate = 100;
 		TArray<FMassEntityHandle> EntitiesCreated;
 		EntitySubsystem->BatchCreateEntities(FloatsArchetype, EntitiesToCreate, EntitiesCreated);
-		int32 EntityProcessedCount = SimpleProcessorRun<UPipeTestProcessor_Floats>(*EntitySubsystem);
+		int32 EntityProcessedCount = SimpleProcessorRun<UMassTestProcessor_Floats>(*EntitySubsystem);
 		AITEST_EQUAL("No matching entities have been created yet so the processor shouldn't do any work.", EntityProcessedCount, EntitiesToCreate);
 
 		return true;
@@ -96,7 +96,7 @@ struct FProcessorTest_MultipleMatchingArchetype : FEntityTestBase
 		EntitySubsystem->BatchCreateEntities(FloatsIntsArchetype, EntitiesToCreate, EntitiesCreated);
 		EntitySubsystem->BatchCreateEntities(IntsArchetype, EntitiesToCreate, EntitiesCreated);
 		// note that only two of these archetypes match 
-		int32 EntityProcessedCount = SimpleProcessorRun<UPipeTestProcessor_Floats>(*EntitySubsystem);
+		int32 EntityProcessedCount = SimpleProcessorRun<UMassTestProcessor_Floats>(*EntitySubsystem);
 		AITEST_EQUAL("Two of given three archetypes should match.", EntityProcessedCount, EntitiesToCreate * 2);
 
 		return true;
diff --git a/Engine/Plugins/Runtime/MassEntity/Source/MassEntityTestSuite/Private/MassQueryTest.cpp b/Engine/Plugins/Runtime/MassEntity/Source/MassEntityTestSuite/Private/MassQueryTest.cpp
--- a/Engine/Plugins/Runtime/MassEntity/Source/MassEntityTestSuite/Private/MassQueryTest.cpp
+++ b/Engine/Plugins/Runtime/MassEntity/Source/MassEntityTestSuite/Private/MassQueryTest.cpp
@@ -24,7 +24,7 @@ mz@todo:
 */
 
 
-namespace FPipeQueryTest
+namespace FMassQueryTest
 {
 
 struct FQueryTest_ProcessorRequirements : FEntityTestBase
@@ -33,7 +33,7 @@ struct FQueryTest_ProcessorRequirements : FEntityTestBase
 	{
 		CA_ASSUME(EntitySubsystem);
 
-		UPipeTestProcessor_Floats* Processor = NewObject<UPipeTestProcessor_Floats>(EntitySubsystem);
+		UMassTestProcessor_Floats* Processor = NewObject<UMassTestProcessor_Floats>(EntitySubsystem);
 		TConstArrayView<FMassFragmentRequirement> Requirements = Processor->TestGetQuery().GetRequirements();
 		
 		AITEST_TRUE("Query should have extracted some requirements from the given Processor", Requirements.Num() > 0);
@@ -75,7 +75,7 @@ struct FQueryTest_FragmentViewBinding : FEntityTestBase
 		FTestFragment_Float& TestedFragment = EntitySubsystem->GetFragmentDataChecked<FTestFragment_Float>(Entity);
 		AITEST_TRUE("Initial value of the fragment should match expectations", TestedFragment.Value == 0.f);
 
-		UPipeTestProcessor_Floats* Processor = NewObject<UPipeTestProcessor_Floats>(EntitySubsystem);
+		UMassTestProcessor_Floats* Processor = NewObject<UMassTestProcessor_Floats>(EntitySubsystem);
 		Processor->ExecutionFunction = [Processor](UMassEntitySubsystem& InEntitySubsystem, FMassExecutionContext& Context) {
 			check(Processor);
 			//FMassEntityQuery Query(*Processor);
@@ -542,7 +542,7 @@ struct FQueryTest_AnyTag : FEntityTestBase
 };
 IMPLEMENT_AI_INSTANT_TEST(FQueryTest_AnyTag, "System.Mass.Query.AnyTag");
 
-} // FPipeQueryTest
+} // FMassQueryTest
 
 PRAGMA_ENABLE_OPTIMIZATION
 
diff --git a/Engine/Plugins/Runtime/MassEntity/Source/MassEntityTestSuite/Private/MassTestTypes.cpp b/Engine/Plugins/Runtime/MassEntity/Source/MassEntityTestSuite/Private/MassTestTypes.cpp
--- a/Engine/Plugins/Runtime/MassEntity/Source/MassEntityTestSuite/Private/MassTestTypes.cpp
+++ b/Engine/Plugins/Runtime/MassEntity/Source/MassEntityTestSuite/Private/MassTestTypes.cpp
@@ -46,7 +46,7 @@ bool FEntityTestBase::SetUp()
 //----------------------------------------------------------------------//
 // Processors 
 //----------------------------------------------------------------------//
-UPipeTestProcessorBase::UPipeTestProcessorBase()
+UMassTestProcessorBase::UMassTestProcessorBase()
 {
 #if WITH_EDITORONLY_DATA
 	bCanShowUpInSettings = false;
@@ -58,7 +58,7 @@ UPipeTestProcessorBase::UPipeTestProcessorBase()
 	RequirementsFunction = [](FMassEntityQuery& Query){};
 }
 
-UPipeTestProcessor_Floats::UPipeTestProcessor_Floats()
+UMassTestProcessor_Floats::UMassTestProcessor_Floats()
 {
 	RequirementsFunction = [this](FMassEntityQuery& Query)
 	{
@@ -66,7 +66,7 @@ UPipeTestProcessor_Floats::UPipeTestProcessor_Floats()
 	};
 }
 
-UPipeTestProcessor_Ints::UPipeTestProcessor_Ints()
+UMassTestProcessor_Ints::UMassTestProcessor_Ints()
 {
 	RequirementsFunction = [this](FMassEntityQuery& Query)
 	{
@@ -74,7 +74,7 @@ UPipeTestProcessor_Ints::UPipeTestProcessor_Ints()
 	};
 }
 
-UPipeTestProcessor_FloatsInts::UPipeTestProcessor_FloatsInts()
+UMassTestProcessor_FloatsInts::UMassTestProcessor_FloatsInts()
 {
 	RequirementsFunction = [this](FMassEntityQuery& Query)
 	{
diff --git a/Engine/Plugins/Runtime/MassEntity/Source/MassEntityTestSuite/Public/MassEntityTestTypes.h b/Engine/Plugins/Runtime/MassEntity/Source/MassEntityTestSuite/Public/MassEntityTestTypes.h
--- a/Engine/Plugins/Runtime/MassEntity/Source/MassEntityTestSuite/Public/MassEntityTestTypes.h
+++ b/Engine/Plugins/Runtime/MassEntity/Source/MassEntityTestSuite/Public/MassEntityTestTypes.h
@@ -66,11 +66,11 @@ struct FTestTag_D : public FMassTag
 
 
 UCLASS()
-class UPipeTestProcessorBase : public UMassProcessor
+class UMassTestProcessorBase : public UMassProcessor
 {
 	GENERATED_BODY()
 public:
-	UPipeTestProcessorBase();
+	UMassTestProcessorBase();
 	FMassProcessorExecutionOrder& GetMutableExecutionOrder() { return ExecutionOrder; }
 	virtual void Execute(UMassEntitySubsystem& EntitySubsystem, FMassExecutionContext& Context) override 
 	{
@@ -91,67 +91,67 @@ protected:
 };
 
 UCLASS()
-class UPipeTestProcessor_A : public UPipeTestProcessorBase
+class UMassTestProcessor_A : public UMassTestProcessorBase
 {
 	GENERATED_BODY()
 };
 
 UCLASS()
-class UPipeTestProcessor_B : public UPipeTestProcessorBase
+class UMassTestProcessor_B : public UMassTestProcessorBase
 {
 	GENERATED_BODY()
 };
 
 UCLASS()
-class UPipeTestProcessor_C : public UPipeTestProcessorBase
+class UMassTestProcessor_C : public UMassTestProcessorBase
 {
 	GENERATED_BODY()
 };
 
 UCLASS()
-class UPipeTestProcessor_D : public UPipeTestProcessorBase
+class UMassTestProcessor_D : public UMassTestProcessorBase
 {
 	GENERATED_BODY()
 };
 
 UCLASS()
-class UPipeTestProcessor_E : public UPipeTestProcessorBase
+class UMassTestProcessor_E : public UMassTestProcessorBase
 {
 	GENERATED_BODY()
 };
 
 UCLASS()
-class UPipeTestProcessor_F : public UPipeTestProcessorBase
+class UMassTestProcessor_F : public UMassTestProcessorBase
 {
 	GENERATED_BODY()
 };
 
 UCLASS()
-class UPipeTestProcessor_Floats : public UPipeTestProcessorBase
+class UMassTestProcessor_Floats : public UMassTestProcessorBase
 {
 	GENERATED_BODY()
 public:
 	TArrayView<FTestFragment_Float> Floats;
-	UPipeTestProcessor_Floats();
+	UMassTestProcessor_Floats();
 };
 
 UCLASS()
-class UPipeTestProcessor_Ints : public UPipeTestProcessorBase
+class UMassTestProcessor_Ints : public UMassTestProcessorBase
 {
 	GENERATED_BODY()
 public:
 	TArrayView<FTestFragment_Int> Ints;
-	UPipeTestProcessor_Ints();
+	UMassTestProcessor_Ints();
 };
 
 UCLASS()
-class UPipeTestProcessor_FloatsInts : public UPipeTestProcessorBase
+class UMassTestProcessor_FloatsInts : public UMassTestProcessorBase
 {
 	GENERATED_BODY()
 public:
 	TArrayView<FTestFragment_Float> Floats;
 	TArrayView<FTestFragment_Int> Ints;
-	UPipeTestProcessor_FloatsInts();
+	UMassTestProcessor_FloatsInts();
 };
 
 struct FExecutionTestBase : FAITestBase
diff --git a/Engine/Plugins/Runtime/MassGameplay/Source/MassActors/Private/MassAgentSubsystem.cpp b/Engine/Plugins/Runtime/MassGameplay/Source/MassActors/Private/MassAgentSubsystem.cpp
--- a/Engine/Plugins/Runtime/MassGameplay/Source/MassActors/Private/MassAgentSubsystem.cpp
+++ b/Engine/Plugins/Runtime/MassGameplay/Source/MassActors/Private/MassAgentSubsystem.cpp
@@ -286,7 +286,7 @@ void UMassAgentSubsystem::HandlePendingInitialization()
 	check(EntitySystem);
 	check(SpawnerSystem);
 	check(SimulationSystem);
-	check(SimulationSystem->GetPhaseManager().IsDuringPipeProcessing() == false);
+	check(SimulationSystem->GetPhaseManager().IsDuringMassProcessing() == false);
 
 	for (TTuple<FMassEntityTemplateID, FMassAgentInitializationQueue>& Data : PendingAgentEntities)
 	{
diff --git a/Engine/Plugins/Runtime/MassGameplay/Source/MassSpawner/Private/MassSpawnerSubsystem.cpp b/Engine/Plugins/Runtime/MassGameplay/Source/MassSpawner/Private/MassSpawnerSubsystem.cpp
--- a/Engine/Plugins/Runtime/MassGameplay/Source/MassSpawner/Private/MassSpawnerSubsystem.cpp
+++ b/Engine/Plugins/Runtime/MassGameplay/Source/MassSpawner/Private/MassSpawnerSubsystem.cpp
@@ -220,7 +220,7 @@ void UMassSpawnerSubsystem::DestroyEntities(const FMassEntityTemplateID Template
 
 	check(EntitySystem);
 	check(SimulationSystem);
-	checkf(!SimulationSystem->GetPhaseManager().IsDuringPipeProcessing()
+	checkf(!SimulationSystem->GetPhaseManager().IsDuringMassProcessing()
 		, TEXT("%s called while MassEntity processing in progress. This is unsupported and dangerous!"), ANSI_TO_TCHAR(__FUNCTION__));
 
 	UWorld* World = GetWorld();
