Updated Mass' uplugin files to reflect actual version number and filled in some blanks.

diff --git a/Engine/Plugins/AI/MassAI/MassAI.uplugin b/Engine/Plugins/AI/MassAI/MassAI.uplugin
--- a/Engine/Plugins/AI/MassAI/MassAI.uplugin
+++ b/Engine/Plugins/AI/MassAI/MassAI.uplugin
@@ -1,7 +1,7 @@
 {
 	"FileVersion": 3,
 	"Version": 1,
-	"VersionName": "1.0",
+	"VersionName": "0.4",
 	"FriendlyName": "MassAI",
 	"Description": "AI-specific functionality extenting MassGameplay",
 	"Category": "AI",
diff --git a/Engine/Plugins/AI/MassCrowd/MassCrowd.uplugin b/Engine/Plugins/AI/MassCrowd/MassCrowd.uplugin
--- a/Engine/Plugins/AI/MassCrowd/MassCrowd.uplugin
+++ b/Engine/Plugins/AI/MassCrowd/MassCrowd.uplugin
@@ -1,7 +1,7 @@
 {
 	"FileVersion": 1,
 	"Version": 1,
-	"VersionName": "1.0",
+	"VersionName": "0.4",
 	"FriendlyName": "MassCrowd",
 	"Description": "Spline based AI crowd system",
 	"Category": "AI",
diff --git a/Engine/Plugins/Runtime/MassEntity/MassEntity.uplugin b/Engine/Plugins/Runtime/MassEntity/MassEntity.uplugin
--- a/Engine/Plugins/Runtime/MassEntity/MassEntity.uplugin
+++ b/Engine/Plugins/Runtime/MassEntity/MassEntity.uplugin
@@ -1,11 +1,11 @@
 {
 	"FileVersion": 3,
 	"Version": 1,
-	"VersionName": "1.0",
+	"VersionName": "0.4",
 	"FriendlyName": "MassEntity",
-	"Description": "",
+	"Description": "Gameplay-focused framework supporting data-oriented processing",
 	"Category": "Gameplay",
-	"CreatedBy": "",
+	"CreatedBy": "Epic Games, Inc.",
 	"CreatedByURL": "",
 	"DocsURL": "",
 	"MarketplaceURL": "",
diff --git a/Engine/Plugins/Runtime/MassGameplay/MassGameplay.uplugin b/Engine/Plugins/Runtime/MassGameplay/MassGameplay.uplugin
--- a/Engine/Plugins/Runtime/MassGameplay/MassGameplay.uplugin
+++ b/Engine/Plugins/Runtime/MassGameplay/MassGameplay.uplugin
@@ -1,7 +1,7 @@
 {
 	"FileVersion": 3,
 	"Version": 1,
-	"VersionName": "1.0",
+	"VersionName": "0.4",
 	"FriendlyName": "MassGameplay",
 	"Description": "Implementation of large-scale agent simulation based on MassEntity",
 	"Category": "Gameplay",
