Fixed processing phases' description no longer getting displayed in multilined fashion in Mass' Settings

diff --git a/Engine/Plugins/Runtime/MassEntity/Source/MassEntity/Private/MassEntitySettings.cpp b/Engine/Plugins/Runtime/MassEntity/Source/MassEntity/Private/MassEntitySettings.cpp
--- a/Engine/Plugins/Runtime/MassEntity/Source/MassEntity/Private/MassEntitySettings.cpp
+++ b/Engine/Plugins/Runtime/MassEntity/Source/MassEntity/Private/MassEntitySettings.cpp
@@ -57,7 +57,7 @@ void UMassSettings::BuildPhases()
 		
 		FStringOutputDevice Ar;
 		PhaseConfig.PhaseProcessor->DebugOutputDescription(Ar);
-		PhaseConfig.Description = MoveTemp(Ar);
+		PhaseConfig.Description = FText::FromString(Ar);
 	}
 #endif // WITH_EDITOR
 }
diff --git a/Engine/Plugins/Runtime/MassEntity/Source/MassEntity/Public/MassEntitySettings.h b/Engine/Plugins/Runtime/MassEntity/Source/MassEntity/Public/MassEntitySettings.h
--- a/Engine/Plugins/Runtime/MassEntity/Source/MassEntity/Public/MassEntitySettings.h
+++ b/Engine/Plugins/Runtime/MassEntity/Source/MassEntity/Public/MassEntitySettings.h
@@ -39,7 +39,7 @@ struct FMassProcessingPhaseConfig
 	UMassCompositeProcessor* PhaseProcessor = nullptr;
 
 	UPROPERTY(VisibleAnywhere, Category = Mass, Transient)
-	FString Description;
+	FText Description;
 #endif //WITH_EDITORONLY_DATA
 };
 
