Improved EntityTemplate's and Archetype's debug descriptions

diff --git a/Engine/Plugins/Runtime/MassEntity/Source/MassEntity/Private/MassArchetypeData.cpp b/Engine/Plugins/Runtime/MassEntity/Source/MassEntity/Private/MassArchetypeData.cpp
--- a/Engine/Plugins/Runtime/MassEntity/Source/MassEntity/Private/MassArchetypeData.cpp
+++ b/Engine/Plugins/Runtime/MassEntity/Source/MassEntity/Private/MassArchetypeData.cpp
@@ -770,23 +770,16 @@ SIZE_T FMassArchetypeData::GetAllocatedSize() const
 FString FMassArchetypeData::DebugGetDescription() const
 {
 #if WITH_MASSENTITY_DEBUG
-	TStringBuilder<256> ArchetypeDebugName;
+	FStringOutputDevice OutDescription;
+
+	OutDescription += TEXT("Chunk fragments: ");
+	CompositionDescriptor.ChunkFragments.DebugGetStringDesc(OutDescription);
+	OutDescription += TEXT("\nTags: ");
+	CompositionDescriptor.Tags.DebugGetStringDesc(OutDescription);
+	OutDescription += TEXT("\nFragments: ");
+	CompositionDescriptor.Fragments.DebugGetStringDesc(OutDescription);
 	
-	ArchetypeDebugName.Append(TEXT("<"));
-
-	bool bNeedsComma = false;
-	for (const FMassArchetypeFragmentConfig& FragmentConfig : FragmentConfigs)
-	{
-		if (bNeedsComma)
-		{
-			ArchetypeDebugName.Append(TEXT(","));
-		}
-		ArchetypeDebugName.Append(FragmentConfig.FragmentType->GetName());
-		bNeedsComma = true;
-	}
-
-	ArchetypeDebugName.Append(TEXT(">"));
-	return ArchetypeDebugName.ToString();
+	return static_cast<FString>(OutDescription);
 #else
 	return {};
 #endif
diff --git a/Engine/Plugins/Runtime/MassGameplay/Source/MassSpawner/Private/MassEntityTemplate.cpp b/Engine/Plugins/Runtime/MassGameplay/Source/MassSpawner/Private/MassEntityTemplate.cpp
--- a/Engine/Plugins/Runtime/MassGameplay/Source/MassSpawner/Private/MassEntityTemplate.cpp
+++ b/Engine/Plugins/Runtime/MassGameplay/Source/MassSpawner/Private/MassEntityTemplate.cpp
@@ -35,17 +35,20 @@ FString FMassEntityTemplate::DebugGetDescription(UMassEntitySubsystem* EntitySub
 #if WITH_MASSGAMEPLAY_DEBUG
 	Ar.SetAutoEmitLineTerminator(true);
 
-	Ar += TEXT("Fragments:\n");
-	FragmentCollection.DebugOutputDescription(Ar);
-
-	Ar += TEXT("\nProcessors:\n");
-	InitializationPipeline.DebugOutputDescription(Ar);
-
 	if (EntitySubsystem)
 	{
-		Ar += TEXT("Archetype detais:\n");
+		Ar += TEXT("Archetype details:\n");
 		Ar += DebugGetArchetypeDescription(*EntitySubsystem);
 	}
+	else
+	{
+		Ar += TEXT("Fragments:\n");
+		FragmentCollection.DebugOutputDescription(Ar);
+	}
+
+	Ar += TEXT("\nInitialization Processors:\n");
+	InitializationPipeline.DebugOutputDescription(Ar);
+	
 #endif // WITH_MASSGAMEPLAY_DEBUG
 	return MoveTemp(Ar);
 }
