Added a way to override the composite processor used for a given Mass processing phase

diff --git a/Engine/Plugins/Runtime/MassEntity/Source/MassEntity/Private/MassProcessingPhase.cpp b/Engine/Plugins/Runtime/MassEntity/Source/MassEntity/Private/MassProcessingPhase.cpp
--- a/Engine/Plugins/Runtime/MassEntity/Source/MassEntity/Private/MassProcessingPhase.cpp
+++ b/Engine/Plugins/Runtime/MassEntity/Source/MassEntity/Private/MassProcessingPhase.cpp
@@ -24,6 +24,19 @@ namespace FMassTweakables
 	};
 }
 
+namespace UE::MassEntity::Internal  
+{
+	ETickingGroup PhaseToTickingGroup[int(EMassProcessingPhase::MAX)]
+	{
+		ETickingGroup::TG_PrePhysics, // EMassProcessingPhase::PrePhysics
+		ETickingGroup::TG_StartPhysics, // EMassProcessingPhase::StartPhysics
+		ETickingGroup::TG_DuringPhysics, // EMassProcessingPhase::DuringPhysics
+		ETickingGroup::TG_EndPhysics,	// EMassProcessingPhase::EndPhysics
+		ETickingGroup::TG_PostPhysics,	// EMassProcessingPhase::PostPhysics
+		ETickingGroup::TG_LastDemotable, // EMassProcessingPhase::FrameEnd
+	};
+} // UE::MassEntity::Internal
+
 //----------------------------------------------------------------------//
 //  FMassProcessingPhase
 //----------------------------------------------------------------------//
@@ -195,14 +208,32 @@ void UMassProcessingPhaseManager::Start(UMassEntitySubsystem& InEntitySubsystem)
 void UMassProcessingPhaseManager::EnableTickFunctions(const UWorld& World)
 {
 	check(EntitySubsystem);
+
+	const bool bIsGameWorld = World.IsGameWorld();
+
 	for (FMassProcessingPhase& Phase : ProcessingPhases)
 	{
 		Phase.Manager = this;
 		Phase.RegisterTickFunction(World.PersistentLevel);
 		Phase.SetTickFunctionEnable(true);
+#if WITH_MASSENTITY_DEBUG
+		if (Phase.PhaseProcessor && bIsGameWorld)
+		{
+			// not logging this in the editor mode since it messes up the game-recorded vislog display (with its progressively larger timestamp)
+			FStringOutputDevice Ar;
+			Phase.PhaseProcessor->DebugOutputDescription(Ar);
+			UE_VLOG_UELOG(this, LogMass, Log, TEXT("Enabling phase %s tick:\n%s")
+				, *EnumToString(Phase.Phase), *Ar);
+		}
+#endif // WITH_MASSENTITY_DEBUG		
+	}
+
+	if (bIsGameWorld)
+	{
+		// not logging this in the editor mode since it messes up the game-recorded vislog display (with its progressively larger timestamp)
+		UE_VLOG_UELOG(this, LogMass, Log, TEXT("MassProcessingPhaseManager %s.%s has been started")
+			, *GetNameSafe(GetOuter()), *GetName());
 	}
-	UE_VLOG_UELOG(this, LogMass, Log, TEXT("UMassProcessingPhaseManager %s.%s has been started")
-		, *GetNameSafe(GetOuter()), *GetName());
 }
 
 void UMassProcessingPhaseManager::Stop()
@@ -214,32 +245,57 @@ void UMassProcessingPhaseManager::Stop()
 		Phase.SetTickFunctionEnable(false);
 	}
 
-	UE_VLOG_UELOG(this, LogMass, Log, TEXT("UMassProcessingPhaseManager %s.%s has been stopped")
-		, *GetNameSafe(GetOuter()), *GetName());
+	UWorld* World = GetWorld();
+	if (World && World->IsGameWorld())
+	{
+		// not logging this in editor mode since it messes up the game-recorded vislog display (with its progressively larger timestamp) 
+		UE_VLOG_UELOG(this, LogMass, Log, TEXT("MassProcessingPhaseManager %s.%s has been stopped")
+			, *GetNameSafe(GetOuter()), *GetName());
+	}
 }
 
-void UMassProcessingPhaseManager::CreatePhases() 
+void UMassProcessingPhaseManager::SetPhaseProcessor(const EMassProcessingPhase Phase, UMassCompositeProcessor* PhaseProcessor)
 {
-	static ETickingGroup PhaseToTickingGroup[int(EMassProcessingPhase::MAX)]
+	if (Phase == EMassProcessingPhase::MAX)
 	{
-		ETickingGroup::TG_PrePhysics, // EMassProcessingPhase::PrePhysics
-		ETickingGroup::TG_StartPhysics, // EMassProcessingPhase::StartPhysics
-		ETickingGroup::TG_DuringPhysics, // EMassProcessingPhase::DuringPhysics
-		ETickingGroup::TG_EndPhysics,	// EMassProcessingPhase::EndPhysics
-		ETickingGroup::TG_PostPhysics,	// EMassProcessingPhase::PostPhysics
-		ETickingGroup::TG_LastDemotable, // EMassProcessingPhase::FrameEnd
-	};
+		UE_VLOG_UELOG(this, LogMass, Error, TEXT("MassProcessingPhaseManager::OverridePhaseProcessor called with Phase == MAX"));
+		return;
+	}
+	
+	UE_VLOG_UELOG(this, LogMass, Log, TEXT("%s phase processor for phase %s"), PhaseProcessor ? TEXT("Overriding") : TEXT("Nulling-out")
+		, *EnumToString(Phase));
 
+	// note that it's ok to use PhaseProcessor == nullptr
+	ProcessingPhases[int32(Phase)].PhaseProcessor = PhaseProcessor;
+	
+	if (PhaseProcessor)
+	{
+		REDIRECT_OBJECT_TO_VLOG(PhaseProcessor, this);
+		PhaseProcessor->SetProcessingPhase(Phase);
+		PhaseProcessor->SetGroupName(FName(FString::Printf(TEXT("%s-%s"), *PhaseProcessor->GetName(), *EnumToString(Phase))));
+		
+		check(GetOuter());
+		PhaseProcessor->Initialize(*GetOuter());
+
+#if WITH_MASSENTITY_DEBUG
+		FStringOutputDevice Ar;
+		PhaseProcessor->DebugOutputDescription(Ar);
+		UE_VLOG(this, LogMass, Log, TEXT("Setting new group processor for phase %s:\n%s")
+			, *EnumToString(Phase), *Ar);
+#endif // WITH_MASSENTITY_DEBUG
+	}
+}
+
+void UMassProcessingPhaseManager::CreatePhases() 
+{
 	// @todo copy from settings instead of blindly creating from scratch
 	for (int i = 0; i < int(EMassProcessingPhase::MAX); ++i)
 	{
 		ProcessingPhases[i].Phase = EMassProcessingPhase(i);
-		ProcessingPhases[i].TickGroup = PhaseToTickingGroup[i];
-		ProcessingPhases[i].PhaseProcessor = NewObject<UMassCompositeProcessor>(this, UMassCompositeProcessor::StaticClass()
+		ProcessingPhases[i].TickGroup = UE::MassEntity::Internal::PhaseToTickingGroup[i];
+		UMassCompositeProcessor* PhaseProcessor = NewObject<UMassCompositeProcessor>(this, UMassCompositeProcessor::StaticClass()
 			, *FString::Printf(TEXT("ProcessingPhase_%s"), *EnumToString(EMassProcessingPhase(i))));
-		REDIRECT_OBJECT_TO_VLOG(ProcessingPhases[i].PhaseProcessor, this);
-		ProcessingPhases[i].PhaseProcessor->SetGroupName(*EnumToString(EMassProcessingPhase(i)));
-		ProcessingPhases[i].PhaseProcessor->SetProcessingPhase(EMassProcessingPhase(i));
+		SetPhaseProcessor(EMassProcessingPhase(i), PhaseProcessor);
 	}
 }
 
diff --git a/Engine/Plugins/Runtime/MassEntity/Source/MassEntity/Public/MassProcessingPhase.h b/Engine/Plugins/Runtime/MassEntity/Source/MassEntity/Public/MassProcessingPhase.h
--- a/Engine/Plugins/Runtime/MassEntity/Source/MassEntity/Public/MassProcessingPhase.h
+++ b/Engine/Plugins/Runtime/MassEntity/Source/MassEntity/Public/MassProcessingPhase.h
@@ -110,6 +110,12 @@ public:
 	 */
 	bool IsDuringMassProcessing() const { return CurrentPhase != EMassProcessingPhase::MAX && ProcessingPhases[int(CurrentPhase)].IsDuringMassProcessing(); }
 
+	/**
+	 *  Sets the composite processor used for Phase processing phase. Using PhaseProcessor == nullptr results in
+	 *  clearing the Phase's processor out effectively making the Phase perform no calculations. 
+	 */
+	void SetPhaseProcessor(const EMassProcessingPhase Phase, UMassCompositeProcessor* PhaseProcessor);
+
 protected:
 	virtual void PostInitProperties() override;
 	virtual void BeginDestroy() override;
