Massaging UMassProcessingPhaseManager::SetPhaseProcessor's logging

diff --git a/Engine/Plugins/Runtime/MassEntity/Source/MassEntity/Private/MassProcessingPhase.cpp b/Engine/Plugins/Runtime/MassEntity/Source/MassEntity/Private/MassProcessingPhase.cpp
--- a/Engine/Plugins/Runtime/MassEntity/Source/MassEntity/Private/MassProcessingPhase.cpp
+++ b/Engine/Plugins/Runtime/MassEntity/Source/MassEntity/Private/MassProcessingPhase.cpp
@@ -262,9 +262,6 @@ void UMassProcessingPhaseManager::SetPhaseProcessor(const EMassProcessingPhase P
 		return;
 	}
 	
-	UE_VLOG_UELOG(this, LogMass, Log, TEXT("%s phase processor for phase %s"), PhaseProcessor ? TEXT("Overriding") : TEXT("Nulling-out")
-		, *EnumToString(Phase));
-
 	// note that it's ok to use PhaseProcessor == nullptr
 	ProcessingPhases[int32(Phase)].PhaseProcessor = PhaseProcessor;
 	
@@ -284,6 +281,10 @@ void UMassProcessingPhaseManager::SetPhaseProcessor(const EMassProcessingPhase P
 			, *EnumToString(Phase), *Ar);
 #endif // WITH_MASSENTITY_DEBUG
 	}
+	else
+	{
+		UE_VLOG_UELOG(this, LogMass, Log, TEXT("Nulling-out phase processor for phase %s"), *EnumToString(Phase));
+	}
 }
 
 void UMassProcessingPhaseManager::CreatePhases() 
