Added support for optional chunk fragments in EntityQueries

diff --git a/Engine/Plugins/Runtime/MassEntity/Source/MassEntity/Private/MassArchetypeData.cpp b/Engine/Plugins/Runtime/MassEntity/Source/MassEntity/Private/MassArchetypeData.cpp
--- a/Engine/Plugins/Runtime/MassEntity/Source/MassEntity/Private/MassArchetypeData.cpp
+++ b/Engine/Plugins/Runtime/MassEntity/Source/MassEntity/Private/MassArchetypeData.cpp
@@ -676,7 +676,7 @@ void FMassArchetypeData::GetRequirementsChunkFragmentMapping(TConstArrayView<FMa
 				}
 			}
 
-			check(FragmentIndex != INDEX_NONE);
+			check(FragmentIndex != INDEX_NONE || Requirement.IsOptional());
 			OutFragmentIndices.Add(FragmentIndex);
 			LastFoundFragmentIndex = FragmentIndex;
 		}
@@ -742,8 +742,8 @@ void FMassArchetypeData::BindChunkFragmentRequirements(FMassExecutionContext& Ru
 			FMassExecutionContext::FChunkFragmentView& ChunkRequirement = RunContext.ChunkFragments[i];
 			const int32 ChunkFragmentIndex = ChunkFragmentsMapping[i];
 
-			check(ChunkFragmentIndex != INDEX_NONE);
-			ChunkRequirement.ChunkFragmentView = Chunk.GetMutableChunkFragmentViewChecked(ChunkFragmentIndex);
+			check(ChunkFragmentIndex != INDEX_NONE || ChunkRequirement.Requirement.IsOptional());
+			ChunkRequirement.ChunkFragmentView = ChunkFragmentIndex != INDEX_NONE ? Chunk.GetMutableChunkFragmentViewChecked(ChunkFragmentIndex) : FStructView();
 		}
 	}
 	else
@@ -751,8 +751,8 @@ void FMassArchetypeData::BindChunkFragmentRequirements(FMassExecutionContext& Ru
 		for (FMassExecutionContext::FChunkFragmentView& ChunkRequirement : RunContext.GetMutableChunkRequirements())
 		{
 			FInstancedStruct* ChunkFragmentInstance = Chunk.FindMutableChunkFragment(ChunkRequirement.Requirement.StructType);
-			check(ChunkFragmentInstance != nullptr);
-			ChunkRequirement.ChunkFragmentView = FStructView(*ChunkFragmentInstance);
+			check(ChunkFragmentInstance != nullptr || ChunkRequirement.Requirement.IsOptional());
+			ChunkRequirement.ChunkFragmentView = ChunkFragmentInstance ? FStructView(*ChunkFragmentInstance) : FStructView(); 
 		}
 	}
 }
diff --git a/Engine/Plugins/Runtime/MassEntity/Source/MassEntity/Public/MassEntityQuery.h b/Engine/Plugins/Runtime/MassEntity/Source/MassEntity/Public/MassEntityQuery.h
--- a/Engine/Plugins/Runtime/MassEntity/Source/MassEntity/Public/MassEntityQuery.h
+++ b/Engine/Plugins/Runtime/MassEntity/Source/MassEntity/Public/MassEntityQuery.h
@@ -235,17 +235,21 @@ public:
 		static_assert(TIsDerivedFrom<T, FMassChunkFragment>::IsDerived, "Given struct doesn't represent a valid chunk fragment type. Make sure to inherit from FMassChunkFragment or one of its child-types.");
 		checkf(ChunkRequirements.FindByPredicate([](const FMassFragmentRequirement& Item) { return Item.StructType == T::StaticStruct(); }) == nullptr
 			, TEXT("Duplicated requirements are not supported. %s already present"), *T::StaticStruct()->GetName());
-		checkfSlow(Presence == EMassFragmentPresence::None || Presence == EMassFragmentPresence::All
-			, TEXT("The only valid Presence values for AddChunkRequirement are All and None"));
-		
-		if (Presence == EMassFragmentPresence::All)
-		{
-			RequiredAllChunkFragments.Add<T>();
-			ChunkRequirements.Emplace(T::StaticStruct(), AccessMode, Presence);
-		}
-		else
+		checkfSlow(Presence != EMassFragmentPresence::Any, TEXT("\'Any\' is not a valid Presence value for AddChunkRequirement."));
+
+		switch (Presence)
 		{
-			RequiredNoneChunkFragments.Add<T>();
+			case EMassFragmentPresence::All:
+				RequiredAllChunkFragments.Add<T>();
+				ChunkRequirements.Emplace(T::StaticStruct(), AccessMode, Presence);
+				break;
+			case EMassFragmentPresence::Optional:
+				RequiredOptionalChunkFragments.Add<T>();
+				ChunkRequirements.Emplace(T::StaticStruct(), AccessMode, Presence);
+				break;
+			case EMassFragmentPresence::None:
+				RequiredNoneChunkFragments.Add<T>();
+				break;
 		}
 
 		return *this;
@@ -282,6 +286,7 @@ public:
 	const FMassTagBitSet& GetRequiredAnyTags() const { return RequiredAnyTags; }
 	const FMassTagBitSet& GetRequiredNoneTags() const { return RequiredNoneTags; }
 	const FMassChunkFragmentBitSet& GetRequiredAllChunkFragments() const { return RequiredAllChunkFragments; }
+	const FMassChunkFragmentBitSet& GetRequiredOptionalChunkFragments() const { return RequiredOptionalChunkFragments; }
 	const FMassChunkFragmentBitSet& GetRequiredNoneChunkFragments() const { return RequiredNoneChunkFragments; }
 
 	const TArray<FArchetypeHandle>& GetArchetypes() const
@@ -329,6 +334,7 @@ protected:
 	FMassFragmentBitSet RequiredOptionalFragments;
 	FMassFragmentBitSet RequiredNoneFragments;
 	FMassChunkFragmentBitSet RequiredAllChunkFragments;
+	FMassChunkFragmentBitSet RequiredOptionalChunkFragments;
 	FMassChunkFragmentBitSet RequiredNoneChunkFragments;
 
 private:
diff --git a/Engine/Plugins/Runtime/MassEntity/Source/MassEntity/Public/MassEntitySubsystem.h b/Engine/Plugins/Runtime/MassEntity/Source/MassEntity/Public/MassEntitySubsystem.h
--- a/Engine/Plugins/Runtime/MassEntity/Source/MassEntity/Public/MassEntitySubsystem.h
+++ b/Engine/Plugins/Runtime/MassEntity/Source/MassEntity/Public/MassEntitySubsystem.h
@@ -265,7 +265,7 @@ public:
 	void DebugRemoveAllEntities();
 	void DebugForceArchetypeDataVersionBump() { ++ArchetypeDataVersion; }
 	void DebugGetArchetypeStrings(const FArchetypeHandle& Archetype, TArray<FName>& OutFragmentNames, TArray<FName>& OutTagNames);
-	FMassEntityHandle DebugGetEntityIndexHandle(const int32 EntityIndex) const { return Entities.IsValidIndex(EntityIndex) ? FMassEntityHandle({EntityIndex, Entities[EntityIndex].SerialNumber}) : FMassEntityHandle(); }
+	FMassEntityHandle DebugGetEntityIndexHandle(const int32 EntityIndex) const { return Entities.IsValidIndex(EntityIndex) ? FMassEntityHandle(EntityIndex, Entities[EntityIndex].SerialNumber) : FMassEntityHandle(); }
 #endif // WITH_MASSENTITY_DEBUG
 
 protected:
@@ -416,27 +416,41 @@ public:
 	int32 GetChunkSerialModificationNumber() const { return ChunkSerialModificationNumber; }
 
 	template<typename T>
-	T& GetMutableChunkFragment() const
+	T* GetOptionalMutableChunkFragment() const
 	{
 		static_assert(TIsDerivedFrom<T, FMassChunkFragment>::IsDerived, "Given struct doesn't represent a valid chunk fragment type. Make sure to inherit from FMassChunkFragment or one of its child-types.");
 
 		const UScriptStruct* Type = T::StaticStruct();
-		const FChunkFragmentView* FoundChunkFragmentData = ChunkFragments.FindByPredicate([Type](const FChunkFragmentView& Element) { return Element.Requirement.StructType == Type; });
-		checkf(FoundChunkFragmentData, TEXT("Chunk Fragment requirement not found: %s"), *T::StaticStruct()->GetName());
-		return FoundChunkFragmentData->ChunkFragmentView.GetMutable<T>();
+		const FChunkFragmentView* FoundChunkFragmentData = ChunkFragments.FindByPredicate([Type](const FChunkFragmentView& Element) { return Element.Requirement.StructType == Type; } );
+		return FoundChunkFragmentData ? FoundChunkFragmentData->ChunkFragmentView.GetMutablePtr<T>() : static_cast<T*>(nullptr);
+	}
+	
+	template<typename T>
+	T& GetMutableChunkFragment() const
+	{
+		T* ChunkFragment = GetOptionalMutableChunkFragment<T>();
+		checkf(ChunkFragment, TEXT("Chunk Fragment requirement not found: %s"), *T::StaticStruct()->GetName());
+		return *ChunkFragment;
 	}
 
 	template<typename T>
-	const T& GetChunkFragment() const
+	const T* GetOptionalChunkFragment() const
 	{
 		static_assert(TIsDerivedFrom<T, FMassChunkFragment>::IsDerived, "Given struct doesn't represent a valid chunk fragment type. Make sure to inherit from FMassChunkFragment or one of its child-types.");
 
 		const UScriptStruct* Type = T::StaticStruct();
 		const FChunkFragmentView* FoundChunkFragmentData = ChunkFragments.FindByPredicate([Type](const FChunkFragmentView& Element) { return Element.Requirement.StructType == Type; } );
-		checkf(FoundChunkFragmentData, TEXT("Chunk Fragment requirement not found: %s"), *T::StaticStruct()->GetName());
-		return FoundChunkFragmentData->ChunkFragmentView.Get<T>();
+		return FoundChunkFragmentData ? FoundChunkFragmentData->ChunkFragmentView.GetPtr<T>() : static_cast<const T*>(nullptr);
 	}
-
+	
+	template<typename T>
+	const T& GetChunkFragment() const
+	{
+		const T* ChunkFragment = GetOptionalChunkFragment<T>();
+		checkf(ChunkFragment, TEXT("Chunk Fragment requirement not found: %s"), *T::StaticStruct()->GetName());
+		return *ChunkFragment;
+	}
+	
 	template<typename TFragment>
 	TArrayView<TFragment> GetMutableFragmentView()
 	{
