CIS fix

diff --git a/Engine/Plugins/Runtime/MassEntity/Source/MassEntity/Public/MassEntityTypes.h b/Engine/Plugins/Runtime/MassEntity/Source/MassEntity/Public/MassEntityTypes.h
--- a/Engine/Plugins/Runtime/MassEntity/Source/MassEntity/Public/MassEntityTypes.h
+++ b/Engine/Plugins/Runtime/MassEntity/Source/MassEntity/Public/MassEntityTypes.h
@@ -42,7 +42,13 @@ struct FMassEntityHandle
 {
 	GENERATED_BODY()
 
-	FMassEntityHandle() = default;
+	FMassEntityHandle()
+	{
+	}
+	FMassEntityHandle(const int32 InIndex, const int32 InSerialNumber)
+		: Index(InIndex), SerialNumber(InSerialNumber)
+	{
+	}
 	
 	UPROPERTY(VisibleAnywhere, Category = "Mass|Debug", Transient)
 	int32 Index = 0;
