Refactored how Mass modules' developer settings are organized. Introduced a dedicated class, UMassModuleSettings, that per-module setting classes extend now to automatically get registered with main MassSettings and show under a single "Mass" entry in Project Settings

diff --git a/Engine/Plugins/AI/MassAI/Source/MassAIMovement/MassAIMovement.Build.cs b/Engine/Plugins/AI/MassAI/Source/MassAIMovement/MassAIMovement.Build.cs
--- a/Engine/Plugins/AI/MassAI/Source/MassAIMovement/MassAIMovement.Build.cs
+++ b/Engine/Plugins/AI/MassAI/Source/MassAIMovement/MassAIMovement.Build.cs
@@ -38,7 +38,7 @@ namespace UnrealBuildTool.Rules
 					"StructUtils",
 					"ZoneGraph",
 					"ZoneGraphAnnotations",
-					"DeveloperSettings",
+					"DeveloperSettings"
 				}
 			);
 
diff --git a/Engine/Plugins/AI/MassAI/Source/MassAIMovement/Public/MassAvoidanceSettings.h b/Engine/Plugins/AI/MassAI/Source/MassAIMovement/Public/MassAvoidanceSettings.h
--- a/Engine/Plugins/AI/MassAI/Source/MassAIMovement/Public/MassAvoidanceSettings.h
+++ b/Engine/Plugins/AI/MassAI/Source/MassAIMovement/Public/MassAvoidanceSettings.h
@@ -3,7 +3,7 @@
 #pragma once
 
 #include "CoreMinimal.h"
-#include "Engine/DeveloperSettings.h"
+#include "MassSettings.h"
 #include "MassAvoidanceSettings.generated.h"
 
 USTRUCT()
@@ -24,8 +24,8 @@ struct MASSAIMOVEMENT_API FMassAvoidanceVelocityFilter
 	float HighSpeed = 50.f;
 };
 
-UCLASS(config = Mass, defaultconfig, meta = (DisplayName = "Mass Avoidance"))
-class MASSAIMOVEMENT_API UMassAvoidanceSettings : public UDeveloperSettings
+UCLASS(config = Mass, defaultconfig, meta = (DisplayName = "Mass Avoidance"), dontCollapseCategories)
+class MASSAIMOVEMENT_API UMassAvoidanceSettings : public UMassModuleSettings
 {
 	GENERATED_BODY()
 
diff --git a/Engine/Plugins/AI/MassAI/Source/MassAIMovement/Public/MassMovementSettings.h b/Engine/Plugins/AI/MassAI/Source/MassAIMovement/Public/MassMovementSettings.h
--- a/Engine/Plugins/AI/MassAI/Source/MassAIMovement/Public/MassMovementSettings.h
+++ b/Engine/Plugins/AI/MassAI/Source/MassAIMovement/Public/MassMovementSettings.h
@@ -3,12 +3,12 @@
 #pragma once
 
 #include "CoreMinimal.h"
-#include "Engine/DeveloperSettings.h"
+#include "MassSettings.h"
 #include "MassAIMovementTypes.h"
 #include "MassMovementSettings.generated.h"
 
 UCLASS(config = Mass, defaultconfig, meta = (DisplayName = "Mass Movement"))
-class MASSAIMOVEMENT_API UMassMovementSettings : public UDeveloperSettings
+class MASSAIMOVEMENT_API UMassMovementSettings : public UMassModuleSettings
 {
 	GENERATED_BODY()
 
diff --git a/Engine/Plugins/AI/MassAI/Source/MassAIMovementEditor/Public/MassMovementConfigRefDetails.cpp b/Engine/Plugins/AI/MassAI/Source/MassAIMovementEditor/Public/MassMovementConfigRefDetails.cpp
--- a/Engine/Plugins/AI/MassAI/Source/MassAIMovementEditor/Public/MassMovementConfigRefDetails.cpp
+++ b/Engine/Plugins/AI/MassAI/Source/MassAIMovementEditor/Public/MassMovementConfigRefDetails.cpp
@@ -8,6 +8,7 @@
 #include "Widgets/Input/SComboButton.h"
 #include "Framework/MultiBox/MultiBoxBuilder.h"
 #include "ScopedTransaction.h"
+#include "MassSettings.h"
 #include "MassMovementSettings.h"
 #include "Modules/ModuleManager.h"
 #include "ISettingsModule.h"
@@ -52,17 +53,20 @@ void FMassMovementConfigRefDetails::CustomizeChildren(TSharedRef<class IProperty
 
 void FMassMovementConfigRefDetails::OnProfileComboChange(int32 Idx)
 {
-	const UMassMovementSettings* Settings = GetDefault<UMassMovementSettings>();
-	check(Settings);
-
 	if (Idx == -1)
 	{
+		const UMassSettings* MassSettings = GetDefault<UMassSettings>();
+		check(MassSettings);
+		
 		// Goto settings to create new Profile
-		FModuleManager::LoadModuleChecked<ISettingsModule>("Settings").ShowViewer(Settings->GetContainerName(), Settings->GetCategoryName(), Settings->GetSectionName());
+		FModuleManager::LoadModuleChecked<ISettingsModule>("Settings").ShowViewer(MassSettings->GetContainerName(), MassSettings->GetCategoryName(), MassSettings->GetSectionName());
 		return;
 	}
 
-	TConstArrayView<FMassMovementConfig> MovementConfigs = Settings->GetMovementConfigs(); 
+	const UMassMovementSettings* MovementSettings = GetDefault<UMassMovementSettings>();
+	check(MovementSettings);
+
+	TConstArrayView<FMassMovementConfig> MovementConfigs = MovementSettings->GetMovementConfigs();
 	if (MovementConfigs.IsValidIndex(Idx))
 	{
 		const FMassMovementConfig& Config = MovementConfigs[Idx];
diff --git a/Engine/Plugins/AI/MassAI/Source/MassAIMovementEditor/Public/MassMovementStyleRefDetails.cpp b/Engine/Plugins/AI/MassAI/Source/MassAIMovementEditor/Public/MassMovementStyleRefDetails.cpp
--- a/Engine/Plugins/AI/MassAI/Source/MassAIMovementEditor/Public/MassMovementStyleRefDetails.cpp
+++ b/Engine/Plugins/AI/MassAI/Source/MassAIMovementEditor/Public/MassMovementStyleRefDetails.cpp
@@ -8,6 +8,7 @@
 #include "Widgets/Input/SComboButton.h"
 #include "Framework/MultiBox/MultiBoxBuilder.h"
 #include "ScopedTransaction.h"
+#include "MassSettings.h"
 #include "MassMovementSettings.h"
 #include "Modules/ModuleManager.h"
 #include "ISettingsModule.h"
@@ -52,17 +53,20 @@ void FMassMovementStyleRefDetails::CustomizeChildren(TSharedRef<class IPropertyH
 
 void FMassMovementStyleRefDetails::OnProfileComboChange(int32 Idx)
 {
-	const UMassMovementSettings* Settings = GetDefault<UMassMovementSettings>();
-	check(Settings);
-
 	if (Idx == -1)
 	{
+		const UMassSettings* MassSettings = GetDefault<UMassSettings>();
+		check(MassSettings);
+		
 		// Goto settings to create new Profile
-		FModuleManager::LoadModuleChecked<ISettingsModule>("Settings").ShowViewer(Settings->GetContainerName(), Settings->GetCategoryName(), Settings->GetSectionName());
+		FModuleManager::LoadModuleChecked<ISettingsModule>("Settings").ShowViewer(MassSettings->GetContainerName(), MassSettings->GetCategoryName(), MassSettings->GetSectionName());
 		return;
 	}
 
-	TConstArrayView<FMassMovementStyle> MovementStyles = Settings->GetMovementStyles(); 
+	const UMassMovementSettings* MovementSettings = GetDefault<UMassMovementSettings>();
+	check(MovementSettings);
+
+	TConstArrayView<FMassMovementStyle> MovementStyles = MovementSettings->GetMovementStyles();
 	if (MovementStyles.IsValidIndex(Idx))
 	{
 		const FMassMovementStyle& Style = MovementStyles[Idx];
diff --git a/Engine/Plugins/AI/MassCrowd/Source/MassCrowd/MassCrowd.Build.cs b/Engine/Plugins/AI/MassCrowd/Source/MassCrowd/MassCrowd.Build.cs
--- a/Engine/Plugins/AI/MassCrowd/Source/MassCrowd/MassCrowd.Build.cs
+++ b/Engine/Plugins/AI/MassCrowd/Source/MassCrowd/MassCrowd.Build.cs
@@ -27,7 +27,6 @@ namespace UnrealBuildTool.Rules
 					"MassEntity",
 					"Core",
 					"CoreUObject",
-					"DeveloperSettings",
 					"Engine",
 					"MassActors",
 					"MassCommon",
diff --git a/Engine/Plugins/AI/MassCrowd/Source/MassCrowd/Public/MassCrowdSettings.h b/Engine/Plugins/AI/MassCrowd/Source/MassCrowd/Public/MassCrowdSettings.h
--- a/Engine/Plugins/AI/MassCrowd/Source/MassCrowd/Public/MassCrowdSettings.h
+++ b/Engine/Plugins/AI/MassCrowd/Source/MassCrowd/Public/MassCrowdSettings.h
@@ -2,7 +2,7 @@
 
 #pragma once
 
-#include "Engine/DeveloperSettings.h"
+#include "MassSettings.h"
 #include "ZoneGraphTypes.h"
 #include "MassCrowdSettings.generated.h"
 
@@ -46,8 +46,8 @@ struct MASSCROWD_API FMassCrowdLaneDensityDesc
 /**
  * Settings for the MassCrowd plugin.
  */
-UCLASS(config = Plugins, defaultconfig, DisplayName = "Mass Crowd")
-class MASSCROWD_API UMassCrowdSettings : public UDeveloperSettings
+UCLASS(config = Plugins, defaultconfig, DisplayName = "Mass Crowd", dontCollapseCategories)
+class MASSCROWD_API UMassCrowdSettings : public UMassModuleSettings
 {
 	GENERATED_BODY()
 
diff --git a/Engine/Plugins/Runtime/MassEntity/Source/MassEntity/Private/MassEntitySettings.cpp b/Engine/Plugins/Runtime/MassEntity/Source/MassEntity/Private/MassEntitySettings.cpp
--- a/Engine/Plugins/Runtime/MassEntity/Source/MassEntity/Private/MassEntitySettings.cpp
+++ b/Engine/Plugins/Runtime/MassEntity/Source/MassEntity/Private/MassEntitySettings.cpp
@@ -11,9 +11,9 @@
 #endif // WITH_EDITOR
 
 //----------------------------------------------------------------------//
-//  UMassSettings
+//  UMassEntitySettings
 //----------------------------------------------------------------------//
-UMassSettings::UMassSettings(const FObjectInitializer& ObjectInitializer)
+UMassEntitySettings::UMassEntitySettings(const FObjectInitializer& ObjectInitializer)
 	: Super(ObjectInitializer)
 {
 	for (int i = 0; i < (int)EMassProcessingPhase::MAX; ++i)
@@ -21,16 +21,16 @@ UMassSettings::UMassSettings(const FObjectInitializer& ObjectInitializer)
 		ProcessingPhasesConfig[i].PhaseName = *EnumToString(EMassProcessingPhase(i));
 	}
 
-	FCoreDelegates::OnPostEngineInit.AddUObject(this, &UMassSettings::BuildProcessorListAndPhases);
+	FCoreDelegates::OnPostEngineInit.AddUObject(this, &UMassEntitySettings::BuildProcessorListAndPhases);
 }
 
-void UMassSettings::BeginDestroy()
+void UMassEntitySettings::BeginDestroy()
 {
 	FCoreDelegates::OnPostEngineInit.RemoveAll(this);
 	Super::BeginDestroy();
 }
 
-void UMassSettings::BuildProcessorListAndPhases()
+void UMassEntitySettings::BuildProcessorListAndPhases()
 {
 	if (bInitialized)
 	{
@@ -42,7 +42,7 @@ void UMassSettings::BuildProcessorListAndPhases()
 	bInitialized = true;
 }
 
-void UMassSettings::BuildPhases()
+void UMassEntitySettings::BuildPhases()
 {
 #if WITH_EDITOR
 	for (int i = 0; i < int(EMassProcessingPhase::MAX); ++i)
@@ -62,7 +62,7 @@ void UMassSettings::BuildPhases()
 #endif // WITH_EDITOR
 }
 
-void UMassSettings::BuildProcessorList()
+void UMassEntitySettings::BuildProcessorList()
 {
 	ProcessorCDOs.Reset();
 	for (FMassProcessingPhaseConfig& PhaseConfig : ProcessingPhasesConfig)
@@ -101,7 +101,7 @@ void UMassSettings::BuildProcessorList()
 	});
 }
 
-void UMassSettings::AddToActiveProcessorsList(TSubclassOf<UMassProcessor> ProcessorClass)
+void UMassEntitySettings::AddToActiveProcessorsList(TSubclassOf<UMassProcessor> ProcessorClass)
 {
 	if (UMassProcessor* ProcessorCDO = GetMutableDefault<UMassProcessor>(ProcessorClass))
 	{
@@ -126,16 +126,16 @@ void UMassSettings::AddToActiveProcessorsList(TSubclassOf<UMassProcessor> Proces
 	}
 }
 
-const FMassProcessingPhaseConfig* UMassSettings::GetProcessingPhasesConfig()
+const FMassProcessingPhaseConfig* UMassEntitySettings::GetProcessingPhasesConfig()
 {
 	BuildProcessorListAndPhases();
 	return ProcessingPhasesConfig;
 }
 
 #if WITH_EDITOR
-void UMassSettings::PostEditChangeProperty(FPropertyChangedEvent& PropertyChangedEvent)
+void UMassEntitySettings::PostEditChangeProperty(FPropertyChangedEvent& PropertyChangedEvent)
 {
-	static const FName ProcessorCDOsName = GET_MEMBER_NAME_CHECKED(UMassSettings, ProcessorCDOs);
+	static const FName ProcessorCDOsName = GET_MEMBER_NAME_CHECKED(UMassEntitySettings, ProcessorCDOs);
 
 	Super::PostEditChangeProperty(PropertyChangedEvent);
 
@@ -158,7 +158,7 @@ void UMassSettings::PostEditChangeProperty(FPropertyChangedEvent& PropertyChange
 	}
 }
 
-void UMassSettings::PostEditChangeChainProperty(struct FPropertyChangedChainEvent& PropertyChangedEvent)
+void UMassEntitySettings::PostEditChangeChainProperty(struct FPropertyChangedChainEvent& PropertyChangedEvent)
 {
 	static const FName AutoRegisterName = TEXT("bAutoRegisterWithProcessingPhases");
 
diff --git a/Engine/Plugins/Runtime/MassEntity/Source/MassEntity/Private/MassProcessingPhase.cpp b/Engine/Plugins/Runtime/MassEntity/Source/MassEntity/Private/MassProcessingPhase.cpp
--- a/Engine/Plugins/Runtime/MassEntity/Source/MassEntity/Private/MassProcessingPhase.cpp
+++ b/Engine/Plugins/Runtime/MassEntity/Source/MassEntity/Private/MassProcessingPhase.cpp
@@ -120,9 +120,9 @@ void UMassProcessingPhaseManager::PostInitProperties()
 		UWorld* World = GetWorld();
 		if (World && World->IsGameWorld() == false)
 		{
-			UMassSettings* MassSettings = GetMutableDefault<UMassSettings>();
-			check(MassSettings);
-			MassSettingsChangeHandle = MassSettings->GetOnSettingsChange().AddUObject(this, &UMassProcessingPhaseManager::OnMassSettingsChange);
+			UMassEntitySettings* Settings = GetMutableDefault<UMassEntitySettings>();
+			check(Settings);
+			MassEntitySettingsChangeHandle = Settings->GetOnSettingsChange().AddUObject(this, &UMassProcessingPhaseManager::OnMassEntitySettingsChange);
 		}
 #endif // WITH_EDITOR
 	}
@@ -135,9 +135,9 @@ void UMassProcessingPhaseManager::BeginDestroy()
 	Super::BeginDestroy();
 
 #if WITH_EDITOR
-	if (UMassSettings* MassSettings = GetMutableDefault<UMassSettings>())
+	if (UMassEntitySettings* Settings = GetMutableDefault<UMassEntitySettings>())
 	{
-		MassSettings->GetOnSettingsChange().Remove(MassSettingsChangeHandle);
+		Settings->GetOnSettingsChange().Remove(MassEntitySettingsChangeHandle);
 	}
 #endif // WITH_EDITOR
 }
@@ -150,10 +150,10 @@ void UMassProcessingPhaseManager::InitializePhases(UObject& InProcessorOwner)
 
 #if WITH_EDITOR
 	const UWorld* World = InProcessorOwner.GetWorld();
-	const UMassSettings* MassSettings = GetMutableDefault<UMassSettings>();
-	if (World != nullptr && MassSettings != nullptr && !MassSettings->DumpDependencyGraphFileName.IsEmpty())
+	const UMassEntitySettings* Settings = GetMutableDefault<UMassEntitySettings>();
+	if (World != nullptr && Settings != nullptr && !Settings->DumpDependencyGraphFileName.IsEmpty())
 	{
-		DependencyGraphFileName = FString::Printf(TEXT("%s_%s"), *MassSettings->DumpDependencyGraphFileName,*ToString(World->GetNetMode()));
+		DependencyGraphFileName = FString::Printf(TEXT("%s_%s"), *Settings->DumpDependencyGraphFileName,*ToString(World->GetNetMode()));
 	}
 #endif // WITH_EDITOR
 
@@ -326,7 +326,7 @@ void UMassProcessingPhaseManager::OnPhaseEnd(FMassProcessingPhase& Phase)
 }
 
 #if WITH_EDITOR
-void UMassProcessingPhaseManager::OnMassSettingsChange(const FPropertyChangedEvent& PropertyChangedEvent)
+void UMassProcessingPhaseManager::OnMassEntitySettingsChange(const FPropertyChangedEvent& PropertyChangedEvent)
 {
 	check(GetOuter());
 	InitializePhases(*GetOuter());
diff --git a/Engine/Plugins/Runtime/MassEntity/Source/MassEntity/Private/MassSettings.cpp b/Engine/Plugins/Runtime/MassEntity/Source/MassEntity/Private/MassSettings.cpp
new file mode 100644
--- /dev/null
+++ b/Engine/Plugins/Runtime/MassEntity/Source/MassEntity/Private/MassSettings.cpp
@@ -0,0 +1,43 @@
+// Copyright Epic Games, Inc. All Rights Reserved.
+
+#include "MassSettings.h"
+
+
+//----------------------------------------------------------------------//
+//  UMassModuleSettings
+//----------------------------------------------------------------------//
+void UMassModuleSettings::PostInitProperties()
+{
+	Super::PostInitProperties();
+
+	if (HasAnyFlags(RF_ClassDefaultObject) && !GetClass()->HasAnyClassFlags(CLASS_Abstract))
+	{
+		// register with UMassGameplaySettings
+		GetMutableDefault<UMassSettings>()->RegisterModuleSettings(*this);
+	}
+}
+
+//----------------------------------------------------------------------//
+//  UMassSettings
+//----------------------------------------------------------------------//
+void UMassSettings::RegisterModuleSettings(UMassModuleSettings& SettingsCDO)
+{
+	ensureMsgf(SettingsCDO.HasAnyFlags(RF_ClassDefaultObject), TEXT("Registered ModuleSettings need to be its class's CDO"));
+
+	// we should consider a replacement in case we're hot-reloading
+	FName EntryName = SettingsCDO.GetClass()->GetFName();
+
+#if WITH_EDITOR
+	static const FName DisplayNameMeta(TEXT("DisplayName")); 
+	// try reading better name from meta data, available only in editor. Besides, we don't really care about this out 
+	// side of editor. We could even skip populating ModuleSettings but we'll leave it as is for now.
+	const FString& DisplayNameValue = SettingsCDO.GetClass()->GetMetaData(DisplayNameMeta);
+	if (DisplayNameValue.Len())
+	{
+		EntryName = *DisplayNameValue;
+	}
+#endif // WITH_EDITOR
+
+	UMassModuleSettings*& FoundModuleEntry = ModuleSettings.FindOrAdd(EntryName, &SettingsCDO);
+	FoundModuleEntry = &SettingsCDO;
+}
diff --git a/Engine/Plugins/Runtime/MassEntity/Source/MassEntity/Public/MassEntitySettings.h b/Engine/Plugins/Runtime/MassEntity/Source/MassEntity/Public/MassEntitySettings.h
--- a/Engine/Plugins/Runtime/MassEntity/Source/MassEntity/Public/MassEntitySettings.h
+++ b/Engine/Plugins/Runtime/MassEntity/Source/MassEntity/Public/MassEntitySettings.h
@@ -2,13 +2,13 @@
 
 #pragma once
 
-#include "Engine/DeveloperSettings.h"
+#include "MassSettings.h"
 #include "MassProcessingPhase.h"
 #include "MassProcessor.h"
 #include "InstancedStruct.h"
 #include "MassEntitySettings.generated.h"
 
-#define GET_MASS_CONFIG_VALUE(a) (GetMutableDefault<UMassSettings>()->a)
+#define GET_MASS_CONFIG_VALUE(a) (GetMutableDefault<UMassEntitySettings>()->a)
 
 class UMassProcessingPhaseManager;
 struct FMassProcessingPhaseConfig;
@@ -46,8 +46,8 @@ struct FMassProcessingPhaseConfig
 /**
  * Implements the settings for MassEntity plugin
  */
-UCLASS(config = Mass, defaultconfig, DisplayName = "Mass")
-class MASSENTITY_API UMassSettings : public UDeveloperSettings
+UCLASS(config = Mass, defaultconfig, DisplayName = "Mass Entity")
+class MASSENTITY_API UMassEntitySettings : public UMassModuleSettings
 {
 	GENERATED_BODY()
 public:
@@ -55,7 +55,7 @@ public:
 	DECLARE_MULTICAST_DELEGATE_OneParam(FOnSettingsChange, const FPropertyChangedEvent& /*PropertyChangedEvent*/);
 #endif // WITH_EDITORONLY_DATA
 
-	UMassSettings(const FObjectInitializer& ObjectInitializer = FObjectInitializer::Get());
+	UMassEntitySettings(const FObjectInitializer& ObjectInitializer = FObjectInitializer::Get());
 
 	void BuildProcessorListAndPhases();
 	void AddToActiveProcessorsList(TSubclassOf<UMassProcessor> ProcessorClass);
@@ -90,7 +90,7 @@ public:
 	FMassProcessingPhaseConfig ProcessingPhasesConfig[(uint8)EMassProcessingPhase::MAX];
 
 	/** This list contains all the processors available in the given binary (including plugins). The contents are sorted by display name.*/
-	UPROPERTY(VisibleAnywhere, Category = Mass, Transient, Instanced)
+	UPROPERTY(VisibleAnywhere, Category = Mass, Transient, Instanced, EditFixedSize)
 	TArray<UMassProcessor*> ProcessorCDOs;
 
 #if WITH_EDITORONLY_DATA
diff --git a/Engine/Plugins/Runtime/MassEntity/Source/MassEntity/Public/MassProcessingPhase.h b/Engine/Plugins/Runtime/MassEntity/Source/MassEntity/Public/MassProcessingPhase.h
--- a/Engine/Plugins/Runtime/MassEntity/Source/MassEntity/Public/MassProcessingPhase.h
+++ b/Engine/Plugins/Runtime/MassEntity/Source/MassEntity/Public/MassProcessingPhase.h
@@ -88,7 +88,7 @@ public:
 	FMassProcessingPhase::FOnPhaseEvent& GetOnPhaseEnd(const EMassProcessingPhase Phase) { return ProcessingPhases[uint8(Phase)].OnPhaseEnd; }
 
 	/** 
-	 *  Populates hosted FMassProcessingPhase instances with Processors read from MassSettings configuration.
+	 *  Populates hosted FMassProcessingPhase instances with Processors read from MassEntitySettings configuration.
 	 *  Calling this function overrides previous configuration of Phases.
 	 */
 	void InitializePhases(UObject& InProcessorOwner);
@@ -121,11 +121,11 @@ protected:
 	virtual void BeginDestroy() override;
 	void EnableTickFunctions(const UWorld& World);
 
-	/** Creates phase processors instances for each declared phase name, based on MassSettings */
+	/** Creates phase processors instances for each declared phase name, based on MassEntitySettings */
 	virtual void CreatePhases();
 
 #if WITH_EDITOR
-	virtual void OnMassSettingsChange(const FPropertyChangedEvent& PropertyChangedEvent);
+	virtual void OnMassEntitySettingsChange(const FPropertyChangedEvent& PropertyChangedEvent);
 #endif // WITH_EDITOR
 
 	friend FMassProcessingPhase;
@@ -152,6 +152,6 @@ protected:
 	EMassProcessingPhase CurrentPhase = EMassProcessingPhase::MAX;
 
 #if WITH_EDITOR
-	FDelegateHandle MassSettingsChangeHandle;
+	FDelegateHandle MassEntitySettingsChangeHandle;
 #endif // WITH_EDITOR
 };
diff --git a/Engine/Plugins/Runtime/MassEntity/Source/MassEntity/Public/MassSettings.h b/Engine/Plugins/Runtime/MassEntity/Source/MassEntity/Public/MassSettings.h
new file mode 100644
--- /dev/null
+++ b/Engine/Plugins/Runtime/MassEntity/Source/MassEntity/Public/MassSettings.h
@@ -0,0 +1,33 @@
+// Copyright Epic Games, Inc. All Rights Reserved.
+
+#pragma once
+
+#include "Engine/DeveloperSettings.h"
+#include "MassSettings.generated.h"
+
+
+/** 
+ * A common parrent for Mass's per-module settings. Classes extending this class will automatically get registered 
+ * with- and show under Mass settings in Project Settings.
+ */
+UCLASS(Abstract, config = Mass, defaultconfig, collapseCategories)
+class MASSENTITY_API UMassModuleSettings : public UObject
+{
+	GENERATED_BODY()
+protected:
+	virtual void PostInitProperties() override;
+};
+
+
+UCLASS(config = Mass, defaultconfig, DisplayName = "Mass", AutoExpandCategories = "Mass")
+class MASSENTITY_API UMassSettings : public UDeveloperSettings
+{
+	GENERATED_BODY()
+
+public:
+	void RegisterModuleSettings(UMassModuleSettings& SettingsCDO);
+
+public:
+	UPROPERTY(VisibleAnywhere, Category = "Mass", NoClear, EditFixedSize, meta = (EditInline))
+	TMap<FName, UMassModuleSettings*> ModuleSettings;
+};
diff --git a/Engine/Plugins/Runtime/MassGameplay/Source/MassMovement/MassMovement.Build.cs b/Engine/Plugins/Runtime/MassGameplay/Source/MassMovement/MassMovement.Build.cs
--- a/Engine/Plugins/Runtime/MassGameplay/Source/MassMovement/MassMovement.Build.cs
+++ b/Engine/Plugins/Runtime/MassGameplay/Source/MassMovement/MassMovement.Build.cs
@@ -36,7 +36,6 @@ namespace UnrealBuildTool.Rules
 					"StructUtils",
 					"ZoneGraph",
 					"ZoneGraphAnnotations",
-					"DeveloperSettings",
 				}
 			);
 
diff --git a/Engine/Plugins/Runtime/MassGameplay/Source/MassSimulation/MassSimulation.Build.cs b/Engine/Plugins/Runtime/MassGameplay/Source/MassSimulation/MassSimulation.Build.cs
--- a/Engine/Plugins/Runtime/MassGameplay/Source/MassSimulation/MassSimulation.Build.cs
+++ b/Engine/Plugins/Runtime/MassGameplay/Source/MassSimulation/MassSimulation.Build.cs
@@ -29,7 +29,6 @@ namespace UnrealBuildTool.Rules
 					"MassCommon",
 					"NavigationSystem",
 					"MassEntity",
-					"DeveloperSettings",
 					"StructUtils",
 				}
 			);
diff --git a/Engine/Plugins/Runtime/MassGameplay/Source/MassSimulation/Public/MassSimulationSettings.h b/Engine/Plugins/Runtime/MassGameplay/Source/MassSimulation/Public/MassSimulationSettings.h
--- a/Engine/Plugins/Runtime/MassGameplay/Source/MassSimulation/Public/MassSimulationSettings.h
+++ b/Engine/Plugins/Runtime/MassGameplay/Source/MassSimulation/Public/MassSimulationSettings.h
@@ -2,7 +2,7 @@
 
 #pragma once
 
-#include "Engine/DeveloperSettings.h"
+#include "MassSettings.h"
 #include "MassSimulationSettings.generated.h"
 
 #define GET_MASS_CONFIG_VALUE(a) (GetMutableDefault<UMassSimulationSettings>()->a)
@@ -13,7 +13,7 @@ class UMassSchematic;
  * Implements the settings for MassSimulation
  */
 UCLASS(config = Mass, defaultconfig, DisplayName = "Mass Simulation")
-class MASSSIMULATION_API UMassSimulationSettings : public UDeveloperSettings
+class MASSSIMULATION_API UMassSimulationSettings : public UMassModuleSettings
 {
 	GENERATED_BODY()
 public:
diff --git a/Engine/Plugins/Runtime/MassGameplay/Source/MassSmartObjects/MassSmartObjects.Build.cs b/Engine/Plugins/Runtime/MassGameplay/Source/MassSmartObjects/MassSmartObjects.Build.cs
--- a/Engine/Plugins/Runtime/MassGameplay/Source/MassSmartObjects/MassSmartObjects.Build.cs
+++ b/Engine/Plugins/Runtime/MassGameplay/Source/MassSmartObjects/MassSmartObjects.Build.cs
@@ -25,7 +25,6 @@ namespace UnrealBuildTool.Rules
 					"MassEntity",
 					"Core",
 					"CoreUObject",
-					"DeveloperSettings",
 					"Engine",
 					"GameplayTags",
 					"MassCommon",
diff --git a/Engine/Plugins/Runtime/MassGameplay/Source/MassSmartObjects/Public/MassSmartObjectSettings.h b/Engine/Plugins/Runtime/MassGameplay/Source/MassSmartObjects/Public/MassSmartObjectSettings.h
--- a/Engine/Plugins/Runtime/MassGameplay/Source/MassSmartObjects/Public/MassSmartObjectSettings.h
+++ b/Engine/Plugins/Runtime/MassGameplay/Source/MassSmartObjects/Public/MassSmartObjectSettings.h
@@ -2,7 +2,7 @@
 
 #pragma once
 #include "ZoneGraphTypes.h"
-#include "Engine/DeveloperSettings.h"
+#include "MassSettings.h"
 #include "MassSmartObjectSettings.generated.h"
 
 #if WITH_EDITOR
@@ -14,7 +14,7 @@ DECLARE_MULTICAST_DELEGATE(FOnAnnotationSettingsChanged);
  * Settings for the MassSmartObject module.
  */
 UCLASS(config = Plugins, defaultconfig, DisplayName = "Mass SmartObject")
-class MASSSMARTOBJECTS_API UMassSmartObjectSettings : public UDeveloperSettings
+class MASSSMARTOBJECTS_API UMassSmartObjectSettings : public UMassModuleSettings
 {
 	GENERATED_BODY()
 
diff --git a/Engine/Plugins/Runtime/MassGameplay/Source/MassSpawner/MassSpawner.Build.cs b/Engine/Plugins/Runtime/MassGameplay/Source/MassSpawner/MassSpawner.Build.cs
--- a/Engine/Plugins/Runtime/MassGameplay/Source/MassSpawner/MassSpawner.Build.cs
+++ b/Engine/Plugins/Runtime/MassGameplay/Source/MassSpawner/MassSpawner.Build.cs
@@ -29,7 +29,6 @@ namespace UnrealBuildTool.Rules
 					"MassEntity",
 					"MassCommon",
 					"MassSimulation",
-					"DeveloperSettings",
 					"StructUtils",
 					"ZoneGraph",
 				}
