Added console command to output Mass Entity system total memory usage

diff --git a/Engine/Plugins/Runtime/MassEntity/Source/MassEntity/Private/MassEntityDebug.cpp b/Engine/Plugins/Runtime/MassEntity/Source/MassEntity/Private/MassEntityDebug.cpp
--- a/Engine/Plugins/Runtime/MassEntity/Source/MassEntity/Private/MassEntityDebug.cpp
+++ b/Engine/Plugins/Runtime/MassEntity/Source/MassEntity/Private/MassEntityDebug.cpp
@@ -143,6 +143,22 @@ FAutoConsoleCommandWithWorldArgsAndOutputDevice LogFragmentSizes(
 			}
 		})
 	);
+
+FAutoConsoleCommandWithWorldArgsAndOutputDevice LogMemoryUsage(
+	TEXT("mass.LogMemoryUsage"),
+	TEXT("Logs how much memory the mass entity system uses"),
+	FConsoleCommandWithWorldArgsAndOutputDeviceDelegate::CreateLambda([](const TArray<FString>& Params, UWorld* World, FOutputDevice& Ar)
+{
+	check(World);
+	if (UMassEntitySubsystem* System = World->GetSubsystem<UMassEntitySubsystem>())
+	{
+		FResourceSizeEx CumulativeResourceSize;
+		System->GetResourceSizeEx(CumulativeResourceSize);
+		Ar.Logf(ELogVerbosity::Log, TEXT("MassEntity system uses: %d bytes"), CumulativeResourceSize.GetDedicatedSystemMemoryBytes());
+	}
+}));
+
+
 #endif // WITH_MASSENTITY_DEBUG && WITH_MASSENTITY_DEBUG
 
 } // namespace UE::Mass::Debug
