Fixes chunk template default value not being consistently initialized when creating an archetype from and other by forcing them in Archetype Creation API
At the same time, simplified Archetype creation API
Removed unuseful duplicated information "TArray<const UScriptStruct*> FragmentTypes" which was a duplicate of "FMassUniqueFragmentCollection.Fragments" in the Entity templates

#rnx

diff --git a/Engine/Plugins/Runtime/MassEntity/Source/MassEntity/Private/MassArchetypeData.cpp b/Engine/Plugins/Runtime/MassEntity/Source/MassEntity/Private/MassArchetypeData.cpp
--- a/Engine/Plugins/Runtime/MassEntity/Source/MassEntity/Private/MassArchetypeData.cpp
+++ b/Engine/Plugins/Runtime/MassEntity/Source/MassEntity/Private/MassArchetypeData.cpp
@@ -43,10 +43,10 @@ bool FMassArchetypeData::IsEquivalent(TConstArrayView<const UScriptStruct*> Othe
 	return true;
 }
 
-void FMassArchetypeData::Initialize(const FMassFragmentBitSet& Fragments, const FMassTagBitSet& Tags, const FMassChunkFragmentBitSet& ChunkFragments)
+void FMassArchetypeData::Initialize(const FMassCompositionDescriptor& InCompositionDescriptor, const FMassArchetypeFragmentsInitialValues& InInitialValues)
 {
 	TArray<const UScriptStruct*, TInlineAllocator<16>> SortedFragmentList;
-	Fragments.ExportTypes(SortedFragmentList);
+	InCompositionDescriptor.Fragments.ExportTypes(SortedFragmentList);
 
 	SortedFragmentList.Sort(FMassSorterOperator<UScriptStruct>());
 
@@ -78,16 +78,24 @@ void FMassArchetypeData::Initialize(const FMassFragmentBitSet& Fragments, const
 		checkSlow(Fragments.IsEquivalent(CompositionDescriptor.Fragments));
 	}
 
-	CompositionDescriptor.Tags = Tags;
-	CompositionDescriptor.ChunkFragments = ChunkFragments;
+	CompositionDescriptor.Tags = InCompositionDescriptor.Tags;
+	CompositionDescriptor.ChunkFragments = InCompositionDescriptor.ChunkFragments;
 	// create the chunk fragment instances
 	TArray<const UScriptStruct*, TInlineAllocator<16>> ChunkFragmentList;
-	ChunkFragments.ExportTypes(ChunkFragmentList);
+	CompositionDescriptor.ChunkFragments.ExportTypes(ChunkFragmentList);
 	ChunkFragmentList.Sort(FMassSorterOperator<UScriptStruct>());
 	for (const UScriptStruct* ChunkFragmentType : ChunkFragmentList)
 	{
 		check(ChunkFragmentType);
-		ChunkFragmentTemplates.Add(FInstancedStruct(ChunkFragmentType));
+		if (const FInstancedStruct* Struct = InInitialValues.ChunkFragments.FindByPredicate( 
+				[ChunkFragmentType](const FInstancedStruct& Struct) { return Struct.GetScriptStruct() == ChunkFragmentType; }))
+		{
+			InitialValues.ChunkFragments.Add(*Struct);
+		}
+		else
+		{
+			InitialValues.ChunkFragments.Add(FInstancedStruct(ChunkFragmentType));
+		}
 	}
 
 	TotalBytesPerEntity = FragmentSizeTallyBytes;
@@ -117,7 +125,7 @@ void FMassArchetypeData::InitializeWithSibling(const FMassArchetypeData& Sibling
 	CompositionDescriptor.Fragments = SiblingArchetype.CompositionDescriptor.Fragments;
 	CompositionDescriptor.Tags = OverrideTags;
 	CompositionDescriptor.ChunkFragments = SiblingArchetype.CompositionDescriptor.ChunkFragments;
-	ChunkFragmentTemplates = SiblingArchetype.ChunkFragmentTemplates;
+	InitialValues.ChunkFragments = SiblingArchetype.InitialValues.ChunkFragments;
 
 	TotalBytesPerEntity = SiblingArchetype.TotalBytesPerEntity;
 	NumEntitiesPerChunk = SiblingArchetype.NumEntitiesPerChunk;
@@ -173,12 +181,12 @@ int32 FMassArchetypeData::AddEntityInternal(FMassEntityHandle Entity, const bool
 		if (EmptyChunkIndex != INDEX_NONE)
 		{
 			DestinationChunk = &Chunks[EmptyChunkIndex];
-			DestinationChunk->Recycle(ChunkFragmentTemplates);
+			DestinationChunk->Recycle(InitialValues.ChunkFragments);
 			AbsoluteIndex = EmptyAbsoluteIndex;
 		}
 		else
 		{
-			DestinationChunk = &Chunks.Emplace_GetRef(GetChunkAllocSize(), ChunkFragmentTemplates);
+			DestinationChunk = &Chunks.Emplace_GetRef(GetChunkAllocSize(), InitialValues.ChunkFragments);
 		}
 
 		check(DestinationChunk);
@@ -432,17 +440,6 @@ void FMassArchetypeData::SetFragmentData(const FArchetypeChunkCollection& ChunkC
 	}
 }
 
-void FMassArchetypeData::SetDefaultChunkFragmentValue(FConstStructView InstancedStruct)
-{
-	if (ensure(InstancedStruct.GetScriptStruct()) && CompositionDescriptor.ChunkFragments.Contains(*InstancedStruct.GetScriptStruct()))
-	{
-		if (FInstancedStruct* Element = ChunkFragmentTemplates.FindByPredicate([&InstancedStruct](const FInstancedStruct& Element) { return Element.GetScriptStruct() == InstancedStruct.GetScriptStruct(); }))
-		{
-			*Element = InstancedStruct;
-		}
-	}
-}
-
 void FMassArchetypeData::MoveEntityToAnotherArchetype(const FMassEntityHandle Entity, FMassArchetypeData& NewArchetype)
 {
 	check(&NewArchetype != this);
@@ -668,9 +665,9 @@ void FMassArchetypeData::GetRequirementsChunkFragmentMapping(TConstArrayView<FMa
 		if (Requirement.RequiresBinding())
 		{
 			int32 FragmentIndex = INDEX_NONE;
-			for (int32 i = LastFoundFragmentIndex + 1; i < ChunkFragmentTemplates.Num(); ++i)
+			for (int32 i = LastFoundFragmentIndex + 1; i < InitialValues.ChunkFragments.Num(); ++i)
 			{
-				if (ChunkFragmentTemplates[i].GetScriptStruct()->IsChildOf(Requirement.StructType))
+				if (InitialValues.ChunkFragments[i].GetScriptStruct()->IsChildOf(Requirement.StructType))
 				{
 					FragmentIndex = i;
 					break;
@@ -770,7 +767,7 @@ SIZE_T FMassArchetypeData::GetAllocatedSize() const
 	}
 
 	return sizeof(FMassArchetypeData) +
-		ChunkFragmentTemplates.GetAllocatedSize() +
+		InitialValues.ChunkFragments.GetAllocatedSize() +
 		FragmentConfigs.GetAllocatedSize() +
 		Chunks.GetAllocatedSize() +
 		(NumAllocatedChunkBuffers * GetChunkAllocSize()) +
diff --git a/Engine/Plugins/Runtime/MassEntity/Source/MassEntity/Private/MassArchetypeData.h b/Engine/Plugins/Runtime/MassEntity/Source/MassEntity/Private/MassArchetypeData.h
--- a/Engine/Plugins/Runtime/MassEntity/Source/MassEntity/Private/MassArchetypeData.h
+++ b/Engine/Plugins/Runtime/MassEntity/Source/MassEntity/Private/MassArchetypeData.h
@@ -137,8 +137,7 @@ struct FMassArchetypeData
 private:
 	// One-stop-shop variable describing the archetype's fragment and tag composition 
 	FMassCompositionDescriptor CompositionDescriptor;
-
-	TArray<FInstancedStruct> ChunkFragmentTemplates;
+	FMassArchetypeFragmentsInitialValues InitialValues;
 
 	TArray<FMassArchetypeFragmentConfig, TInlineAllocator<16>> FragmentConfigs;
 	
@@ -165,6 +164,8 @@ public:
 	const FMassChunkFragmentBitSet& GetChunkFragmentBitSet() const { return CompositionDescriptor.ChunkFragments; }
 	const FMassCompositionDescriptor& GetCompositionDescriptor() const { return CompositionDescriptor; }
 
+	const FMassArchetypeFragmentsInitialValues& GetInitialValues() const { return InitialValues; }
+
 	/** Method to iterate on all the fragment types */
 	void ForEachFragmentType(TFunction< void(const UScriptStruct* /*FragmentType*/)> Function) const;
 	bool HasFragmentType(const UScriptStruct* FragmentType) const;
@@ -179,7 +180,7 @@ public:
 		return CompositionDescriptor.IsEquivalent(OtherCompositionDescriptor);
 	}
 
-	void Initialize(const FMassFragmentBitSet& Fragments, const FMassTagBitSet& Tags, const FMassChunkFragmentBitSet& ChunkFragments);
+	void Initialize(const FMassCompositionDescriptor& InCompositionDescriptor, const FMassArchetypeFragmentsInitialValues& InInitialValues);
 
 	/** 
 	 * A special way of initializing an archetype resulting in a copy of SiblingArchetype's setup with OverrideTags
@@ -232,8 +233,6 @@ public:
 	 *  FragmentSource.GetScriptStruct to the value represented by FragmentSource.GetMemory */
 	void SetFragmentData(const FArchetypeChunkCollection& ChunkCollection, const FInstancedStruct& FragmentSource);
 
-	void SetDefaultChunkFragmentValue(FConstStructView InstancedStruct);
-
 	/** Returns conversion from given Requirements to archetype's fragment indices */
 	void GetRequirementsFragmentMapping(TConstArrayView<FMassFragmentRequirement> Requirements, FMassFragmentIndicesMapping& OutFragmentIndices);
 
diff --git a/Engine/Plugins/Runtime/MassEntity/Source/MassEntity/Private/MassEntitySubsystem.cpp b/Engine/Plugins/Runtime/MassEntity/Source/MassEntity/Private/MassEntitySubsystem.cpp
--- a/Engine/Plugins/Runtime/MassEntity/Source/MassEntity/Private/MassEntitySubsystem.cpp
+++ b/Engine/Plugins/Runtime/MassEntity/Source/MassEntity/Private/MassEntitySubsystem.cpp
@@ -105,12 +105,7 @@ FArchetypeHandle UMassEntitySubsystem::CreateArchetype(TConstArrayView<const USc
 		}
 	}
 
-	return CreateArchetype(FMassCompositionDescriptor(FMassFragmentBitSet(FragmentList), Tags, ChunkFragments));
-}
-
-FArchetypeHandle UMassEntitySubsystem::CreateArchetype(TConstArrayView<const UScriptStruct*> FragmentList, const FMassTagBitSet& Tags, const FMassChunkFragmentBitSet& ChunkFragments)
-{
-	return CreateArchetype(FMassCompositionDescriptor(FragmentList, Tags, ChunkFragments));
+	return CreateArchetype(FMassCompositionDescriptor(FMassFragmentBitSet(FragmentList), Tags, ChunkFragments), FMassArchetypeFragmentsInitialValues());
 }
 
 FArchetypeHandle UMassEntitySubsystem::CreateArchetype(const TSharedPtr<FMassArchetypeData>& SourceArchetype, const FMassFragmentBitSet& NewFragments)
@@ -118,10 +113,10 @@ FArchetypeHandle UMassEntitySubsystem::CreateArchetype(const TSharedPtr<FMassArc
 	check(SourceArchetype.IsValid());
 	checkf(NewFragments.IsEmpty() == false, TEXT("%s Adding an empty fragment list to an archetype is not supported."), ANSI_TO_TCHAR(__FUNCTION__));
 
-	return CreateArchetype(FMassCompositionDescriptor(NewFragments + SourceArchetype->GetFragmentBitSet(), SourceArchetype->GetTagBitSet(), SourceArchetype->GetChunkFragmentBitSet()));
+	return CreateArchetype(FMassCompositionDescriptor(NewFragments + SourceArchetype->GetFragmentBitSet(), SourceArchetype->GetTagBitSet(), SourceArchetype->GetChunkFragmentBitSet()), SourceArchetype->GetInitialValues());
 }
 
-FArchetypeHandle UMassEntitySubsystem::CreateArchetype(const FMassCompositionDescriptor& Composition)
+FArchetypeHandle UMassEntitySubsystem::CreateArchetype(const FMassCompositionDescriptor& Composition, const FMassArchetypeFragmentsInitialValues& InitialValues)
 {
 	const uint32 TypeHash = Composition.CalculateHash();
 
@@ -141,7 +136,7 @@ FArchetypeHandle UMassEntitySubsystem::CreateArchetype(const FMassCompositionDes
 	{
 		// Create a new archetype
 		TSharedPtr<FMassArchetypeData> NewArchetype = MakeShareable(new FMassArchetypeData);
-		NewArchetype->Initialize(Composition.Fragments, Composition.Tags, Composition.ChunkFragments);
+		NewArchetype->Initialize(Composition, InitialValues);
 		HashRow.Add(NewArchetype);
 
 		for (const FMassArchetypeFragmentConfig& FragmentConfig : NewArchetype->GetFragmentConfigs())
@@ -214,13 +209,6 @@ void UMassEntitySubsystem::ForEachArchetypeFragmentType(const FArchetypeHandle A
 	ArchetypeData.ForEachFragmentType(Function);
 }
 
-void UMassEntitySubsystem::SetDefaultChunkFragmentValue(const FArchetypeHandle Archetype, FConstStructView InstancedStruct)
-{
-	check(Archetype.DataPtr.IsValid());
-	FMassArchetypeData& ArchetypeData = *Archetype.DataPtr.Get();
-	ArchetypeData.SetDefaultChunkFragmentValue(InstancedStruct);
-}
-
 void UMassEntitySubsystem::DoEntityCompaction(const double TimeAllowed)
 {
 	const double TimeAllowedEnd = FPlatformTime::Seconds() + TimeAllowed;
@@ -258,7 +246,7 @@ FMassEntityHandle UMassEntitySubsystem::CreateEntity(TConstArrayView<FInstancedS
 {
 	check(FragmentInstanceList.Num() > 0);
 
-	FArchetypeHandle Archetype = CreateArchetype(FMassCompositionDescriptor(FragmentInstanceList, FMassTagBitSet(), FMassChunkFragmentBitSet()));
+	FArchetypeHandle Archetype = CreateArchetype(FMassCompositionDescriptor(FragmentInstanceList, FMassTagBitSet(), FMassChunkFragmentBitSet()), FMassArchetypeFragmentsInitialValues());
 	check(Archetype.DataPtr.IsValid());
 
 	FMassEntityHandle Entity = ReserveEntity();
@@ -301,7 +289,7 @@ void UMassEntitySubsystem::BuildEntity(FMassEntityHandle Entity, TConstArrayView
 	check(FragmentInstanceList.Num() > 0);
 	checkf(!IsEntityBuilt(Entity), TEXT("Expecting an entity that is not already built"));
 
-	FArchetypeHandle Archetype = CreateArchetype(FMassCompositionDescriptor(FragmentInstanceList, FMassTagBitSet(), FMassChunkFragmentBitSet()));
+	FArchetypeHandle Archetype = CreateArchetype(FMassCompositionDescriptor(FragmentInstanceList, FMassTagBitSet(), FMassChunkFragmentBitSet()), FMassArchetypeFragmentsInitialValues());
 	check(Archetype.DataPtr.IsValid());
 
 	InternalBuildEntity(Entity, Archetype);
@@ -428,7 +416,7 @@ void UMassEntitySubsystem::AddFragmentListToEntity(FMassEntityHandle Entity, TCo
 	InternalAddFragmentListToEntityChecked(Entity, FMassFragmentBitSet(FragmentList));
 }
 
-void UMassEntitySubsystem::AddCompositionToEntity_GetDelta(FMassEntityHandle Entity, FMassCompositionDescriptor& Composition)
+void UMassEntitySubsystem::AddCompositionToEntity_GetDelta(FMassEntityHandle Entity, FMassCompositionDescriptor& InDescriptor)
 {
 	CheckIfEntityIsActive(Entity);
 
@@ -436,16 +424,18 @@ void UMassEntitySubsystem::AddCompositionToEntity_GetDelta(FMassEntityHandle Ent
 	FMassArchetypeData* OldArchetype = EntityData.CurrentArchetype.Get();
 	check(OldArchetype);
 
-	Composition.Fragments -= OldArchetype->GetCompositionDescriptor().Fragments;
-	Composition.Tags -= OldArchetype->GetCompositionDescriptor().Tags;
+	InDescriptor.Fragments -= OldArchetype->GetCompositionDescriptor().Fragments;
+	InDescriptor.Tags -= OldArchetype->GetCompositionDescriptor().Tags;
+
+	ensureMsgf(InDescriptor.ChunkFragments.IsEmpty(), TEXT("Adding new chunk fragments is not supported"));
 
-	if (Composition.IsEmpty() == false)
+	if (InDescriptor.IsEmpty() == false)
 	{
 		FMassCompositionDescriptor NewDescriptor = OldArchetype->GetCompositionDescriptor();
-		NewDescriptor.Fragments += Composition.Fragments;
-		NewDescriptor.Tags+= Composition.Tags;
+		NewDescriptor.Fragments += InDescriptor.Fragments;
+		NewDescriptor.Tags+= InDescriptor.Tags;
 
-		const FArchetypeHandle NewArchetypeHandle = CreateArchetype(NewDescriptor);
+		const FArchetypeHandle NewArchetypeHandle = CreateArchetype(NewDescriptor, OldArchetype->GetInitialValues());
 
 		if (ensure(NewArchetypeHandle.DataPtr != EntityData.CurrentArchetype))
 		{
@@ -462,25 +452,30 @@ void UMassEntitySubsystem::RemoveCompositionFromEntity(FMassEntityHandle Entity,
 {
 	CheckIfEntityIsActive(Entity);
 
-	FEntityData& EntityData = Entities[Entity.Index];
-	FMassArchetypeData* OldArchetype = EntityData.CurrentArchetype.Get();
-	check(OldArchetype);
+	if(InDescriptor.IsEmpty() == false)
+	{
+		FEntityData& EntityData = Entities[Entity.Index];
+		FMassArchetypeData* OldArchetype = EntityData.CurrentArchetype.Get();
+		check(OldArchetype);
 
-	FMassCompositionDescriptor NewDescriptor = OldArchetype->GetCompositionDescriptor();
-	NewDescriptor.Fragments -= InDescriptor.Fragments;
-	NewDescriptor.Tags -= InDescriptor.Tags;
+		FMassCompositionDescriptor NewDescriptor = OldArchetype->GetCompositionDescriptor();
+		NewDescriptor.Fragments -= InDescriptor.Fragments;
+		NewDescriptor.Tags -= InDescriptor.Tags;
 
-	if (NewDescriptor.IsEmpty() == false)
-	{
-		const FArchetypeHandle NewArchetypeHandle = CreateArchetype(NewDescriptor);
+		ensureMsgf(InDescriptor.ChunkFragments.IsEmpty(), TEXT("Removing chunk fragments is not supported"));
 
-		if (ensure(NewArchetypeHandle.DataPtr != EntityData.CurrentArchetype))
+		if (NewDescriptor.IsEquivalent(InDescriptor) == false)
 		{
-			// Move the entity over
-			FMassArchetypeData* NewArchetype = NewArchetypeHandle.DataPtr.Get();
-			check(NewArchetype);
-			EntityData.CurrentArchetype->MoveEntityToAnotherArchetype(Entity, *NewArchetype);
-			EntityData.CurrentArchetype = NewArchetypeHandle.DataPtr;
+			const FArchetypeHandle NewArchetypeHandle = CreateArchetype(NewDescriptor, OldArchetype->GetInitialValues());
+
+			if (ensure(NewArchetypeHandle.DataPtr != EntityData.CurrentArchetype))
+			{
+				// Move the entity over
+				FMassArchetypeData* NewArchetype = NewArchetypeHandle.DataPtr.Get();
+				check(NewArchetype);
+				EntityData.CurrentArchetype->MoveEntityToAnotherArchetype(Entity, *NewArchetype);
+				EntityData.CurrentArchetype = NewArchetypeHandle.DataPtr;
+			}
 		}
 	}
 }
@@ -567,7 +562,7 @@ void UMassEntitySubsystem::RemoveFragmentListFromEntity(FMassEntityHandle Entity
 	if (ensureMsgf(OldArchetype->GetFragmentBitSet().HasAny(FragmentsToRemove), TEXT("Trying to remove a list of fragments from an entity but none of the fragments given was found.")))
 	{
 		// If all the fragments got removed this will result in fetching of the empty archetype
-		const FArchetypeHandle NewArchetypeHandle = CreateArchetype(FMassCompositionDescriptor(OldArchetype->GetFragmentBitSet() - FragmentsToRemove, OldArchetype->GetTagBitSet(), OldArchetype->GetChunkFragmentBitSet()));
+		const FArchetypeHandle NewArchetypeHandle = CreateArchetype(FMassCompositionDescriptor(OldArchetype->GetFragmentBitSet() - FragmentsToRemove, OldArchetype->GetTagBitSet(), OldArchetype->GetChunkFragmentBitSet()), OldArchetype->GetInitialValues());
 
 		// Move the entity over
 		FMassArchetypeData* NewArchetype = NewArchetypeHandle.DataPtr.Get();
diff --git a/Engine/Plugins/Runtime/MassEntity/Source/MassEntity/Public/MassEntitySubsystem.h b/Engine/Plugins/Runtime/MassEntity/Source/MassEntity/Public/MassEntitySubsystem.h
--- a/Engine/Plugins/Runtime/MassEntity/Source/MassEntity/Public/MassEntitySubsystem.h
+++ b/Engine/Plugins/Runtime/MassEntity/Source/MassEntity/Public/MassEntitySubsystem.h
@@ -82,7 +82,14 @@ public:
 	 */
 	FArchetypeHandle CreateArchetype(TConstArrayView<const UScriptStruct*> FragmentsAngTagsList);
 
-	FArchetypeHandle CreateArchetype(TConstArrayView<const UScriptStruct*> FragmentList, const FMassTagBitSet& Tags, const FMassChunkFragmentBitSet& ChunkFragments);
+	/**
+	 * CreateArchetype from a composition descriptor and initial values
+	 *
+	 * @param Composition of fragment, tag and chunk fragment types
+	 * @param InitialValues like default chunk fragment values
+	 * @return a handle of a new archetype 
+	 */
+	FArchetypeHandle CreateArchetype(const FMassCompositionDescriptor& Composition, const FMassArchetypeFragmentsInitialValues& InitialValues);
 
 	/** 
 	 *  Creates an archetype like SourceArchetype + NewFragmentList. 
@@ -96,15 +103,10 @@ public:
 	 */
 	FArchetypeHandle CreateArchetype(const TSharedPtr<FMassArchetypeData>& SourceArchetype, const FMassFragmentBitSet& NewFragmentList);
 
-	FArchetypeHandle CreateArchetype(const FMassCompositionDescriptor& Descriptor);
-
-
 	FArchetypeHandle GetArchetypeForEntity(FMassEntityHandle Entity) const;
 	/** Method to iterate on all the fragment types of an archetype */
 	static void ForEachArchetypeFragmentType(const FArchetypeHandle Archetype, TFunction< void(const UScriptStruct* /*FragmentType*/)> Function);
 
-	void SetDefaultChunkFragmentValue(const FArchetypeHandle Archetype, FConstStructView InstancedStruct);
-
 	/**
 	 * Go through all archetypes and compact entities
 	 * @param TimeAllowed to do entity compaction, once it reach that time it will stop and return
diff --git a/Engine/Plugins/Runtime/MassEntity/Source/MassEntity/Public/MassEntityTypes.h b/Engine/Plugins/Runtime/MassEntity/Source/MassEntity/Public/MassEntityTypes.h
--- a/Engine/Plugins/Runtime/MassEntity/Source/MassEntity/Public/MassEntityTypes.h
+++ b/Engine/Plugins/Runtime/MassEntity/Source/MassEntity/Public/MassEntityTypes.h
@@ -157,3 +157,18 @@ struct FMassCompositionDescriptor
 	FMassTagBitSet Tags;
 	FMassChunkFragmentBitSet ChunkFragments;
 };
+
+struct FMassArchetypeFragmentsInitialValues
+{
+	FMassArchetypeFragmentsInitialValues() = default;
+	FMassArchetypeFragmentsInitialValues(TConstArrayView<FInstancedStruct> InChunkFragments)
+	: ChunkFragments(InChunkFragments)
+	{
+		
+	}
+
+	// @todo add the regular fragments initial value here.
+	TArray<FInstancedStruct> ChunkFragments;
+
+	// @todo add the shared fragments values here
+};
\ No newline at end of file
diff --git a/Engine/Plugins/Runtime/MassGameplay/Source/MassCommon/Private/MassCommonTypes.cpp b/Engine/Plugins/Runtime/MassGameplay/Source/MassCommon/Private/MassCommonTypes.cpp
--- a/Engine/Plugins/Runtime/MassGameplay/Source/MassCommon/Private/MassCommonTypes.cpp
+++ b/Engine/Plugins/Runtime/MassGameplay/Source/MassCommon/Private/MassCommonTypes.cpp
@@ -121,31 +121,25 @@ FColor GetEntityDebugColor(FMassEntityHandle Entity)
 //----------------------------------------------------------------------//
 void FMassUniqueFragmentCollection::Add(const FInstancedStruct& InFragment)
 {
-	if (InFragment.GetScriptStruct() && FragmentInstances.FindByPredicate(FSameTypeScriptStructPredicate(InFragment)) == nullptr)
+	const UScriptStruct* InFragmentType = InFragment.GetScriptStruct();
+	if (InFragmentType && !FragmentBitSet.Contains(*InFragmentType))
 	{
+		ensureMsgf(InFragmentType->IsChildOf(FMassFragment::StaticStruct()), TEXT("Given struct is not of a fragment type. Make sure your struct extends FMassFragment or its derivatives"));
 		FragmentInstances.Add(InFragment);
+		FragmentBitSet.Add(*InFragmentType);
 	}
 }
 
 void FMassUniqueFragmentCollection::Add(const UScriptStruct* InFragmentType)
 {
-	if (InFragmentType && FragmentInstances.FindByPredicate(FSameTypeScriptStructPredicate(InFragmentType)) == nullptr)
+	if (InFragmentType && !FragmentBitSet.Contains(*InFragmentType))
 	{
+		ensureMsgf(InFragmentType->IsChildOf(FMassFragment::StaticStruct()), TEXT("Given struct is not of a fragment type. Make sure your struct extends FMassFragment or its derivatives"));
 		FragmentInstances.Add(FInstancedStruct(InFragmentType));
 		FragmentBitSet.Add(*InFragmentType);
 	}
 }
 
-void FMassUniqueFragmentCollection::Append(const FMassUniqueFragmentCollection& Other)
-{
-	FragmentInstances.Reserve(FragmentInstances.Num() + Other.FragmentInstances.Num());
-	for (const FInstancedStruct& NewFragment : Other.FragmentInstances)
-	{
-		Add(NewFragment);
-	}
-	// not shrinking on purpose, we can take the hit
-}
-
 void FMassUniqueFragmentCollection::DebugOutputDescription(FOutputDevice& Ar) const
 {
 #if WITH_MASSGAMEPLAY_DEBUG
diff --git a/Engine/Plugins/Runtime/MassGameplay/Source/MassCommon/Public/MassCommonTypes.h b/Engine/Plugins/Runtime/MassGameplay/Source/MassCommon/Public/MassCommonTypes.h
--- a/Engine/Plugins/Runtime/MassGameplay/Source/MassCommon/Public/MassCommonTypes.h
+++ b/Engine/Plugins/Runtime/MassGameplay/Source/MassCommon/Public/MassCommonTypes.h
@@ -88,8 +88,6 @@ struct MASSCOMMON_API FMassUniqueFragmentCollection
 		return FragmentInstances.Add_GetRef(FInstancedStruct(T::StaticStruct())).GetMutable<T>();
 	}
 
-	void Append(const FMassUniqueFragmentCollection& Other);
-
 	template<typename T>
 	T& GetRefChecked()
 	{
diff --git a/Engine/Plugins/Runtime/MassGameplay/Source/MassSpawner/Private/MassEntityTemplateRegistry.cpp b/Engine/Plugins/Runtime/MassGameplay/Source/MassSpawner/Private/MassEntityTemplateRegistry.cpp
--- a/Engine/Plugins/Runtime/MassGameplay/Source/MassSpawner/Private/MassEntityTemplateRegistry.cpp
+++ b/Engine/Plugins/Runtime/MassGameplay/Source/MassSpawner/Private/MassEntityTemplateRegistry.cpp
@@ -255,13 +255,8 @@ void UMassEntityTemplateRegistry::InitializeEntityTemplate(FMassEntityTemplate&
 	UMassEntitySubsystem* EntitySys = UWorld::GetSubsystem<UMassEntitySubsystem>(World);
 	check(EntitySys);
 
-	FTemplateRegistryHelpers::FragmentInstancesToTypes(OutTemplate.GetFragments(), OutTemplate.GetMutableFragmentTypes());
-	const FArchetypeHandle ArchetypeHandle = EntitySys->CreateArchetype(OutTemplate.GetFragmentTypes(), OutTemplate.GetTags(), FMassChunkFragmentBitSet(OutTemplate.GetChunkFragments()));
+	const FArchetypeHandle ArchetypeHandle = EntitySys->CreateArchetype(OutTemplate.GetCompositionDescriptor(), OutTemplate.GetArchetypeFragmentsInitialValues());
 	OutTemplate.SetArchetype(ArchetypeHandle);
-	for (const FInstancedStruct& InstancedStruct : OutTemplate.GetChunkFragments())
-	{
-		EntitySys->SetDefaultChunkFragmentValue(ArchetypeHandle, InstancedStruct);
-	}
 }
 
 void UMassEntityTemplateRegistry::DebugReset()
diff --git a/Engine/Plugins/Runtime/MassGameplay/Source/MassSpawner/Public/MassEntityTemplate.h b/Engine/Plugins/Runtime/MassGameplay/Source/MassSpawner/Public/MassEntityTemplate.h
--- a/Engine/Plugins/Runtime/MassGameplay/Source/MassSpawner/Public/MassEntityTemplate.h
+++ b/Engine/Plugins/Runtime/MassGameplay/Source/MassSpawner/Public/MassEntityTemplate.h
@@ -106,12 +106,11 @@ struct MASSSPAWNER_API FMassEntityTemplate
 
 	int32 GetFragmentsNum() const { return FragmentCollection.GetFragmentsNum(); }
 	TArrayView<const FInstancedStruct> GetFragments() const { return FragmentCollection.GetFragments(); }
-	TArray<const UScriptStruct*>& GetMutableFragmentTypes() { return FragmentTypes; }
-	TConstArrayView<const UScriptStruct*> GetFragmentTypes() const { return FragmentTypes; }
 	FMassUniqueFragmentCollection& GetMutableFragmentCollection() { return FragmentCollection; }
 	TConstArrayView<FInstancedStruct> GetChunkFragments() const { return ChunkFragments; }
 
 	FMassCompositionDescriptor GetCompositionDescriptor() const { return FMassCompositionDescriptor(FragmentCollection.GetFragmentBitSet(), TagBitSet, ChunkFragmentBitSet); }
+	FMassArchetypeFragmentsInitialValues GetArchetypeFragmentsInitialValues() const { return FMassArchetypeFragmentsInitialValues(ChunkFragments); }
 
 	template<typename T>
 	void AddTag()
@@ -178,10 +177,6 @@ private:
 	UPROPERTY()
 	FMassRuntimePipeline DeinitializationPipeline;
 
-	/** An mirror of FMassUniqueFragmentCollection.Fragments' types, used for adding fragment types instead of fragment instances */
-	UPROPERTY()
-	TArray<const UScriptStruct*> FragmentTypes;
-
 	UPROPERTY()
 	TArray<FInstancedStruct> ChunkFragments;
 
