Fix build break, a checkslow condition was not updated.

#rnx

diff --git a/Engine/Plugins/Runtime/MassEntity/Source/MassEntity/Private/MassArchetypeData.cpp b/Engine/Plugins/Runtime/MassEntity/Source/MassEntity/Private/MassArchetypeData.cpp
--- a/Engine/Plugins/Runtime/MassEntity/Source/MassEntity/Private/MassArchetypeData.cpp
+++ b/Engine/Plugins/Runtime/MassEntity/Source/MassEntity/Private/MassArchetypeData.cpp
@@ -75,7 +75,6 @@ void FMassArchetypeData::Initialize(const FMassCompositionDescriptor& InComposit
 			FragmentIndexMap.Add(FragmentType, FragmentIndex);
 			CompositionDescriptor.Fragments.Add(*FragmentType);
 		}
-		checkSlow(Fragments.IsEquivalent(CompositionDescriptor.Fragments));
 	}
 
 	CompositionDescriptor.Tags = InCompositionDescriptor.Tags;
