Renamed FMassComposition to FMassArchetypeComposition

#rnx

diff --git a/Engine/Plugins/Runtime/MassEntity/Source/MassEntity/Private/MassArchetypeData.cpp b/Engine/Plugins/Runtime/MassEntity/Source/MassEntity/Private/MassArchetypeData.cpp
--- a/Engine/Plugins/Runtime/MassEntity/Source/MassEntity/Private/MassArchetypeData.cpp
+++ b/Engine/Plugins/Runtime/MassEntity/Source/MassEntity/Private/MassArchetypeData.cpp
@@ -43,7 +43,7 @@ bool FMassArchetypeData::IsEquivalent(TConstArrayView<const UScriptStruct*> Othe
 	return true;
 }
 
-void FMassArchetypeData::Initialize(const FMassCompositionDescriptor& InCompositionDescriptor, const FMassArchetypeFragmentsInitialValues& InInitialValues)
+void FMassArchetypeData::Initialize(const FMassArchetypeCompositionDescriptor& InCompositionDescriptor, const FMassArchetypeFragmentsInitialValues& InInitialValues)
 {
 	TArray<const UScriptStruct*, TInlineAllocator<16>> SortedFragmentList;
 	InCompositionDescriptor.Fragments.ExportTypes(SortedFragmentList);
diff --git a/Engine/Plugins/Runtime/MassEntity/Source/MassEntity/Private/MassArchetypeData.h b/Engine/Plugins/Runtime/MassEntity/Source/MassEntity/Private/MassArchetypeData.h
--- a/Engine/Plugins/Runtime/MassEntity/Source/MassEntity/Private/MassArchetypeData.h
+++ b/Engine/Plugins/Runtime/MassEntity/Source/MassEntity/Private/MassArchetypeData.h
@@ -136,7 +136,7 @@ struct FMassArchetypeData
 {
 private:
 	// One-stop-shop variable describing the archetype's fragment and tag composition 
-	FMassCompositionDescriptor CompositionDescriptor;
+	FMassArchetypeCompositionDescriptor CompositionDescriptor;
 	FMassArchetypeFragmentsInitialValues InitialValues;
 
 	TArray<FMassArchetypeFragmentConfig, TInlineAllocator<16>> FragmentConfigs;
@@ -162,8 +162,8 @@ public:
 	const FMassFragmentBitSet& GetFragmentBitSet() const { return CompositionDescriptor.Fragments; }
 	const FMassTagBitSet& GetTagBitSet() const { return CompositionDescriptor.Tags; }
 	const FMassChunkFragmentBitSet& GetChunkFragmentBitSet() const { return CompositionDescriptor.ChunkFragments; }
-	const FMassCompositionDescriptor& GetCompositionDescriptor() const { return CompositionDescriptor; }
 
+	const FMassArchetypeCompositionDescriptor& GetCompositionDescriptor() const { return CompositionDescriptor; }
 	const FMassArchetypeFragmentsInitialValues& GetInitialValues() const { return InitialValues; }
 
 	/** Method to iterate on all the fragment types */
@@ -175,12 +175,12 @@ public:
 	{ 
 		return CompositionDescriptor.IsEquivalent(InFragmentBitSet, InTagBitSet, InChunkFragmentsBitSet);
 	}
-	bool IsEquivalent(const FMassCompositionDescriptor& OtherCompositionDescriptor) const
+	bool IsEquivalent(const FMassArchetypeCompositionDescriptor& OtherCompositionDescriptor) const
 	{
 		return CompositionDescriptor.IsEquivalent(OtherCompositionDescriptor);
 	}
 
-	void Initialize(const FMassCompositionDescriptor& InCompositionDescriptor, const FMassArchetypeFragmentsInitialValues& InInitialValues);
+	void Initialize(const FMassArchetypeCompositionDescriptor& InCompositionDescriptor, const FMassArchetypeFragmentsInitialValues& InInitialValues);
 
 	/** 
 	 * A special way of initializing an archetype resulting in a copy of SiblingArchetype's setup with OverrideTags
diff --git a/Engine/Plugins/Runtime/MassEntity/Source/MassEntity/Private/MassEntityQuery.cpp b/Engine/Plugins/Runtime/MassEntity/Source/MassEntity/Private/MassEntityQuery.cpp
--- a/Engine/Plugins/Runtime/MassEntity/Source/MassEntity/Private/MassEntityQuery.cpp
+++ b/Engine/Plugins/Runtime/MassEntity/Source/MassEntity/Private/MassEntityQuery.cpp
@@ -100,7 +100,7 @@ bool FMassEntityQuery::DoesArchetypeMatchRequirements(const FArchetypeHandle& Ar
 	const FMassArchetypeData* Archetype = ArchetypeHandle.DataPtr.Get();
 	CA_ASSUME(Archetype);
 	
-	const FMassCompositionDescriptor& ArchetypeComposition = Archetype->GetCompositionDescriptor();
+	const FMassArchetypeCompositionDescriptor& ArchetypeComposition = Archetype->GetCompositionDescriptor();
 
 	return ArchetypeComposition.Fragments.HasAll(RequiredAllFragments)
 		&& (RequiredAnyFragments.IsEmpty() || ArchetypeComposition.Fragments.HasAny(RequiredAnyFragments))
@@ -301,7 +301,7 @@ FString FMassEntityQuery::DebugGetArchetypeCompatibilityDescription(const FArche
 	const FMassArchetypeData& Archetype = *ArchetypeHandle.DataPtr;
 	FStringOutputDevice OutDescription;
 #if WITH_MASSENTITY_DEBUG	
-	const FMassCompositionDescriptor& ArchetypeComposition = Archetype.GetCompositionDescriptor();
+	const FMassArchetypeCompositionDescriptor& ArchetypeComposition = Archetype.GetCompositionDescriptor();
 
 	if (ArchetypeComposition.Fragments.HasAll(RequiredAllFragments) == false)
 	{
diff --git a/Engine/Plugins/Runtime/MassEntity/Source/MassEntity/Private/MassEntitySubsystem.cpp b/Engine/Plugins/Runtime/MassEntity/Source/MassEntity/Private/MassEntitySubsystem.cpp
--- a/Engine/Plugins/Runtime/MassEntity/Source/MassEntity/Private/MassEntitySubsystem.cpp
+++ b/Engine/Plugins/Runtime/MassEntity/Source/MassEntity/Private/MassEntitySubsystem.cpp
@@ -105,7 +105,7 @@ FArchetypeHandle UMassEntitySubsystem::CreateArchetype(TConstArrayView<const USc
 		}
 	}
 
-	return CreateArchetype(FMassCompositionDescriptor(FMassFragmentBitSet(FragmentList), Tags, ChunkFragments), FMassArchetypeFragmentsInitialValues());
+	return CreateArchetype(FMassArchetypeCompositionDescriptor(FMassFragmentBitSet(FragmentList), Tags, ChunkFragments), FMassArchetypeFragmentsInitialValues());
 }
 
 FArchetypeHandle UMassEntitySubsystem::CreateArchetype(const TSharedPtr<FMassArchetypeData>& SourceArchetype, const FMassFragmentBitSet& NewFragments)
@@ -113,10 +113,10 @@ FArchetypeHandle UMassEntitySubsystem::CreateArchetype(const TSharedPtr<FMassArc
 	check(SourceArchetype.IsValid());
 	checkf(NewFragments.IsEmpty() == false, TEXT("%s Adding an empty fragment list to an archetype is not supported."), ANSI_TO_TCHAR(__FUNCTION__));
 
-	return CreateArchetype(FMassCompositionDescriptor(NewFragments + SourceArchetype->GetFragmentBitSet(), SourceArchetype->GetTagBitSet(), SourceArchetype->GetChunkFragmentBitSet()), SourceArchetype->GetInitialValues());
+	return CreateArchetype(FMassArchetypeCompositionDescriptor(NewFragments + SourceArchetype->GetFragmentBitSet(), SourceArchetype->GetTagBitSet(), SourceArchetype->GetChunkFragmentBitSet()), SourceArchetype->GetInitialValues());
 }
 
-FArchetypeHandle UMassEntitySubsystem::CreateArchetype(const FMassCompositionDescriptor& Composition, const FMassArchetypeFragmentsInitialValues& InitialValues)
+FArchetypeHandle UMassEntitySubsystem::CreateArchetype(const FMassArchetypeCompositionDescriptor& Composition, const FMassArchetypeFragmentsInitialValues& InitialValues)
 {
 	const uint32 TypeHash = Composition.CalculateHash();
 
@@ -157,7 +157,7 @@ FArchetypeHandle UMassEntitySubsystem::InternalCreateSiblingArchetype(const TSha
 {
 	checkSlow(SourceArchetype.IsValid());
 	const FMassArchetypeData& SourceArchetypeRef = *SourceArchetype.Get();
-	const uint32 TypeHash = FMassCompositionDescriptor::CalculateHash(SourceArchetypeRef.GetFragmentBitSet(), OverrideTags, SourceArchetypeRef.GetChunkFragmentBitSet());
+	const uint32 TypeHash = FMassArchetypeCompositionDescriptor::CalculateHash(SourceArchetypeRef.GetFragmentBitSet(), OverrideTags, SourceArchetypeRef.GetChunkFragmentBitSet());
 
 	TArray<TSharedPtr<FMassArchetypeData>>& HashRow = FragmentHashToArchetypeMap.FindOrAdd(TypeHash);
 
@@ -246,7 +246,7 @@ FMassEntityHandle UMassEntitySubsystem::CreateEntity(TConstArrayView<FInstancedS
 {
 	check(FragmentInstanceList.Num() > 0);
 
-	FArchetypeHandle Archetype = CreateArchetype(FMassCompositionDescriptor(FragmentInstanceList, FMassTagBitSet(), FMassChunkFragmentBitSet()), FMassArchetypeFragmentsInitialValues());
+	FArchetypeHandle Archetype = CreateArchetype(FMassArchetypeCompositionDescriptor(FragmentInstanceList, FMassTagBitSet(), FMassChunkFragmentBitSet()), FMassArchetypeFragmentsInitialValues());
 	check(Archetype.DataPtr.IsValid());
 
 	FMassEntityHandle Entity = ReserveEntity();
@@ -289,7 +289,7 @@ void UMassEntitySubsystem::BuildEntity(FMassEntityHandle Entity, TConstArrayView
 	check(FragmentInstanceList.Num() > 0);
 	checkf(!IsEntityBuilt(Entity), TEXT("Expecting an entity that is not already built"));
 
-	FArchetypeHandle Archetype = CreateArchetype(FMassCompositionDescriptor(FragmentInstanceList, FMassTagBitSet(), FMassChunkFragmentBitSet()), FMassArchetypeFragmentsInitialValues());
+	FArchetypeHandle Archetype = CreateArchetype(FMassArchetypeCompositionDescriptor(FragmentInstanceList, FMassTagBitSet(), FMassChunkFragmentBitSet()), FMassArchetypeFragmentsInitialValues());
 	check(Archetype.DataPtr.IsValid());
 
 	InternalBuildEntity(Entity, Archetype);
@@ -416,7 +416,7 @@ void UMassEntitySubsystem::AddFragmentListToEntity(FMassEntityHandle Entity, TCo
 	InternalAddFragmentListToEntityChecked(Entity, FMassFragmentBitSet(FragmentList));
 }
 
-void UMassEntitySubsystem::AddCompositionToEntity_GetDelta(FMassEntityHandle Entity, FMassCompositionDescriptor& InDescriptor)
+void UMassEntitySubsystem::AddCompositionToEntity_GetDelta(FMassEntityHandle Entity, FMassArchetypeCompositionDescriptor& InDescriptor)
 {
 	CheckIfEntityIsActive(Entity);
 
@@ -431,7 +431,7 @@ void UMassEntitySubsystem::AddCompositionToEntity_GetDelta(FMassEntityHandle Ent
 
 	if (InDescriptor.IsEmpty() == false)
 	{
-		FMassCompositionDescriptor NewDescriptor = OldArchetype->GetCompositionDescriptor();
+		FMassArchetypeCompositionDescriptor NewDescriptor = OldArchetype->GetCompositionDescriptor();
 		NewDescriptor.Fragments += InDescriptor.Fragments;
 		NewDescriptor.Tags+= InDescriptor.Tags;
 
@@ -448,37 +448,37 @@ void UMassEntitySubsystem::AddCompositionToEntity_GetDelta(FMassEntityHandle Ent
 	}
 }
 
-void UMassEntitySubsystem::RemoveCompositionFromEntity(FMassEntityHandle Entity, const FMassCompositionDescriptor& InDescriptor)
+void UMassEntitySubsystem::RemoveCompositionFromEntity(FMassEntityHandle Entity, const FMassArchetypeCompositionDescriptor& InDescriptor)
 {
 	CheckIfEntityIsActive(Entity);
 
 	if(InDescriptor.IsEmpty() == false)
 	{
-		FEntityData& EntityData = Entities[Entity.Index];
-		FMassArchetypeData* OldArchetype = EntityData.CurrentArchetype.Get();
-		check(OldArchetype);
+	FEntityData& EntityData = Entities[Entity.Index];
+	FMassArchetypeData* OldArchetype = EntityData.CurrentArchetype.Get();
+	check(OldArchetype);
 
-		FMassCompositionDescriptor NewDescriptor = OldArchetype->GetCompositionDescriptor();
-		NewDescriptor.Fragments -= InDescriptor.Fragments;
-		NewDescriptor.Tags -= InDescriptor.Tags;
+		FMassArchetypeCompositionDescriptor NewDescriptor = OldArchetype->GetCompositionDescriptor();
+	NewDescriptor.Fragments -= InDescriptor.Fragments;
+	NewDescriptor.Tags -= InDescriptor.Tags;
 
 		ensureMsgf(InDescriptor.ChunkFragments.IsEmpty(), TEXT("Removing chunk fragments is not supported"));
 
 		if (NewDescriptor.IsEquivalent(InDescriptor) == false)
-		{
+	{
 			const FArchetypeHandle NewArchetypeHandle = CreateArchetype(NewDescriptor, OldArchetype->GetInitialValues());
 
-			if (ensure(NewArchetypeHandle.DataPtr != EntityData.CurrentArchetype))
-			{
-				// Move the entity over
-				FMassArchetypeData* NewArchetype = NewArchetypeHandle.DataPtr.Get();
-				check(NewArchetype);
-				EntityData.CurrentArchetype->MoveEntityToAnotherArchetype(Entity, *NewArchetype);
-				EntityData.CurrentArchetype = NewArchetypeHandle.DataPtr;
-			}
+		if (ensure(NewArchetypeHandle.DataPtr != EntityData.CurrentArchetype))
+		{
+			// Move the entity over
+			FMassArchetypeData* NewArchetype = NewArchetypeHandle.DataPtr.Get();
+			check(NewArchetype);
+			EntityData.CurrentArchetype->MoveEntityToAnotherArchetype(Entity, *NewArchetype);
+			EntityData.CurrentArchetype = NewArchetypeHandle.DataPtr;
 		}
 	}
 }
+}
 
 void UMassEntitySubsystem::InternalBuildEntity(FMassEntityHandle Entity, const FArchetypeHandle Archetype)
 {
@@ -562,7 +562,7 @@ void UMassEntitySubsystem::RemoveFragmentListFromEntity(FMassEntityHandle Entity
 	if (ensureMsgf(OldArchetype->GetFragmentBitSet().HasAny(FragmentsToRemove), TEXT("Trying to remove a list of fragments from an entity but none of the fragments given was found.")))
 	{
 		// If all the fragments got removed this will result in fetching of the empty archetype
-		const FArchetypeHandle NewArchetypeHandle = CreateArchetype(FMassCompositionDescriptor(OldArchetype->GetFragmentBitSet() - FragmentsToRemove, OldArchetype->GetTagBitSet(), OldArchetype->GetChunkFragmentBitSet()), OldArchetype->GetInitialValues());
+		const FArchetypeHandle NewArchetypeHandle = CreateArchetype(FMassArchetypeCompositionDescriptor(OldArchetype->GetFragmentBitSet() - FragmentsToRemove, OldArchetype->GetTagBitSet(), OldArchetype->GetChunkFragmentBitSet()), OldArchetype->GetInitialValues());
 
 		// Move the entity over
 		FMassArchetypeData* NewArchetype = NewArchetypeHandle.DataPtr.Get();
diff --git a/Engine/Plugins/Runtime/MassEntity/Source/MassEntity/Public/MassCommandBuffer.h b/Engine/Plugins/Runtime/MassEntity/Source/MassEntity/Public/MassCommandBuffer.h
--- a/Engine/Plugins/Runtime/MassEntity/Source/MassEntity/Public/MassCommandBuffer.h
+++ b/Engine/Plugins/Runtime/MassEntity/Source/MassEntity/Public/MassCommandBuffer.h
@@ -270,7 +270,7 @@ struct MASSENTITY_API FCommandRemoveComposition : public FCommandBufferEntryBase
 	GENERATED_BODY()
 
 	FCommandRemoveComposition() = default;
-	FCommandRemoveComposition(FMassEntityHandle InEntity, const FMassCompositionDescriptor& InDescriptor)
+	FCommandRemoveComposition(FMassEntityHandle InEntity, const FMassArchetypeCompositionDescriptor& InDescriptor)
 		: FCommandBufferEntryBase(InEntity)
 		, Descriptor(InDescriptor)
 	{}
@@ -278,7 +278,7 @@ struct MASSENTITY_API FCommandRemoveComposition : public FCommandBufferEntryBase
 protected:
 	virtual void Execute(UMassEntitySubsystem& System) const override;
 
-	FMassCompositionDescriptor Descriptor;
+	FMassArchetypeCompositionDescriptor Descriptor;
 };
 
 struct MASSENTITY_API FMassCommandBuffer
diff --git a/Engine/Plugins/Runtime/MassEntity/Source/MassEntity/Public/MassEntitySubsystem.h b/Engine/Plugins/Runtime/MassEntity/Source/MassEntity/Public/MassEntitySubsystem.h
--- a/Engine/Plugins/Runtime/MassEntity/Source/MassEntity/Public/MassEntitySubsystem.h
+++ b/Engine/Plugins/Runtime/MassEntity/Source/MassEntity/Public/MassEntitySubsystem.h
@@ -89,7 +89,7 @@ public:
 	 * @param InitialValues like default chunk fragment values
 	 * @return a handle of a new archetype 
 	 */
-	FArchetypeHandle CreateArchetype(const FMassCompositionDescriptor& Composition, const FMassArchetypeFragmentsInitialValues& InitialValues);
+	FArchetypeHandle CreateArchetype(const FMassArchetypeCompositionDescriptor& Composition, const FMassArchetypeFragmentsInitialValues& InitialValues);
 
 	/** 
 	 *  Creates an archetype like SourceArchetype + NewFragmentList. 
@@ -189,8 +189,8 @@ public:
 	 * in InOutDescriptor are missing from the current composition of the given entity and then returns the resulting 
 	 * delta via InOutDescriptor.
 	 */
-	void AddCompositionToEntity_GetDelta(FMassEntityHandle Entity, FMassCompositionDescriptor& InOutDescriptor);
-	void RemoveCompositionFromEntity(FMassEntityHandle Entity, const FMassCompositionDescriptor& InDescriptor);
+	void AddCompositionToEntity_GetDelta(FMassEntityHandle Entity, FMassArchetypeCompositionDescriptor& InOutDescriptor);
+	void RemoveCompositionFromEntity(FMassEntityHandle Entity, const FMassArchetypeCompositionDescriptor& InDescriptor);
 
 	/** 
 	 * Moves an entity over to a new archetype by copying over fragments common to both archetypes
diff --git a/Engine/Plugins/Runtime/MassEntity/Source/MassEntity/Public/MassEntityTypes.h b/Engine/Plugins/Runtime/MassEntity/Source/MassEntity/Public/MassEntityTypes.h
--- a/Engine/Plugins/Runtime/MassEntity/Source/MassEntity/Public/MassEntityTypes.h
+++ b/Engine/Plugins/Runtime/MassEntity/Source/MassEntity/Public/MassEntityTypes.h
@@ -103,22 +103,22 @@ using FMassChunkFragmentBitSet = TScriptStructTypeBitSet<FMassChunkFragment>;
 
 /** The type summarily describing a composition of an entity or an archetype. It contains information on both the
  *  fragments as well as tags */
-struct FMassCompositionDescriptor
+struct FMassArchetypeCompositionDescriptor
 {
-	FMassCompositionDescriptor() = default;
-	FMassCompositionDescriptor(const FMassFragmentBitSet& InFragments, const FMassTagBitSet& InTags, const FMassChunkFragmentBitSet& InChunkFragments)
+	FMassArchetypeCompositionDescriptor() = default;
+	FMassArchetypeCompositionDescriptor(const FMassFragmentBitSet& InFragments, const FMassTagBitSet& InTags, const FMassChunkFragmentBitSet& InChunkFragments)
 		: Fragments(InFragments), Tags(InTags), ChunkFragments(InChunkFragments)
 	{}
 
-	FMassCompositionDescriptor(TConstArrayView<const UScriptStruct*> InFragments, const FMassTagBitSet& InTags, const FMassChunkFragmentBitSet& InChunkFragments)
-		: FMassCompositionDescriptor(FMassFragmentBitSet(InFragments), InTags, InChunkFragments)
+	FMassArchetypeCompositionDescriptor(TConstArrayView<const UScriptStruct*> InFragments, const FMassTagBitSet& InTags, const FMassChunkFragmentBitSet& InChunkFragments)
+		: FMassArchetypeCompositionDescriptor(FMassFragmentBitSet(InFragments), InTags, InChunkFragments)
 	{}
 
-	FMassCompositionDescriptor(TConstArrayView<FInstancedStruct> InFragmentInstances, const FMassTagBitSet& InTags, const FMassChunkFragmentBitSet& InChunkFragments)
-		: FMassCompositionDescriptor(FMassFragmentBitSet(InFragmentInstances), InTags, InChunkFragments)
+	FMassArchetypeCompositionDescriptor(TConstArrayView<FInstancedStruct> InFragmentInstances, const FMassTagBitSet& InTags, const FMassChunkFragmentBitSet& InChunkFragments)
+		: FMassArchetypeCompositionDescriptor(FMassFragmentBitSet(InFragmentInstances), InTags, InChunkFragments)
 	{}
 
-	FMassCompositionDescriptor(FMassFragmentBitSet&& InFragments, FMassTagBitSet&& InTags, FMassChunkFragmentBitSet&& InChunkFragments)
+	FMassArchetypeCompositionDescriptor(FMassFragmentBitSet&& InFragments, FMassTagBitSet&& InTags, FMassChunkFragmentBitSet&& InChunkFragments)
 		: Fragments(MoveTemp(InFragments)), Tags(MoveTemp(InTags)), ChunkFragments(MoveTemp(InChunkFragments))
 	{}
 
@@ -133,7 +133,7 @@ struct FMassCompositionDescriptor
 		return Fragments.IsEquivalent(InFragmentBitSet) && Tags.IsEquivalent(InTagBitSet) && ChunkFragments.IsEquivalent(InChunkFragmentsBitSet);
 	}
 
-	bool IsEquivalent(const FMassCompositionDescriptor& OtherDescriptor) const
+	bool IsEquivalent(const FMassArchetypeCompositionDescriptor& OtherDescriptor) const
 	{
 		return IsEquivalent(OtherDescriptor.Fragments, OtherDescriptor.Tags, OtherDescriptor.ChunkFragments);
 	}
diff --git a/Engine/Plugins/Runtime/MassGameplay/Source/MassActors/Private/MassAgentSubsystem.cpp b/Engine/Plugins/Runtime/MassGameplay/Source/MassActors/Private/MassAgentSubsystem.cpp
--- a/Engine/Plugins/Runtime/MassGameplay/Source/MassActors/Private/MassAgentSubsystem.cpp
+++ b/Engine/Plugins/Runtime/MassGameplay/Source/MassActors/Private/MassAgentSubsystem.cpp
@@ -336,7 +336,7 @@ void UMassAgentSubsystem::HandlePendingInitialization()
 			continue;
 		}
 
-		const FMassCompositionDescriptor TemplateDescriptor = EntityTemplate->GetCompositionDescriptor();
+		const FMassArchetypeCompositionDescriptor TemplateDescriptor = EntityTemplate->GetCompositionDescriptor();
 
 		TArray<UMassAgentComponent*>& AgentComponents = Data.Get<1>().AgentComponents;
 
@@ -348,7 +348,7 @@ void UMassAgentSubsystem::HandlePendingInitialization()
 				continue;
 			}
 
-			FMassCompositionDescriptor& PuppetDescriptor = AgentComp->GetMutablePuppetSpecificAddition();
+			FMassArchetypeCompositionDescriptor& PuppetDescriptor = AgentComp->GetMutablePuppetSpecificAddition();
 			PuppetDescriptor = TemplateDescriptor;
 			EntitySystem->AddCompositionToEntity_GetDelta(PuppetEntity, PuppetDescriptor);
 			
diff --git a/Engine/Plugins/Runtime/MassGameplay/Source/MassActors/Public/MassAgentComponent.h b/Engine/Plugins/Runtime/MassGameplay/Source/MassActors/Public/MassAgentComponent.h
--- a/Engine/Plugins/Runtime/MassGameplay/Source/MassActors/Public/MassAgentComponent.h
+++ b/Engine/Plugins/Runtime/MassGameplay/Source/MassActors/Public/MassAgentComponent.h
@@ -100,8 +100,8 @@ public:
 	const FMassEntityConfig& GetEntityConfig() const { return EntityConfig; }
 	void SetEntityConfig(const FMassEntityConfig& InEntityConfig);
 
-	const FMassCompositionDescriptor& GetPuppetSpecificAddition() const { return PuppetSpecificAddition; }
-	FMassCompositionDescriptor& GetMutablePuppetSpecificAddition() { return PuppetSpecificAddition; }
+	const FMassArchetypeCompositionDescriptor& GetPuppetSpecificAddition() const { return PuppetSpecificAddition; }
+	FMassArchetypeCompositionDescriptor& GetMutablePuppetSpecificAddition() { return PuppetSpecificAddition; }
 
 	/** Registers the component with the owner effectively turning it on. Calling it multiple times won't break anything  */
 	UFUNCTION(BlueprintCallable, Category = Mass)
@@ -137,7 +137,7 @@ protected:
 	 *  Contains all the fragments added to the entity during puppet's initialization. Required for clean up when
 	 *  despawning puppet while the entity remains alive.
 	 */
-	FMassCompositionDescriptor PuppetSpecificAddition;
+	FMassArchetypeCompositionDescriptor PuppetSpecificAddition;
 
 	UPROPERTY(EditAnywhere, Category = "Mass")
 	FMassEntityConfig EntityConfig;
diff --git a/Engine/Plugins/Runtime/MassGameplay/Source/MassSpawner/Public/MassEntityTemplate.h b/Engine/Plugins/Runtime/MassGameplay/Source/MassSpawner/Public/MassEntityTemplate.h
--- a/Engine/Plugins/Runtime/MassGameplay/Source/MassSpawner/Public/MassEntityTemplate.h
+++ b/Engine/Plugins/Runtime/MassGameplay/Source/MassSpawner/Public/MassEntityTemplate.h
@@ -109,7 +109,7 @@ struct MASSSPAWNER_API FMassEntityTemplate
 	FMassUniqueFragmentCollection& GetMutableFragmentCollection() { return FragmentCollection; }
 	TConstArrayView<FInstancedStruct> GetChunkFragments() const { return ChunkFragments; }
 
-	FMassCompositionDescriptor GetCompositionDescriptor() const { return FMassCompositionDescriptor(FragmentCollection.GetFragmentBitSet(), TagBitSet, ChunkFragmentBitSet); }
+	FMassArchetypeCompositionDescriptor GetCompositionDescriptor() const { return FMassArchetypeCompositionDescriptor(FragmentCollection.GetFragmentBitSet(), TagBitSet, ChunkFragmentBitSet); }
 	FMassArchetypeFragmentsInitialValues GetArchetypeFragmentsInitialValues() const { return FMassArchetypeFragmentsInitialValues(ChunkFragments); }
 
 	template<typename T>
