Misc type conversion warning fixes

diff --git a/Engine/Plugins/Runtime/MassEntity/Source/MassEntityTestSuite/Private/LWComponentTestFarmPlot.cpp b/Engine/Plugins/Runtime/MassEntity/Source/MassEntityTestSuite/Private/LWComponentTestFarmPlot.cpp
--- a/Engine/Plugins/Runtime/MassEntity/Source/MassEntityTestSuite/Private/LWComponentTestFarmPlot.cpp
+++ b/Engine/Plugins/Runtime/MassEntity/Source/MassEntityTestSuite/Private/LWComponentTestFarmPlot.cpp
@@ -64,7 +64,7 @@ ALWFragmentTestFarmPlot::ALWFragmentTestFarmPlot()
 	HarvestIconISMC->SetCollisionProfileName(UCollisionProfile::NoCollision_ProfileName);
 }
 
-void ALWFragmentTestFarmPlot::AddItemToGrid(UMassEntitySubsystem* EntitySystem, int32 X, int32 Y, FArchetypeHandle Archetype, int32 VisualIndex)
+void ALWFragmentTestFarmPlot::AddItemToGrid(UMassEntitySubsystem* EntitySystem, uint16 X, uint16 Y, FArchetypeHandle Archetype, uint16 VisualIndex)
 {
 	FMassEntityHandle NewItem = EntitySystem->CreateEntity(Archetype);
 	PlantedSquares[X + Y * GridWidth] = NewItem;
@@ -132,12 +132,12 @@ void ALWFragmentTestFarmPlot::BeginPlay()
 	const int32 NumGridCells = GridWidth * GridHeight;
 	PlantedSquares.AddDefaulted(NumGridCells);
 
-	for (int32 Y = 0; Y < GridHeight; ++Y)
+	for (uint16 Y = 0; Y < GridHeight; ++Y)
 	{
-		for (int32 X = 0; X < GridWidth; ++X)
+		for (uint16 X = 0; X < GridWidth; ++X)
 		{
 			const bool bIsOdd = ((X + Y) & 1) != 0;
-			const int32 VisualIndex = bIsOdd ? TestDataCropIndicies[FMath::RandRange(0, TestDataCropIndicies.Num() - 1)] : TestDataFlowerIndicies[FMath::RandRange(0, TestDataFlowerIndicies.Num() - 1)];
+			const uint16 VisualIndex = bIsOdd ? TestDataCropIndicies[FMath::RandRange(0, TestDataCropIndicies.Num() - 1)] : TestDataFlowerIndicies[FMath::RandRange(0, TestDataFlowerIndicies.Num() - 1)];
 
 			AddItemToGrid(EntitySystem, X, Y, bIsOdd ? CropArchetype : FlowerArchetype, VisualIndex);
 		}
diff --git a/Engine/Plugins/Runtime/MassEntity/Source/MassEntityTestSuite/Private/MassCompositeProcessorTest.cpp b/Engine/Plugins/Runtime/MassEntity/Source/MassEntityTestSuite/Private/MassCompositeProcessorTest.cpp
--- a/Engine/Plugins/Runtime/MassEntity/Source/MassEntityTestSuite/Private/MassCompositeProcessorTest.cpp
+++ b/Engine/Plugins/Runtime/MassEntity/Source/MassEntityTestSuite/Private/MassCompositeProcessorTest.cpp
@@ -50,7 +50,7 @@ struct FCompositeProcessorTest_MultipleSubProcessors : FEntityTestBase
 				UMassTestProcessorBase* Processor = NewObject<UMassTestProcessorBase>(EntitySubsystem);
 				Processor->ExecutionFunction = [Processor, &Result, i](UMassEntitySubsystem& InEntitySubsystem, FMassExecutionContext& Context) {
 						check(Processor);
-						Result += FMath::Pow(10.f, float(i));
+						Result += (int)FMath::Pow(10.f, float(i));
 					};
 				Processors.Add(Processor);
 			}
diff --git a/Engine/Plugins/Runtime/MassEntity/Source/MassEntityTestSuite/Public/LWComponentTestFarmPlot.h b/Engine/Plugins/Runtime/MassEntity/Source/MassEntityTestSuite/Public/LWComponentTestFarmPlot.h
--- a/Engine/Plugins/Runtime/MassEntity/Source/MassEntityTestSuite/Public/LWComponentTestFarmPlot.h
+++ b/Engine/Plugins/Runtime/MassEntity/Source/MassEntityTestSuite/Public/LWComponentTestFarmPlot.h
@@ -281,8 +281,8 @@ class ALWFragmentTestFarmPlot : public AActor
 	GENERATED_BODY()
 
 public:
-	int32 GridWidth = 40*7;
-	int32 GridHeight = 20*7;
+	uint16 GridWidth = 40*7;
+	uint16 GridHeight = 20*7;
 
 	UPROPERTY(EditAnywhere, Category=Farm)
 	float GridCellWidth = 150.0f;
@@ -312,23 +312,23 @@ public:
 
 	// Indicies into VisualDataTable for flowers
 	UPROPERTY(EditAnywhere, Category=Farm)
-	TArray<int32> TestDataFlowerIndicies;
+	TArray<uint16> TestDataFlowerIndicies;
 
 	// Indicies into VisualDataTable for crops
 	UPROPERTY(EditAnywhere, Category=Farm)
-	TArray<int32> TestDataCropIndicies;
+	TArray<uint16> TestDataCropIndicies;
 
 	UPROPERTY(EditAnywhere, Category = Farm)
-	float VisualNearCullDistance = 1000.0f;
+	uint32 VisualNearCullDistance = 1000;
 
 	UPROPERTY(EditAnywhere, Category = Farm)
-	float VisualFarCullDistance = 1200.0f;
+	uint32 VisualFarCullDistance = 1200;
 
 	UPROPERTY(EditAnywhere, Category = Farm)
-	float IconNearCullDistance = 400.0f;
+	uint32 IconNearCullDistance = 400;
 
 	UPROPERTY(EditAnywhere, Category = Farm)
-	float IconFarCullDistance = 800.0f;
+	uint32 IconFarCullDistance = 800;
 
 private:
 	UPROPERTY(VisibleAnywhere, BlueprintReadOnly, meta=(AllowPrivateAccess = "true"), Category=Farm)
@@ -337,7 +337,7 @@ private:
 	TArray<int32> FreeHarvestIconIndicies;
 
 private:
-	void AddItemToGrid(UMassEntitySubsystem* EntitySystem, int32 X, int32 Y, FArchetypeHandle Archetype, int32 VisualIndex);
+	void AddItemToGrid(UMassEntitySubsystem* EntitySystem, uint16 X, uint16 Y, FArchetypeHandle Archetype, uint16 VisualIndex);
 
 public:
 	ALWFragmentTestFarmPlot();
diff --git a/Engine/Plugins/Runtime/MassGameplay/Source/MassGameplayDebug/Private/MassDebugVisualizationComponent.cpp b/Engine/Plugins/Runtime/MassGameplay/Source/MassGameplayDebug/Private/MassDebugVisualizationComponent.cpp
--- a/Engine/Plugins/Runtime/MassGameplay/Source/MassGameplayDebug/Private/MassDebugVisualizationComponent.cpp
+++ b/Engine/Plugins/Runtime/MassGameplay/Source/MassGameplayDebug/Private/MassDebugVisualizationComponent.cpp
@@ -67,9 +67,11 @@ void UMassDebugVisualizationComponent::ConstructVisualComponent()
 	}
 }
 
-int32 UMassDebugVisualizationComponent::AddDebugVisType(const FAgentDebugVisualization& Data)
+uint16 UMassDebugVisualizationComponent::AddDebugVisType(const FAgentDebugVisualization& Data)
 {
-	return VisualDataTable.Add(Data);
+	const int32 Index = VisualDataTable.Add(Data);
+	check(VisualDataTable.Num() <= (int32)MAX_uint16);
+	return (uint16)Index;
 }
 
 void UMassDebugVisualizationComponent::Clear()
diff --git a/Engine/Plugins/Runtime/MassGameplay/Source/MassGameplayDebug/Public/MassDebugVisualizationComponent.h b/Engine/Plugins/Runtime/MassGameplay/Source/MassGameplayDebug/Public/MassDebugVisualizationComponent.h
--- a/Engine/Plugins/Runtime/MassGameplay/Source/MassGameplayDebug/Public/MassDebugVisualizationComponent.h
+++ b/Engine/Plugins/Runtime/MassGameplay/Source/MassGameplayDebug/Public/MassDebugVisualizationComponent.h
@@ -21,7 +21,7 @@ public:
 	void DirtyVisuals();
 	int32 AddDebugVisInstance(const uint16 VisualType);
 	/** returns index to the newly created VisualDataTable entry */
-	int32 AddDebugVisType(const FAgentDebugVisualization& Data);
+	uint16 AddDebugVisType(const FAgentDebugVisualization& Data);
 	TArrayView<UHierarchicalInstancedStaticMeshComponent*> GetVisualDataISMCs() { return MakeArrayView(VisualDataISMCs); }
 
 	void Clear();
diff --git a/Engine/Plugins/Runtime/MassGameplay/Source/MassGameplayDebug/Public/MassGameplayDebugTypes.h b/Engine/Plugins/Runtime/MassGameplay/Source/MassGameplayDebug/Public/MassGameplayDebugTypes.h
--- a/Engine/Plugins/Runtime/MassGameplay/Source/MassGameplayDebug/Public/MassGameplayDebugTypes.h
+++ b/Engine/Plugins/Runtime/MassGameplay/Source/MassGameplayDebug/Public/MassGameplayDebugTypes.h
@@ -82,11 +82,11 @@ struct FAgentDebugVisualization : public FTableRowBase
 
 	/** Near cull distance to override default value for that agent type */
 	UPROPERTY(EditAnywhere, Category = "Mass|Debug")
-	float VisualNearCullDistance = 5000.f;
+	uint32 VisualNearCullDistance = 5000;
 
 	/** Far cull distance to override default value for that agent type */
 	UPROPERTY(EditAnywhere, Category = "Mass|Debug")
-	float VisualFarCullDistance = 7500.f;
+	uint32 VisualFarCullDistance = 7500;
 
 	/** If Mesh is not set this WireShape will be used for debug drawing via GameplayDebugger */
 	UPROPERTY(EditAnywhere, Category = "Mass|Debug")
diff --git a/Engine/Plugins/Runtime/MassGameplay/Source/MassLOD/Public/MassLODTickRateController.h b/Engine/Plugins/Runtime/MassGameplay/Source/MassLOD/Public/MassLODTickRateController.h
--- a/Engine/Plugins/Runtime/MassGameplay/Source/MassLOD/Public/MassLODTickRateController.h
+++ b/Engine/Plugins/Runtime/MassGameplay/Source/MassLOD/Public/MassLODTickRateController.h
@@ -49,7 +49,7 @@ public:
 protected:
 
 	/** Tick rate for each LOD */
-	TStaticArray<int32, EMassLOD::Max> TickRates;
+	TStaticArray<float, EMassLOD::Max> TickRates;
 
 	/* Whether or not to spread the first update over the period specified in tick rate member for its LOD */
 	bool bShouldSpreadFirstUpdate = false;
diff --git a/Engine/Plugins/Runtime/MassGameplay/Source/MassRepresentation/Private/MassRepresentationProcessor.cpp b/Engine/Plugins/Runtime/MassGameplay/Source/MassRepresentation/Private/MassRepresentationProcessor.cpp
--- a/Engine/Plugins/Runtime/MassGameplay/Source/MassRepresentation/Private/MassRepresentationProcessor.cpp
+++ b/Engine/Plugins/Runtime/MassGameplay/Source/MassRepresentation/Private/MassRepresentationProcessor.cpp
@@ -148,7 +148,7 @@ void UMassRepresentationProcessor::UpdateRepresentation(FMassExecutionContext& C
 					AActor* NewActor = nullptr;
 					if (!Actor || ActorInfo.IsOwnedByMass())
 					{
-						const int32 WantedTemplateActorIndex = bHighResActor ? Representation.HighResTemplateActorIndex : Representation.LowResTemplateActorIndex;
+						const int16 WantedTemplateActorIndex = bHighResActor ? Representation.HighResTemplateActorIndex : Representation.LowResTemplateActorIndex;
 
 						// If the low res is different than the high res, cancel any pending spawn request that is the opposite of what is needed.
 						if (Representation.LowResTemplateActorIndex != Representation.HighResTemplateActorIndex)
diff --git a/Engine/Plugins/Runtime/MassGameplay/Source/MassSignals/Private/MassSignalTypes.cpp b/Engine/Plugins/Runtime/MassGameplay/Source/MassSignals/Private/MassSignalTypes.cpp
--- a/Engine/Plugins/Runtime/MassGameplay/Source/MassSignals/Private/MassSignalTypes.cpp
+++ b/Engine/Plugins/Runtime/MassGameplay/Source/MassSignals/Private/MassSignalTypes.cpp
@@ -28,14 +28,14 @@ void FMassSignalNameLookup::GetSignalsForEntity(const FMassEntityHandle Entity,
 	OutSignals.Reset();
 	if (const uint64* Signals = EntitySignals.Find(Entity))
 	{
-		int Start = FMath::CountTrailingZeros64(*Signals);
-		int End = FMath::CountLeadingZeros64(*Signals);
-		for (int64 i = Start; i < 64 - End; i++)
+		const uint64 Start = FMath::CountTrailingZeros64(*Signals);
+		const uint64 End = FMath::CountLeadingZeros64(*Signals);
+		for (uint64 i = Start; i < 64 - End; i++)
 		{
 			const uint64 SignalFlag = 1ULL << i;
 			if ((*Signals) & SignalFlag)
 			{
-				OutSignals.Add(SignalNames[i]);
+				OutSignals.Add(SignalNames[(uint32)i]);
 			}
 		}
 	}
diff --git a/Engine/Plugins/Runtime/ZoneGraph/Source/ZoneGraph/Private/ZoneShapeComponent.cpp b/Engine/Plugins/Runtime/ZoneGraph/Source/ZoneGraph/Private/ZoneShapeComponent.cpp
--- a/Engine/Plugins/Runtime/ZoneGraph/Source/ZoneGraph/Private/ZoneShapeComponent.cpp
+++ b/Engine/Plugins/Runtime/ZoneGraph/Source/ZoneGraph/Private/ZoneShapeComponent.cpp
@@ -846,7 +846,7 @@ FPrimitiveSceneProxy* UZoneShapeComponent::CreateSceneProxy()
 		}
 
 		virtual uint32 GetMemoryFootprint(void) const override { return sizeof * this + GetAllocatedSize(); }
-		uint32 GetAllocatedSize(void) const { return FPrimitiveSceneProxy::GetAllocatedSize(); }
+		uint32 GetAllocatedSize(void) const { return (uint32)FPrimitiveSceneProxy::GetAllocatedSize(); }
 
 	private:
 		TArray<FZoneShapeConnector> Connectors;
diff --git a/Engine/Plugins/Runtime/ZoneGraph/Source/ZoneGraph/Private/ZoneShapeUtilities.cpp b/Engine/Plugins/Runtime/ZoneGraph/Source/ZoneGraph/Private/ZoneShapeUtilities.cpp
--- a/Engine/Plugins/Runtime/ZoneGraph/Source/ZoneGraph/Private/ZoneShapeUtilities.cpp
+++ b/Engine/Plugins/Runtime/ZoneGraph/Source/ZoneGraph/Private/ZoneShapeUtilities.cpp
@@ -1781,7 +1781,8 @@ void TessellatePolygonShape(TConstArrayView<FZoneShapePoint> Points, const EZone
 	Zone.LanesBegin = OutZoneStorage.Lanes.Num();
 
 	TArray<FConnectionEntry> Destinations;
-	for (int32 SourceIdx = 0; SourceIdx < Points.Num(); SourceIdx++)
+	check(Points.Num() <= (int32)MAX_uint16);
+	for (uint16 SourceIdx = 0; SourceIdx < Points.Num(); SourceIdx++)
 	{
 		const FZoneShapePoint& SourcePoint = Points[SourceIdx];
 		const FZoneLaneProfile& SourceLaneProfile = LaneProfiles[SourceIdx];
@@ -1791,7 +1792,7 @@ void TessellatePolygonShape(TConstArrayView<FZoneShapePoint> Points, const EZone
 		}
 		// Collect all potential destinations
 		Destinations.Reset();
-		for (int32 j = 1; j < Points.Num(); j++)
+		for (uint16 j = 1; j < Points.Num(); j++)
 		{
 			const int DestIdx = (SourceIdx + j) % Points.Num();
 			const FZoneShapePoint& DestPoint = Points[DestIdx];
@@ -1801,7 +1802,7 @@ void TessellatePolygonShape(TConstArrayView<FZoneShapePoint> Points, const EZone
 				continue;
 			}
 
-			Destinations.Emplace(DestPoint, DestLaneProfile, DestIdx, OutgoingConnections[DestIdx], IncomingConnections[DestIdx]);
+			Destinations.Emplace(DestPoint, DestLaneProfile, (uint16)DestIdx, OutgoingConnections[DestIdx], IncomingConnections[DestIdx]);
 		}
 		// Connect source to destinations.
 		BuildLanesBetweenPoints(FConnectionEntry(SourcePoint, SourceLaneProfile, SourceIdx, OutgoingConnections[SourceIdx], IncomingConnections[SourceIdx]),
diff --git a/Engine/Plugins/Runtime/ZoneGraph/Source/ZoneGraphDebug/Private/ZoneGraphTestingActor.cpp b/Engine/Plugins/Runtime/ZoneGraph/Source/ZoneGraphDebug/Private/ZoneGraphTestingActor.cpp
--- a/Engine/Plugins/Runtime/ZoneGraph/Source/ZoneGraphDebug/Private/ZoneGraphTestingActor.cpp
+++ b/Engine/Plugins/Runtime/ZoneGraph/Source/ZoneGraphDebug/Private/ZoneGraphTestingActor.cpp
@@ -253,7 +253,7 @@ FPrimitiveSceneProxy* UZoneGraphTestingComponent::CreateSceneProxy()
 		}
 
 		virtual uint32 GetMemoryFootprint(void) const override { return sizeof(*this) + GetAllocatedSize(); }
-		uint32 GetAllocatedSize(void) const { return FPrimitiveSceneProxy::GetAllocatedSize(); }
+		uint32 GetAllocatedSize(void) const { return (uint32)FPrimitiveSceneProxy::GetAllocatedSize(); }
 
 		void DrawView(FPrimitiveDrawInterface* PDI) const
 		{
diff --git a/Engine/Plugins/Runtime/ZoneGraph/Source/ZoneGraphEditor/Private/ZoneShapeComponentDetails.cpp b/Engine/Plugins/Runtime/ZoneGraph/Source/ZoneGraphEditor/Private/ZoneShapeComponentDetails.cpp
--- a/Engine/Plugins/Runtime/ZoneGraph/Source/ZoneGraphEditor/Private/ZoneShapeComponentDetails.cpp
+++ b/Engine/Plugins/Runtime/ZoneGraph/Source/ZoneGraphEditor/Private/ZoneShapeComponentDetails.cpp
@@ -994,7 +994,7 @@ FText FZoneShapePointDetails::GetLaneConnectionRestrictions() const
 			{
 				continue;
 			}
-			const int32 Value = Enum->GetValueByIndex(Index);
+			const int32 Value = (int32)Enum->GetValueByIndex(Index);
 			if ((Flags & Value) != 0)
 			{
 				Names.Add(Enum->GetDisplayNameTextByValue((int64)Value));
diff --git a/Engine/Plugins/Runtime/ZoneGraphAnnotations/Source/ZoneGraphAnnotations/Private/ZoneGraphAnnotationSubsystem.cpp b/Engine/Plugins/Runtime/ZoneGraphAnnotations/Source/ZoneGraphAnnotations/Private/ZoneGraphAnnotationSubsystem.cpp
--- a/Engine/Plugins/Runtime/ZoneGraphAnnotations/Source/ZoneGraphAnnotations/Private/ZoneGraphAnnotationSubsystem.cpp
+++ b/Engine/Plugins/Runtime/ZoneGraphAnnotations/Source/ZoneGraphAnnotations/Private/ZoneGraphAnnotationSubsystem.cpp
@@ -57,11 +57,12 @@ void UZoneGraphAnnotationSubsystem::Deinitialize()
 
 void UZoneGraphAnnotationSubsystem::AddToAnnotationLookup(UZoneGraphAnnotationComponent& Annotation, const FZoneGraphTagMask AnnotationTags)
 {
-	const int32 Start = FMath::CountTrailingZeros(AnnotationTags.GetValue());
-	const int32 End = 32 - FMath::CountLeadingZeros(AnnotationTags.GetValue());
-	for (int32 Index = Start; Index < End; Index++)
+	const uint32 Start = FMath::CountTrailingZeros(AnnotationTags.GetValue());
+	const uint32 End = 32 - FMath::CountLeadingZeros(AnnotationTags.GetValue());
+	for (uint32 Index = Start; Index < End; Index++)
 	{
-		const FZoneGraphTag BitTag(Index);
+		check(Index <= (uint32)MAX_uint8);
+		const FZoneGraphTag BitTag((uint8)Index);
 		if (AnnotationTags.Contains(BitTag))
 		{
 			if (ensureMsgf(TagToAnnotationLookup[Index] == nullptr, TEXT("Annotation at index %d is already set to %s."), Index, *GetNameSafe(TagToAnnotationLookup[Index])) &&
diff --git a/Engine/Source/Developer/GameplayDebugger/Public/GameplayDebuggerCategoryReplicator.h b/Engine/Source/Developer/GameplayDebugger/Public/GameplayDebuggerCategoryReplicator.h
--- a/Engine/Source/Developer/GameplayDebugger/Public/GameplayDebuggerCategoryReplicator.h
+++ b/Engine/Source/Developer/GameplayDebugger/Public/GameplayDebuggerCategoryReplicator.h
@@ -116,7 +116,7 @@ public:
 	FName GetDebugActorName() const { return DebugActor.ActorName; }
 
 	/** get sync counter, increased with every change of DebugActor */
-	int16 GetDebugActorCounter() const { return DebugActor.SyncCounter; }
+	int32 GetDebugActorCounter() const { return DebugActor.SyncCounter; }
 
 	const FGameplayDebuggerVisLogSync& GetVisLogSyncData() const { return VisLogSync; }
 
diff --git a/Engine/Source/Runtime/AIModule/Public/IndexedHandle.h b/Engine/Source/Runtime/AIModule/Public/IndexedHandle.h
--- a/Engine/Source/Runtime/AIModule/Public/IndexedHandle.h
+++ b/Engine/Source/Runtime/AIModule/Public/IndexedHandle.h
@@ -120,8 +120,13 @@ struct AIMODULE_API FCompactIndexedHandleBase
 	FCompactIndexedHandleBase() = default;
 
 	/** @note passing INDEX_NONE as index will make this handle Invalid */
-	FCompactIndexedHandleBase(int16 InIndex, uint16 InSerialNumber) : Index(InIndex), SerialNumber(InSerialNumber)
+	FCompactIndexedHandleBase(const int32 InIndex, const uint32 InSerialNumber)
 	{
+		check(InIndex <= (int32)MAX_int16);
+		Index = (int16)InIndex;
+		
+		check(InSerialNumber <= (uint32)MAX_uint16);
+		SerialNumber = (uint16)InSerialNumber;
 	};
 
 	bool operator==(const FCompactIndexedHandleBase& Other) const
diff --git a/Engine/Source/Runtime/NavigationSystem/Public/NavMesh/NavMeshPath.h b/Engine/Source/Runtime/NavigationSystem/Public/NavMesh/NavMeshPath.h
--- a/Engine/Source/Runtime/NavigationSystem/Public/NavMesh/NavMeshPath.h
+++ b/Engine/Source/Runtime/NavigationSystem/Public/NavMesh/NavMeshPath.h
@@ -22,7 +22,7 @@ struct NAVIGATIONSYSTEM_API FNavMeshNodeFlags
 	uint16 AreaFlags;
 
 	FNavMeshNodeFlags() : PathFlags(0), Area(0), AreaFlags(0) {}
-	FNavMeshNodeFlags(uint32 Flags) : PathFlags(Flags), Area(Flags >> 8), AreaFlags(Flags >> 16) {}
+	FNavMeshNodeFlags(const uint32 Flags) : PathFlags((uint8)Flags), Area((uint8)(Flags >> 8)), AreaFlags((uint16)(Flags >> 16)) {}
 	uint32 Pack() const { return PathFlags | ((uint32)Area << 8) | ((uint32)AreaFlags << 16); }
 	bool IsNavLink() const { return (PathFlags & RECAST_STRAIGHTPATH_OFFMESH_CONNECTION) != 0; }
 
diff --git a/Engine/Source/Runtime/NavigationSystem/Public/NavMesh/RecastNavMesh.h b/Engine/Source/Runtime/NavigationSystem/Public/NavMesh/RecastNavMesh.h
--- a/Engine/Source/Runtime/NavigationSystem/Public/NavMesh/RecastNavMesh.h
+++ b/Engine/Source/Runtime/NavigationSystem/Public/NavMesh/RecastNavMesh.h
@@ -143,7 +143,7 @@ struct FRecastDebugPathfindingNode
 
 namespace ERecastDebugPathfindingFlags
 {
-	enum Type
+	enum Type : uint8
 	{
 		Basic = 0x0,
 		BestNode = 0x1,
@@ -999,7 +999,7 @@ public:
 	FVector GetModifiedQueryExtent(const FVector& QueryExtent) const
 	{
 		// Using HALF_WORLD_MAX instead of BIG_NUMBER, else using the extent for a box will result in NaN.
-		return FVector(QueryExtent.X, QueryExtent.Y, QueryExtent.Z >= HALF_WORLD_MAX ? HALF_WORLD_MAX : (QueryExtent.Z + FMath::Max(0.0f, VerticalDeviationFromGroundCompensation)));
+		return FVector(QueryExtent.X, QueryExtent.Y, QueryExtent.Z >= (float)HALF_WORLD_MAX ? (float)HALF_WORLD_MAX : (QueryExtent.Z + FMath::Max(0.0f, VerticalDeviationFromGroundCompensation)));
 	}
 
 	//----------------------------------------------------------------------//
diff --git a/Engine/Source/Runtime/NavigationSystem/Public/NavigationOctree.h b/Engine/Source/Runtime/NavigationSystem/Public/NavigationOctree.h
--- a/Engine/Source/Runtime/NavigationSystem/Public/NavigationOctree.h
+++ b/Engine/Source/Runtime/NavigationSystem/Public/NavigationOctree.h
@@ -68,7 +68,7 @@ public:
 
 	FORCEINLINE int32 GetAllocatedSize() const
 	{
-		return Data->GetAllocatedSize();
+		return (int32)Data->GetAllocatedSize();
 	}
 
 	FORCEINLINE void Shrink()
