New shared fragment type for Mass Entity Subsystem

#rnx

diff --git a/Engine/Plugins/Runtime/MassEntity/Source/MassEntity/Private/MassArchetypeData.cpp b/Engine/Plugins/Runtime/MassEntity/Source/MassEntity/Private/MassArchetypeData.cpp
--- a/Engine/Plugins/Runtime/MassEntity/Source/MassEntity/Private/MassArchetypeData.cpp
+++ b/Engine/Plugins/Runtime/MassEntity/Source/MassEntity/Private/MassArchetypeData.cpp
@@ -25,24 +25,6 @@ bool FMassArchetypeData::HasFragmentType(const UScriptStruct* FragmentType) cons
 	return (FragmentType && CompositionDescriptor.Fragments.Contains(*FragmentType));
 }
 
-bool FMassArchetypeData::IsEquivalent(TConstArrayView<const UScriptStruct*> OtherList) const
-{
-	if (OtherList.Num() != FragmentConfigs.Num())
-	{
-		return false;
-	}
-
-	for (int32 Index = 0; Index < FragmentConfigs.Num(); ++Index)
-	{
-		if (OtherList[Index] != FragmentConfigs[Index].FragmentType)
-		{
-			return false;
-		}
-	}
-
-	return true;
-}
-
 void FMassArchetypeData::Initialize(const FMassArchetypeCompositionDescriptor& InCompositionDescriptor, const FMassArchetypeFragmentsInitialValues& InInitialValues)
 {
 	TArray<const UScriptStruct*, TInlineAllocator<16>> SortedFragmentList;
@@ -76,26 +58,35 @@ void FMassArchetypeData::Initialize(const FMassArchetypeCompositionDescriptor& I
 			CompositionDescriptor.Fragments.Add(*FragmentType);
 		}
 	}
+	// @todo, remove this from the Initial values and rename this structure differently
+	InitialValues.SetFragments(InInitialValues.GetFragments());
 
+	// Tags
 	CompositionDescriptor.Tags = InCompositionDescriptor.Tags;
+
+	// Chunk fragments
 	CompositionDescriptor.ChunkFragments = InCompositionDescriptor.ChunkFragments;
-	// create the chunk fragment instances
 	TArray<const UScriptStruct*, TInlineAllocator<16>> ChunkFragmentList;
 	CompositionDescriptor.ChunkFragments.ExportTypes(ChunkFragmentList);
 	ChunkFragmentList.Sort(FMassSorterOperator<UScriptStruct>());
 	for (const UScriptStruct* ChunkFragmentType : ChunkFragmentList)
 	{
 		check(ChunkFragmentType);
-		if (const FInstancedStruct* Struct = InInitialValues.ChunkFragments.FindByPredicate(FSameTypeScriptStructPredicate(ChunkFragmentType)))
+		if (const FInstancedStruct* Struct = InInitialValues.GetChunkFragments().FindByPredicate(FSameTypeScriptStructPredicate(ChunkFragmentType)))
 		{
-			InitialValues.ChunkFragments.Add(*Struct);
+			InitialValues.AddChunkFragment(FConstStructView(*Struct));
 		}
 		else
 		{
-			InitialValues.ChunkFragments.Add(FInstancedStruct(ChunkFragmentType));
+			InitialValues.AddChunkFragment(FConstStructView(ChunkFragmentType));
 		}
 	}
 
+	// Share fragments
+	CompositionDescriptor.SharedFragments = InCompositionDescriptor.SharedFragments;
+	InitialValues.SetConstSharedFragments(InInitialValues.GetConstSharedFragments());
+	InitialValues.SetSharedFragments(InInitialValues.GetSharedFragments());
+
 	TotalBytesPerEntity = FragmentSizeTallyBytes;
 	int32 ChunkAvailableSize = GetChunkAllocSize() - AlignmentPadding;
 	check(TotalBytesPerEntity <= ChunkAvailableSize);
@@ -123,7 +114,12 @@ void FMassArchetypeData::InitializeWithSibling(const FMassArchetypeData& Sibling
 	CompositionDescriptor.Fragments = SiblingArchetype.CompositionDescriptor.Fragments;
 	CompositionDescriptor.Tags = OverrideTags;
 	CompositionDescriptor.ChunkFragments = SiblingArchetype.CompositionDescriptor.ChunkFragments;
-	InitialValues.ChunkFragments = SiblingArchetype.InitialValues.ChunkFragments;
+	CompositionDescriptor.SharedFragments = SiblingArchetype.CompositionDescriptor.SharedFragments;
+	// @todo maybe we should implement a copy operator in the FMassArchetypeInitialValues struct
+	InitialValues.SetFragments(SiblingArchetype.InitialValues.GetFragments());
+	InitialValues.SetChunkFragments(SiblingArchetype.InitialValues.GetChunkFragments());
+	InitialValues.SetConstSharedFragments(SiblingArchetype.InitialValues.GetConstSharedFragments());
+	InitialValues.SetSharedFragments(SiblingArchetype.InitialValues.GetSharedFragments());
 
 	TotalBytesPerEntity = SiblingArchetype.TotalBytesPerEntity;
 	NumEntitiesPerChunk = SiblingArchetype.NumEntitiesPerChunk;
@@ -179,12 +175,12 @@ int32 FMassArchetypeData::AddEntityInternal(FMassEntityHandle Entity, const bool
 		if (EmptyChunkIndex != INDEX_NONE)
 		{
 			DestinationChunk = &Chunks[EmptyChunkIndex];
-			DestinationChunk->Recycle(InitialValues.ChunkFragments);
+			DestinationChunk->Recycle(InitialValues.GetChunkFragments());
 			AbsoluteIndex = EmptyAbsoluteIndex;
 		}
 		else
 		{
-			DestinationChunk = &Chunks.Emplace_GetRef(GetChunkAllocSize(), InitialValues.ChunkFragments);
+			DestinationChunk = &Chunks.Emplace_GetRef(GetChunkAllocSize(), InitialValues.GetChunkFragments());
 		}
 
 		check(DestinationChunk);
@@ -493,6 +489,10 @@ void FMassArchetypeData::ExecuteFunction(FMassExecutionContext& RunContext, cons
 
 	// mz@todo to be removed
 	RunContext.SetCurrentArchetypeData(*this);
+
+	BindConstSharedFragmentRequirements(RunContext, RequirementMapping.ChunkFragments);
+	BindSharedFragmentRequirements(RunContext, RequirementMapping.ChunkFragments);
+
 	for (FMassArchetypeChunkIterator ChunkIterator(ChunkCollection); ChunkIterator; ++ChunkIterator)
 	{
 		FMassArchetypeChunk& Chunk = Chunks[ChunkIterator->ChunkIndex];
@@ -511,21 +511,28 @@ void FMassArchetypeData::ExecuteFunction(FMassExecutionContext& RunContext, cons
 	}
 }
 
-void FMassArchetypeData::ExecuteFunction(FMassExecutionContext& RunContext, const FMassExecuteFunction& Function, const FMassQueryRequirementIndicesMapping& RequirementMapping, const FMassChunkConditionFunction& ChunkCondition)
+void FMassArchetypeData::ExecuteFunction(FMassExecutionContext& RunContext, const FMassExecuteFunction& Function, const FMassQueryRequirementIndicesMapping& RequirementMapping, const FMassArchetypeConditionFunction& ArchetypeCondition, const FMassChunkConditionFunction& ChunkCondition)
 {
 	// mz@todo to be removed
 	RunContext.SetCurrentArchetypeData(*this);
-	for (FMassArchetypeChunk& Chunk : Chunks)
+
+	BindConstSharedFragmentRequirements(RunContext, RequirementMapping.ConstSharedFragments);
+	BindSharedFragmentRequirements(RunContext, RequirementMapping.SharedFragments);
+
+	if(!ArchetypeCondition || ArchetypeCondition(RunContext))
 	{
-		if (Chunk.GetNumInstances())
+		for (FMassArchetypeChunk& Chunk : Chunks)
 		{
-			RunContext.SetCurrentChunkSerialModificationNumber(Chunk.GetSerialModificationNumber());
-			BindChunkFragmentRequirements(RunContext, RequirementMapping.ChunkFragments, Chunk);
-
-			if (!ChunkCondition || ChunkCondition(RunContext))
+			if (Chunk.GetNumInstances())
 			{
-				BindEntityRequirements(RunContext, RequirementMapping.EntityFragments, Chunk, 0, Chunk.GetNumInstances());
-				Function(RunContext);
+				RunContext.SetCurrentChunkSerialModificationNumber(Chunk.GetSerialModificationNumber());
+				BindChunkFragmentRequirements(RunContext, RequirementMapping.ChunkFragments, Chunk);
+
+				if (!ChunkCondition || ChunkCondition(RunContext))
+				{
+					BindEntityRequirements(RunContext, RequirementMapping.EntityFragments, Chunk, 0, Chunk.GetNumInstances());
+					Function(RunContext);
+				}
 			}
 		}
 	}
@@ -538,6 +545,9 @@ void FMassArchetypeData::ExecutionFunctionForChunk(FMassExecutionContext RunCont
 
 	if (ChunkLength)
 	{
+		BindConstSharedFragmentRequirements(RunContext, RequirementMapping.ChunkFragments);
+		BindSharedFragmentRequirements(RunContext, RequirementMapping.ChunkFragments);
+
 		RunContext.SetCurrentArchetypeData(*this);
 		RunContext.SetCurrentChunkSerialModificationNumber(Chunk.GetSerialModificationNumber());
 		BindChunkFragmentRequirements(RunContext, RequirementMapping.ChunkFragments, Chunk);
@@ -663,9 +673,10 @@ void FMassArchetypeData::GetRequirementsChunkFragmentMapping(TConstArrayView<FMa
 		if (Requirement.RequiresBinding())
 		{
 			int32 FragmentIndex = INDEX_NONE;
-			for (int32 i = LastFoundFragmentIndex + 1; i < InitialValues.ChunkFragments.Num(); ++i)
+			// mz@todo Add comment here as this code seems to be assuming a certain order for chunk fragments, please explain
+			for (int32 i = LastFoundFragmentIndex + 1; i < InitialValues.GetChunkFragments().Num(); ++i)
 			{
-				if (InitialValues.ChunkFragments[i].GetScriptStruct()->IsChildOf(Requirement.StructType))
+				if (InitialValues.GetChunkFragments()[i].GetScriptStruct()->IsChildOf(Requirement.StructType))
 				{
 					FragmentIndex = i;
 					break;
@@ -679,6 +690,34 @@ void FMassArchetypeData::GetRequirementsChunkFragmentMapping(TConstArrayView<FMa
 	}
 }
 
+void FMassArchetypeData::GetRequirementsConstSharedFragmentMapping(TConstArrayView<FMassFragmentRequirement> Requirements, FMassFragmentIndicesMapping& OutFragmentIndices)
+{
+	OutFragmentIndices.Reset(Requirements.Num());
+	for (const FMassFragmentRequirement& Requirement : Requirements)
+	{
+		if (Requirement.RequiresBinding())
+		{
+			const int32 FragmentIndex = InitialValues.GetConstSharedFragments().IndexOfByPredicate(FSameTypeScriptStructPredicate(Requirement.StructType));
+			check(FragmentIndex != INDEX_NONE || Requirement.IsOptional());
+			OutFragmentIndices.Add(FragmentIndex);
+		}
+	}
+}
+
+void FMassArchetypeData::GetRequirementsSharedFragmentMapping(TConstArrayView<FMassFragmentRequirement> Requirements, FMassFragmentIndicesMapping& OutFragmentIndices)
+{
+	OutFragmentIndices.Reset(Requirements.Num());
+	for (const FMassFragmentRequirement& Requirement : Requirements)
+	{
+		if (Requirement.RequiresBinding())
+		{
+			const int32 FragmentIndex = InitialValues.GetSharedFragments().IndexOfByPredicate(FSameTypeScriptStructPredicate(Requirement.StructType));
+			check(FragmentIndex != INDEX_NONE || Requirement.IsOptional());
+			OutFragmentIndices.Add(FragmentIndex);
+		}
+	}
+}
+
 void FMassArchetypeData::BindEntityRequirements(FMassExecutionContext& RunContext, const FMassFragmentIndicesMapping& EntityFragmentsMapping, FMassArchetypeChunk& Chunk, const int32 SubchunkStart, const int32 SubchunkLength)
 {
 	// auto-correcting number of entities to process in case SubchunkStart +  SubchunkLength > Chunk.GetNumInstances()
@@ -735,11 +774,11 @@ void FMassArchetypeData::BindChunkFragmentRequirements(FMassExecutionContext& Ru
 
 		for (int i = 0; i < ChunkFragmentsMapping.Num(); ++i)
 		{
-			FMassExecutionContext::FChunkFragmentView& ChunkRequirement = RunContext.ChunkFragments[i];
+			FMassExecutionContext::FChunkFragmentView& ChunkRequirement = RunContext.ChunkFragmentViews[i];
 			const int32 ChunkFragmentIndex = ChunkFragmentsMapping[i];
 
 			check(ChunkFragmentIndex != INDEX_NONE || ChunkRequirement.Requirement.IsOptional());
-			ChunkRequirement.ChunkFragmentView = ChunkFragmentIndex != INDEX_NONE ? Chunk.GetMutableChunkFragmentViewChecked(ChunkFragmentIndex) : FStructView();
+			ChunkRequirement.FragmentView = ChunkFragmentIndex != INDEX_NONE ? Chunk.GetMutableChunkFragmentViewChecked(ChunkFragmentIndex) : FStructView();
 		}
 	}
 	else
@@ -748,7 +787,59 @@ void FMassArchetypeData::BindChunkFragmentRequirements(FMassExecutionContext& Ru
 		{
 			FInstancedStruct* ChunkFragmentInstance = Chunk.FindMutableChunkFragment(ChunkRequirement.Requirement.StructType);
 			check(ChunkFragmentInstance != nullptr || ChunkRequirement.Requirement.IsOptional());
-			ChunkRequirement.ChunkFragmentView = ChunkFragmentInstance ? FStructView(*ChunkFragmentInstance) : FStructView(); 
+			ChunkRequirement.FragmentView = ChunkFragmentInstance ? FStructView(*ChunkFragmentInstance) : FStructView();
+		}
+	}
+}
+
+void FMassArchetypeData::BindConstSharedFragmentRequirements(FMassExecutionContext& RunContext, const FMassFragmentIndicesMapping& FragmentsMapping)
+{
+	if (FragmentsMapping.Num() > 0)
+	{
+		check(RunContext.GetMutableConstSharedRequirements().Num() == FragmentsMapping.Num());
+
+		for (int i = 0; i < FragmentsMapping.Num(); ++i)
+		{
+			FMassExecutionContext::FConstSharedFragmentView& Requirement = RunContext.ConstSharedFragmentViews[i];
+			const int32 FragmentIndex = FragmentsMapping[i];
+
+			check(FragmentIndex != INDEX_NONE || Requirement.Requirement.IsOptional());
+			Requirement.FragmentView = FragmentIndex != INDEX_NONE ? InitialValues.GetConstSharedFragments()[FragmentIndex] : FConstSharedStruct();
+		}
+	}
+	else
+	{
+		for (FMassExecutionContext::FConstSharedFragmentView& Requirement : RunContext.GetMutableConstSharedRequirements())
+		{
+			const FConstSharedStruct* SharedFragment = InitialValues.GetConstSharedFragments().FindByPredicate( FSameTypeScriptStructPredicate(Requirement.Requirement.StructType) );
+			check(SharedFragment != nullptr || Requirement.Requirement.IsOptional());
+			Requirement.FragmentView = SharedFragment ? *SharedFragment : FConstSharedStruct();
+		}
+	}
+}
+
+void FMassArchetypeData::BindSharedFragmentRequirements(FMassExecutionContext& RunContext, const FMassFragmentIndicesMapping& FragmentsMapping)
+{
+	if (FragmentsMapping.Num() > 0)
+	{
+		check(RunContext.GetMutableSharedRequirements().Num() == FragmentsMapping.Num());
+
+		for (int i = 0; i < FragmentsMapping.Num(); ++i)
+		{
+			FMassExecutionContext::FSharedFragmentView& Requirement = RunContext.SharedFragmentViews[i];
+			const int32 FragmentIndex = FragmentsMapping[i];
+
+			check(FragmentIndex != INDEX_NONE || Requirement.Requirement.IsOptional());
+			Requirement.FragmentView = FragmentIndex != INDEX_NONE ? InitialValues.GetSharedFragments()[FragmentIndex] : FSharedStruct();
+		}
+	}
+	else
+	{
+		for (FMassExecutionContext::FSharedFragmentView& Requirement : RunContext.GetMutableSharedRequirements())
+		{
+			const FSharedStruct* SharedFragment = InitialValues.GetSharedFragments().FindByPredicate(FSameTypeScriptStructPredicate(Requirement.Requirement.StructType));
+			check(SharedFragment != nullptr || Requirement.Requirement.IsOptional());
+			Requirement.FragmentView = SharedFragment ? *SharedFragment : FSharedStruct();
 		}
 	}
 }
@@ -765,7 +856,7 @@ SIZE_T FMassArchetypeData::GetAllocatedSize() const
 	}
 
 	return sizeof(FMassArchetypeData) +
-		InitialValues.ChunkFragments.GetAllocatedSize() +
+		InitialValues.GetAllocatedSize() +
 		FragmentConfigs.GetAllocatedSize() +
 		Chunks.GetAllocatedSize() +
 		(NumAllocatedChunkBuffers * GetChunkAllocSize()) +
diff --git a/Engine/Plugins/Runtime/MassEntity/Source/MassEntity/Private/MassArchetypeData.h b/Engine/Plugins/Runtime/MassEntity/Source/MassEntity/Private/MassArchetypeData.h
--- a/Engine/Plugins/Runtime/MassEntity/Source/MassEntity/Private/MassArchetypeData.h
+++ b/Engine/Plugins/Runtime/MassEntity/Source/MassEntity/Private/MassArchetypeData.h
@@ -162,6 +162,7 @@ public:
 	const FMassFragmentBitSet& GetFragmentBitSet() const { return CompositionDescriptor.Fragments; }
 	const FMassTagBitSet& GetTagBitSet() const { return CompositionDescriptor.Tags; }
 	const FMassChunkFragmentBitSet& GetChunkFragmentBitSet() const { return CompositionDescriptor.ChunkFragments; }
+	const FMassSharedFragmentBitSet& GetSharedFragmentBitSet() const { return CompositionDescriptor.SharedFragments; }
 
 	const FMassArchetypeCompositionDescriptor& GetCompositionDescriptor() const { return CompositionDescriptor; }
 	const FMassArchetypeFragmentsInitialValues& GetInitialValues() const { return InitialValues; }
@@ -170,14 +171,10 @@ public:
 	void ForEachFragmentType(TFunction< void(const UScriptStruct* /*FragmentType*/)> Function) const;
 	bool HasFragmentType(const UScriptStruct* FragmentType) const;
 	bool HasTagType(const UScriptStruct* FragmentType) const { check(FragmentType); return CompositionDescriptor.Tags.Contains(*FragmentType); }
-	bool IsEquivalent(TConstArrayView<const UScriptStruct*> OtherList) const;
-	bool IsEquivalent(const FMassFragmentBitSet& InFragmentBitSet, const FMassTagBitSet& InTagBitSet, const FMassChunkFragmentBitSet& InChunkFragmentsBitSet) const
-	{ 
-		return CompositionDescriptor.IsEquivalent(InFragmentBitSet, InTagBitSet, InChunkFragmentsBitSet);
-	}
-	bool IsEquivalent(const FMassArchetypeCompositionDescriptor& OtherCompositionDescriptor) const
+
+	bool IsEquivalent(const FMassArchetypeCompositionDescriptor& OtherCompositionDescriptor, const FMassArchetypeFragmentsInitialValues& OtherInitialValues) const
 	{
-		return CompositionDescriptor.IsEquivalent(OtherCompositionDescriptor);
+		return CompositionDescriptor.IsEquivalent(OtherCompositionDescriptor) && InitialValues.IsEquivalent(OtherInitialValues);
 	}
 
 	void Initialize(const FMassArchetypeCompositionDescriptor& InCompositionDescriptor, const FMassArchetypeFragmentsInitialValues& InInitialValues);
@@ -206,7 +203,7 @@ public:
 	int32 GetChunkCount() const { return Chunks.Num(); }
 
 	void ExecuteFunction(FMassExecutionContext& RunContext, const FMassExecuteFunction& Function, const FMassQueryRequirementIndicesMapping& RequirementMapping, const FArchetypeChunkCollection& ChunkCollection);
-	void ExecuteFunction(FMassExecutionContext& RunContext, const FMassExecuteFunction& Function, const FMassQueryRequirementIndicesMapping& RequirementMapping, const FMassChunkConditionFunction& ChunkCondition = FMassChunkConditionFunction());
+	void ExecuteFunction(FMassExecutionContext& RunContext, const FMassExecuteFunction& Function, const FMassQueryRequirementIndicesMapping& RequirementMapping, const FMassArchetypeConditionFunction& ArchetypeCondition, const FMassChunkConditionFunction& ChunkCondition);
 
 	void ExecutionFunctionForChunk(FMassExecutionContext RunContext, const FMassExecuteFunction& Function, const FMassQueryRequirementIndicesMapping& RequirementMapping, const FArchetypeChunkCollection::FChunkInfo& ChunkInfo, const FMassChunkConditionFunction& ChunkCondition = FMassChunkConditionFunction());
 
@@ -239,6 +236,12 @@ public:
 	/** Returns conversion from given ChunkRequirements to archetype's chunk fragment indices */
 	void GetRequirementsChunkFragmentMapping(TConstArrayView<FMassFragmentRequirement> ChunkRequirements, FMassFragmentIndicesMapping& OutFragmentIndices);
 
+	/** Returns conversion from given const shared requirements to archetype's const shared fragment indices */
+	void GetRequirementsConstSharedFragmentMapping(TConstArrayView<FMassFragmentRequirement> Requirements, FMassFragmentIndicesMapping& OutFragmentIndices);
+
+	/** Returns conversion from given shared requirements to archetype's shared fragment indices */
+	void GetRequirementsSharedFragmentMapping(TConstArrayView<FMassFragmentRequirement> Requirements, FMassFragmentIndicesMapping& OutFragmentIndices);
+
 	SIZE_T GetAllocatedSize() const;
 
 	// Converts the list of fragments into a user-readable debug string
@@ -294,6 +297,8 @@ protected:
 
 	void BindEntityRequirements(FMassExecutionContext& RunContext, const FMassFragmentIndicesMapping& EntityFragmentsMapping, FMassArchetypeChunk& Chunk, const int32 SubchunkStart, const int32 SubchunkLength);
 	void BindChunkFragmentRequirements(FMassExecutionContext& RunContext, const FMassFragmentIndicesMapping& ChunkFragmentsMapping, FMassArchetypeChunk& Chunk);
+	void BindConstSharedFragmentRequirements(FMassExecutionContext& RunContext, const FMassFragmentIndicesMapping& ChunkFragmentsMapping);
+	void BindSharedFragmentRequirements(FMassExecutionContext& RunContext, const FMassFragmentIndicesMapping& ChunkFragmentsMapping);
 
 private:
 	int32 AddEntityInternal(FMassEntityHandle Entity, const bool bInitializeFragments);
diff --git a/Engine/Plugins/Runtime/MassEntity/Source/MassEntity/Private/MassEntityQuery.cpp b/Engine/Plugins/Runtime/MassEntity/Source/MassEntity/Private/MassEntityQuery.cpp
--- a/Engine/Plugins/Runtime/MassEntity/Source/MassEntity/Private/MassEntityQuery.cpp
+++ b/Engine/Plugins/Runtime/MassEntity/Source/MassEntity/Private/MassEntityQuery.cpp
@@ -53,6 +53,8 @@ void FMassEntityQuery::SortRequirements()
 	// FragmentConfigs we want to access have already been fetched and are available in processor cache.
 	Requirements.Sort(FMassSorterOperator<FMassFragmentRequirement>());
 	ChunkRequirements.Sort(FMassSorterOperator<FMassFragmentRequirement>());
+	ConstSharedRequirements.Sort(FMassSorterOperator<FMassFragmentRequirement>());
+	SharedRequirements.Sort(FMassSorterOperator<FMassFragmentRequirement>());
 }
 
 void FMassEntityQuery::CacheArchetypes(UMassEntitySubsystem& InEntitySubsystem)
@@ -80,6 +82,14 @@ void FMassEntityQuery::CacheArchetypes(UMassEntitySubsystem& InEntitySubsystem)
 				{
 					ValidArchetypes[i].DataPtr->GetRequirementsChunkFragmentMapping(ChunkRequirements, ArchetypeFragmentMapping[i].ChunkFragments);
 				}
+				if (ConstSharedRequirements.Num())
+				{
+					ValidArchetypes[i].DataPtr->GetRequirementsConstSharedFragmentMapping(ConstSharedRequirements, ArchetypeFragmentMapping[i].ConstSharedFragments);
+				}
+				if (SharedRequirements.Num())
+				{
+					ValidArchetypes[i].DataPtr->GetRequirementsSharedFragmentMapping(SharedRequirements, ArchetypeFragmentMapping[i].SharedFragments);
+				}
 			}
 		}
 		else
@@ -137,7 +147,7 @@ void FMassEntityQuery::ForEachEntityChunk(UMassEntitySubsystem& EntitySubsystem,
 				, *DebugGetArchetypeCompatibilityDescription(ExecutionContext.GetChunkCollection().GetArchetype()));
 			return;
 		}
-		ExecutionContext.SetRequirements(Requirements, ChunkRequirements);
+		ExecutionContext.SetRequirements(Requirements, ChunkRequirements, ConstSharedRequirements, SharedRequirements);
 		ExecutionContext.GetChunkCollection().GetArchetype().DataPtr->ExecuteFunction(ExecutionContext, ExecuteFunction, {}, ExecutionContext.GetChunkCollection());
 #if WITH_MASSENTITY_DEBUG
 		NumEntitiesToProcess = ExecutionContext.GetNumEntities();
@@ -147,13 +157,13 @@ void FMassEntityQuery::ForEachEntityChunk(UMassEntitySubsystem& EntitySubsystem,
 	{
 		CacheArchetypes(EntitySubsystem);
 		// it's important to set requirements after caching archetypes due to that call potentially sorting the requirements and the order is relevant here.
-		ExecutionContext.SetRequirements(Requirements, ChunkRequirements);
+		ExecutionContext.SetRequirements(Requirements, ChunkRequirements, ConstSharedRequirements, SharedRequirements);
 
 		for (int i = 0; i < ValidArchetypes.Num(); ++i)
 		{
 			FArchetypeHandle& Archetype = ValidArchetypes[i];
 			check(Archetype.IsValid());
-			Archetype.DataPtr->ExecuteFunction(ExecutionContext, ExecuteFunction, ArchetypeFragmentMapping[i], ChunkCondition);
+			Archetype.DataPtr->ExecuteFunction(ExecutionContext, ExecuteFunction, ArchetypeFragmentMapping[i], ArchetypeCondition, ChunkCondition);
 			ExecutionContext.ClearFragmentViews();
 #if WITH_MASSENTITY_DEBUG
 			NumEntitiesToProcess += ExecutionContext.GetNumEntities();
@@ -200,7 +210,7 @@ void FMassEntityQuery::ParallelForEachEntityChunk(UMassEntitySubsystem& EntitySu
 			return;
 		}
 
-		ExecutionContext.SetRequirements(Requirements, ChunkRequirements);
+		ExecutionContext.SetRequirements(Requirements, ChunkRequirements, ConstSharedRequirements, SharedRequirements);
 		check(ArchetypeHandle.IsValid());
 		FMassArchetypeData& ArchetypeRef = *ArchetypeHandle.DataPtr.Get();
 		const FArchetypeChunkCollection AsChunkCollection(ArchetypeHandle.DataPtr);
@@ -212,7 +222,7 @@ void FMassEntityQuery::ParallelForEachEntityChunk(UMassEntitySubsystem& EntitySu
 	else
 	{
 		CacheArchetypes(EntitySubsystem);
-		ExecutionContext.SetRequirements(Requirements, ChunkRequirements);
+		ExecutionContext.SetRequirements(Requirements, ChunkRequirements, ConstSharedRequirements, SharedRequirements);
 		for (int ArchetypeIndex = 0; ArchetypeIndex < ValidArchetypes.Num(); ++ArchetypeIndex)
 		{
 			FArchetypeHandle& Archetype = ValidArchetypes[ArchetypeIndex];
diff --git a/Engine/Plugins/Runtime/MassEntity/Source/MassEntity/Private/MassEntitySubsystem.cpp b/Engine/Plugins/Runtime/MassEntity/Source/MassEntity/Private/MassEntitySubsystem.cpp
--- a/Engine/Plugins/Runtime/MassEntity/Source/MassEntity/Private/MassEntitySubsystem.cpp
+++ b/Engine/Plugins/Runtime/MassEntity/Source/MassEntity/Private/MassEntitySubsystem.cpp
@@ -105,7 +105,8 @@ FArchetypeHandle UMassEntitySubsystem::CreateArchetype(TConstArrayView<const USc
 		}
 	}
 
-	return CreateArchetype(FMassArchetypeCompositionDescriptor(FMassFragmentBitSet(FragmentList), Tags, ChunkFragments), FMassArchetypeFragmentsInitialValues());
+	const FMassArchetypeCompositionDescriptor Composition(FMassFragmentBitSet(FragmentList), Tags, ChunkFragments, FMassSharedFragmentBitSet());
+	return CreateArchetype(Composition, FMassArchetypeFragmentsInitialValues());
 }
 
 FArchetypeHandle UMassEntitySubsystem::CreateArchetype(const TSharedPtr<FMassArchetypeData>& SourceArchetype, const FMassFragmentBitSet& NewFragments)
@@ -113,7 +114,8 @@ FArchetypeHandle UMassEntitySubsystem::CreateArchetype(const TSharedPtr<FMassArc
 	check(SourceArchetype.IsValid());
 	checkf(NewFragments.IsEmpty() == false, TEXT("%s Adding an empty fragment list to an archetype is not supported."), ANSI_TO_TCHAR(__FUNCTION__));
 
-	return CreateArchetype(FMassArchetypeCompositionDescriptor(NewFragments + SourceArchetype->GetFragmentBitSet(), SourceArchetype->GetTagBitSet(), SourceArchetype->GetChunkFragmentBitSet()), SourceArchetype->GetInitialValues());
+	const FMassArchetypeCompositionDescriptor Composition(NewFragments + SourceArchetype->GetFragmentBitSet(), SourceArchetype->GetTagBitSet(), SourceArchetype->GetChunkFragmentBitSet(), SourceArchetype->GetSharedFragmentBitSet());
+	return CreateArchetype(Composition, SourceArchetype->GetInitialValues());
 }
 
 FArchetypeHandle UMassEntitySubsystem::CreateArchetype(const FMassArchetypeCompositionDescriptor& Composition, const FMassArchetypeFragmentsInitialValues& InitialValues)
@@ -125,7 +127,7 @@ FArchetypeHandle UMassEntitySubsystem::CreateArchetype(const FMassArchetypeCompo
 	FArchetypeHandle Result;
 	for (TSharedPtr<FMassArchetypeData>& Ptr : HashRow)
 	{
-		if (Ptr->IsEquivalent(Composition))
+		if (Ptr->IsEquivalent(Composition, InitialValues))
 		{
 			Result.DataPtr = Ptr;
 			break;
@@ -157,14 +159,15 @@ FArchetypeHandle UMassEntitySubsystem::InternalCreateSiblingArchetype(const TSha
 {
 	checkSlow(SourceArchetype.IsValid());
 	const FMassArchetypeData& SourceArchetypeRef = *SourceArchetype.Get();
-	const uint32 TypeHash = FMassArchetypeCompositionDescriptor::CalculateHash(SourceArchetypeRef.GetFragmentBitSet(), OverrideTags, SourceArchetypeRef.GetChunkFragmentBitSet());
+	FMassArchetypeCompositionDescriptor NewComposition(SourceArchetypeRef.GetFragmentBitSet(), OverrideTags, SourceArchetypeRef.GetChunkFragmentBitSet(), SourceArchetypeRef.GetSharedFragmentBitSet());
+	const uint32 TypeHash = NewComposition.CalculateHash();
 
 	TArray<TSharedPtr<FMassArchetypeData>>& HashRow = FragmentHashToArchetypeMap.FindOrAdd(TypeHash);
 
 	FArchetypeHandle Result;
 	for (TSharedPtr<FMassArchetypeData>& Ptr : HashRow)
 	{
-		if (Ptr->IsEquivalent(SourceArchetypeRef.GetFragmentBitSet(), OverrideTags, SourceArchetypeRef.GetChunkFragmentBitSet()))
+		if (Ptr->IsEquivalent(NewComposition, SourceArchetypeRef.GetInitialValues()))
 		{
 			Result.DataPtr = Ptr;
 			break;
@@ -246,7 +249,7 @@ FMassEntityHandle UMassEntitySubsystem::CreateEntity(TConstArrayView<FInstancedS
 {
 	check(FragmentInstanceList.Num() > 0);
 
-	FArchetypeHandle Archetype = CreateArchetype(FMassArchetypeCompositionDescriptor(FragmentInstanceList, FMassTagBitSet(), FMassChunkFragmentBitSet()), FMassArchetypeFragmentsInitialValues());
+	FArchetypeHandle Archetype = CreateArchetype(FMassArchetypeCompositionDescriptor(FragmentInstanceList, FMassTagBitSet(), FMassChunkFragmentBitSet(), FMassSharedFragmentBitSet()), FMassArchetypeFragmentsInitialValues());
 	check(Archetype.DataPtr.IsValid());
 
 	FMassEntityHandle Entity = ReserveEntity();
@@ -289,7 +292,8 @@ void UMassEntitySubsystem::BuildEntity(FMassEntityHandle Entity, TConstArrayView
 	check(FragmentInstanceList.Num() > 0);
 	checkf(!IsEntityBuilt(Entity), TEXT("Expecting an entity that is not already built"));
 
-	FArchetypeHandle Archetype = CreateArchetype(FMassArchetypeCompositionDescriptor(FragmentInstanceList, FMassTagBitSet(), FMassChunkFragmentBitSet()), FMassArchetypeFragmentsInitialValues());
+	const FMassArchetypeCompositionDescriptor Composition(FragmentInstanceList, FMassTagBitSet(), FMassChunkFragmentBitSet(), FMassSharedFragmentBitSet());
+	FArchetypeHandle Archetype = CreateArchetype(Composition, FMassArchetypeFragmentsInitialValues());
 	check(Archetype.DataPtr.IsValid());
 
 	InternalBuildEntity(Entity, Archetype);
@@ -454,15 +458,16 @@ void UMassEntitySubsystem::RemoveCompositionFromEntity(FMassEntityHandle Entity,
 
 	if(InDescriptor.IsEmpty() == false)
 	{
-	FEntityData& EntityData = Entities[Entity.Index];
-	FMassArchetypeData* OldArchetype = EntityData.CurrentArchetype.Get();
-	check(OldArchetype);
+		FEntityData& EntityData = Entities[Entity.Index];
+		FMassArchetypeData* OldArchetype = EntityData.CurrentArchetype.Get();
+		check(OldArchetype);
 
 		FMassArchetypeCompositionDescriptor NewDescriptor = OldArchetype->GetCompositionDescriptor();
-	NewDescriptor.Fragments -= InDescriptor.Fragments;
-	NewDescriptor.Tags -= InDescriptor.Tags;
+		NewDescriptor.Fragments -= InDescriptor.Fragments;
+		NewDescriptor.Tags -= InDescriptor.Tags;
 
 		ensureMsgf(InDescriptor.ChunkFragments.IsEmpty(), TEXT("Removing chunk fragments is not supported"));
+		ensureMsgf(InDescriptor.SharedFragments.IsEmpty(), TEXT("Removing shared fragments is not supported"));
 
 		if (NewDescriptor.IsEquivalent(InDescriptor) == false)
 	{
@@ -562,7 +567,8 @@ void UMassEntitySubsystem::RemoveFragmentListFromEntity(FMassEntityHandle Entity
 	if (ensureMsgf(OldArchetype->GetFragmentBitSet().HasAny(FragmentsToRemove), TEXT("Trying to remove a list of fragments from an entity but none of the fragments given was found.")))
 	{
 		// If all the fragments got removed this will result in fetching of the empty archetype
-		const FArchetypeHandle NewArchetypeHandle = CreateArchetype(FMassArchetypeCompositionDescriptor(OldArchetype->GetFragmentBitSet() - FragmentsToRemove, OldArchetype->GetTagBitSet(), OldArchetype->GetChunkFragmentBitSet()), OldArchetype->GetInitialValues());
+		const FMassArchetypeCompositionDescriptor NewComposition(OldArchetype->GetFragmentBitSet() - FragmentsToRemove, OldArchetype->GetTagBitSet(), OldArchetype->GetChunkFragmentBitSet(), OldArchetype->GetSharedFragmentBitSet());
+		const FArchetypeHandle NewArchetypeHandle = CreateArchetype(NewComposition, OldArchetype->GetInitialValues());
 
 		// Move the entity over
 		FMassArchetypeData* NewArchetype = NewArchetypeHandle.DataPtr.Get();
@@ -845,6 +851,33 @@ void UMassEntitySubsystem::GetValidArchetypes(const FMassEntityQuery& Query, TAr
 			continue;
 		}
 
+		if (Archetype.GetSharedFragmentBitSet().HasAll(Query.GetRequiredAllSharedFragments()) == false)
+		{
+			// missing some required fragments, skip.
+#if WITH_MASSENTITY_DEBUG
+			const FMassSharedFragmentBitSet UnsatisfiedFragments = Query.GetRequiredAllSharedFragments() - Archetype.GetSharedFragmentBitSet();
+			FStringOutputDevice Description;
+			UnsatisfiedFragments.DebugGetStringDesc(Description);
+			UE_LOG(LogMass, VeryVerbose, TEXT("Archetype did not match due to missing Shared Fragments: %s")
+				, *Description);
+#endif // WITH_MASSENTITY_DEBUG
+			continue;
+		}
+
+		if (Archetype.GetSharedFragmentBitSet().HasNone(Query.GetRequiredNoneSharedFragments()) == false)
+		{
+			// has some Fragments required to be absent
+#if WITH_MASSENTITY_DEBUG
+			const FMassSharedFragmentBitSet UnwantedFragments = Query.GetRequiredNoneSharedFragments().GetOverlap(Archetype.GetSharedFragmentBitSet());
+			FStringOutputDevice Description;
+			UnwantedFragments.DebugGetStringDesc(Description);
+			UE_LOG(LogMass, VeryVerbose, TEXT("Archetype has Shared Fragments required absent: %s")
+				, *Description);
+#endif // WITH_MASSENTITY_DEBUG
+			continue;
+		}
+
+
 		OutValidArchetypes.Add(ArchetypePtr);
 	}
 }
@@ -1054,23 +1087,44 @@ void FMassExecutionContext::SetChunkCollection(FArchetypeChunkCollection&& InChu
 	ChunkCollection = MoveTemp(InChunkCollection);
 }
 
-void FMassExecutionContext::SetRequirements(TConstArrayView<FMassFragmentRequirement> InRequirements, TConstArrayView<FMassFragmentRequirement> InChunkRequirements)
+void FMassExecutionContext::SetRequirements(TConstArrayView<FMassFragmentRequirement> InRequirements, 
+	TConstArrayView<FMassFragmentRequirement> InChunkRequirements, 
+	TConstArrayView<FMassFragmentRequirement> InConstSharedRequirements, 
+	TConstArrayView<FMassFragmentRequirement> InSharedRequirements)
 { 
 	FragmentViews.Reset();
 	for (const FMassFragmentRequirement& Requirement : InRequirements)
 	{
 		if (Requirement.RequiresBinding())
 		{
-			FragmentViews.Add(FFragmentView(Requirement));
+			FragmentViews.Emplace(Requirement);
 		}
 	}
 
-	ChunkFragments.Reset();
+	ChunkFragmentViews.Reset();
 	for (const FMassFragmentRequirement& Requirement : InChunkRequirements)
 	{
 		if (Requirement.RequiresBinding())
 		{
-			ChunkFragments.Add(FChunkFragmentView(Requirement));
+			ChunkFragmentViews.Emplace(Requirement);
+		}
+	}
+
+	ConstSharedFragmentViews.Reset();
+	for (const FMassFragmentRequirement& Requirement : InConstSharedRequirements)
+	{
+		if (Requirement.RequiresBinding())
+		{
+			ConstSharedFragmentViews.Emplace(Requirement);
+		}
+	}
+
+	SharedFragmentViews.Reset();
+	for (const FMassFragmentRequirement& Requirement : InSharedRequirements)
+	{
+		if (Requirement.RequiresBinding())
+		{
+			SharedFragmentViews.Emplace(Requirement);
 		}
 	}
 }
diff --git a/Engine/Plugins/Runtime/MassEntity/Source/MassEntity/Private/MassEntityTypes.cpp b/Engine/Plugins/Runtime/MassEntity/Source/MassEntity/Private/MassEntityTypes.cpp
new file mode 100644
--- /dev/null
+++ b/Engine/Plugins/Runtime/MassEntity/Source/MassEntity/Private/MassEntityTypes.cpp
@@ -0,0 +1,31 @@
+// Copyright Epic Games, Inc. All Rights Reserved.
+
+#include "MassEntityTypes.h"
+#include "StructUtilsTypes.h"
+
+uint32 FMassArchetypeFragmentsInitialValues::CalculateHash() const
+{
+	TArray<uint32, TInlineAllocator<16>> FragmentHashes;
+	FragmentHashes.Reserve(Fragments.Num() + ChunkFragments.Num() + ConstSharedFragments.Num() + SharedFragments.Num());
+
+	// max@todo: Fragments and chunk fragments are not part of the uniqueness and should be removed from the initial values, also should think of a better name for this struct.
+	for (const FConstSharedStruct& Fragment : ConstSharedFragments)
+	{
+		FragmentHashes.Add(PointerHash(Fragment.GetMemory()));
+	}
+
+	for (const FSharedStruct& Fragment : SharedFragments)
+	{
+		FragmentHashes.Add(PointerHash(Fragment.GetMemory()));
+	}
+
+	FragmentHashes.Sort();
+
+	uint32 CalculatedHash = 0;
+	for (const uint32 FragmentHash : FragmentHashes)
+	{
+		CalculatedHash = HashCombine(CalculatedHash, FragmentHash);
+	}
+
+	return CalculatedHash;
+}
diff --git a/Engine/Plugins/Runtime/MassEntity/Source/MassEntity/Public/MassArchetypeTypes.h b/Engine/Plugins/Runtime/MassEntity/Source/MassEntity/Public/MassArchetypeTypes.h
--- a/Engine/Plugins/Runtime/MassEntity/Source/MassEntity/Public/MassArchetypeTypes.h
+++ b/Engine/Plugins/Runtime/MassEntity/Source/MassEntity/Public/MassArchetypeTypes.h
@@ -18,6 +18,7 @@ struct FMassEntityView;
 
 typedef TFunction< void(FMassExecutionContext& /*ExecutionContext*/) > FMassExecuteFunction;
 typedef TFunction< bool(const FMassExecutionContext& /*ExecutionContext*/) > FMassChunkConditionFunction;
+typedef TFunction< bool(const FMassExecutionContext& /*ExecutionContext*/) > FMassArchetypeConditionFunction;
 
 
 //////////////////////////////////////////////////////////////////////
@@ -140,6 +141,8 @@ struct FMassQueryRequirementIndicesMapping
 
 	FMassFragmentIndicesMapping EntityFragments;
 	FMassFragmentIndicesMapping ChunkFragments;
+	FMassFragmentIndicesMapping ConstSharedFragments;
+	FMassFragmentIndicesMapping SharedFragments;
 	FORCEINLINE bool IsEmpty() const
 	{
 		return EntityFragments.Num() == 0 || ChunkFragments.Num() == 0;
diff --git a/Engine/Plugins/Runtime/MassEntity/Source/MassEntity/Public/MassEntityQuery.h b/Engine/Plugins/Runtime/MassEntity/Source/MassEntity/Public/MassEntityQuery.h
--- a/Engine/Plugins/Runtime/MassEntity/Source/MassEntity/Public/MassEntityQuery.h
+++ b/Engine/Plugins/Runtime/MassEntity/Source/MassEntity/Public/MassEntityQuery.h
@@ -255,6 +255,58 @@ public:
 		return *this;
 	}
 
+	template<typename T>
+	FMassEntityQuery& AddConstSharedRequirement(const EMassFragmentPresence Presence = EMassFragmentPresence::All)
+	{
+		static_assert(TIsDerivedFrom<T, FMassSharedFragment>::IsDerived, "Given struct doesn't represent a valid shared fragment type. Make sure to inherit from FMassSharedFragment or one of its child-types.");
+		checkf(ConstSharedRequirements.FindByPredicate([](const FMassFragmentRequirement& Item) { return Item.StructType == T::StaticStruct(); }) == nullptr
+			, TEXT("Duplicated requirements are not supported. %s already present"), *T::StaticStruct()->GetName());
+		checkfSlow(Presence != EMassFragmentPresence::Any, TEXT("\'Any\' is not a valid Presence value for AddSharedRequirement."));
+
+		switch (Presence)
+		{
+		case EMassFragmentPresence::All:
+			RequiredAllSharedFragments.Add<T>();
+			ConstSharedRequirements.Emplace(T::StaticStruct(), EMassFragmentAccess::ReadOnly, Presence);
+			break;
+		case EMassFragmentPresence::Optional:
+			RequiredOptionalSharedFragments.Add<T>();
+			ConstSharedRequirements.Emplace(T::StaticStruct(), EMassFragmentAccess::ReadOnly, Presence);
+			break;
+		case EMassFragmentPresence::None:
+			RequiredNoneSharedFragments.Add<T>();
+			break;
+		}
+
+		return *this;
+	}
+
+	template<typename T>
+	FMassEntityQuery& AddSharedRequirement(const EMassFragmentAccess AccessMode, const EMassFragmentPresence Presence = EMassFragmentPresence::All)
+	{
+		static_assert(TIsDerivedFrom<T, FMassSharedFragment>::IsDerived, "Given struct doesn't represent a valid shared fragment type. Make sure to inherit from FMassSharedFragment or one of its child-types.");
+		checkf(SharedRequirements.FindByPredicate([](const FMassFragmentRequirement& Item) { return Item.StructType == T::StaticStruct(); }) == nullptr
+			, TEXT("Duplicated requirements are not supported. %s already present"), *T::StaticStruct()->GetName());
+		checkfSlow(Presence != EMassFragmentPresence::Any, TEXT("\'Any\' is not a valid Presence value for AddSharedRequirement."));
+
+		switch (Presence)
+		{
+		case EMassFragmentPresence::All:
+			RequiredAllSharedFragments.Add<T>();
+			SharedRequirements.Emplace(T::StaticStruct(), AccessMode, Presence);
+			break;
+		case EMassFragmentPresence::Optional:
+			RequiredOptionalSharedFragments.Add<T>();
+			SharedRequirements.Emplace(T::StaticStruct(), AccessMode, Presence);
+			break;
+		case EMassFragmentPresence::None:
+			RequiredNoneSharedFragments.Add<T>();
+			break;
+		}
+
+		return *this;
+	}
+
 	void Clear()
 	{
 		Requirements.Empty();
@@ -288,6 +340,9 @@ public:
 	const FMassChunkFragmentBitSet& GetRequiredAllChunkFragments() const { return RequiredAllChunkFragments; }
 	const FMassChunkFragmentBitSet& GetRequiredOptionalChunkFragments() const { return RequiredOptionalChunkFragments; }
 	const FMassChunkFragmentBitSet& GetRequiredNoneChunkFragments() const { return RequiredNoneChunkFragments; }
+	const FMassSharedFragmentBitSet& GetRequiredAllSharedFragments() const { return RequiredAllSharedFragments; }
+	const FMassSharedFragmentBitSet& GetRequiredOptionalSharedFragments() const { return RequiredOptionalSharedFragments; }
+	const FMassSharedFragmentBitSet& GetRequiredNoneSharedFragments() const { return RequiredNoneSharedFragments; }
 
 	const TArray<FArchetypeHandle>& GetArchetypes() const
 	{ 
@@ -319,6 +374,16 @@ public:
 
 	bool HasChunkFilter() const { return bool(ChunkCondition); }
 
+	/**
+	 * Sets a archetype filter condition that will applied to each valid archetypes.
+	 * The value returned by InFunction controls whether to allow execution (true) or block it (false).
+	 */
+	void SetArchetypeFilter(const FMassArchetypeConditionFunction& InFunction) { ArchetypeCondition = InFunction; }
+
+	void ClearArchetypeFilter() { ArchetypeCondition.Reset(); }
+
+	bool HasArchetypeFilter() const { return bool(ArchetypeCondition); }
+
 protected:
 	void SortRequirements();
 	void ReadCommandlineParams();
@@ -326,6 +391,8 @@ protected:
 protected:
 	TArray<FMassFragmentRequirement> Requirements;
 	TArray<FMassFragmentRequirement> ChunkRequirements;
+	TArray<FMassFragmentRequirement> ConstSharedRequirements;
+	TArray<FMassFragmentRequirement> SharedRequirements;
 	FMassTagBitSet RequiredAllTags;
 	FMassTagBitSet RequiredAnyTags;
 	FMassTagBitSet RequiredNoneTags;
@@ -336,6 +403,9 @@ protected:
 	FMassChunkFragmentBitSet RequiredAllChunkFragments;
 	FMassChunkFragmentBitSet RequiredOptionalChunkFragments;
 	FMassChunkFragmentBitSet RequiredNoneChunkFragments;
+	FMassSharedFragmentBitSet RequiredAllSharedFragments;
+	FMassSharedFragmentBitSet RequiredOptionalSharedFragments;
+	FMassSharedFragmentBitSet RequiredNoneSharedFragments;
 
 private:
 	/** 
@@ -345,6 +415,13 @@ private:
 	 */
 	FMassChunkConditionFunction ChunkCondition;
 
+	/**
+	 * This function represents a condition that will be called for every archetype to be processed before the actual
+	 * execution function is called. The shared fragment requirements are already bound and ready to be used by the time
+	 * ArchetypeCondition is executed.
+	 */
+	FMassArchetypeConditionFunction ArchetypeCondition;
+
 	uint32 EntitySubsystemHash = 0;
 	uint32 ArchetypeDataVersion = 0;
 
diff --git a/Engine/Plugins/Runtime/MassEntity/Source/MassEntity/Public/MassEntitySubsystem.h b/Engine/Plugins/Runtime/MassEntity/Source/MassEntity/Public/MassEntitySubsystem.h
--- a/Engine/Plugins/Runtime/MassEntity/Source/MassEntity/Public/MassEntitySubsystem.h
+++ b/Engine/Plugins/Runtime/MassEntity/Source/MassEntity/Public/MassEntitySubsystem.h
@@ -254,6 +254,48 @@ public:
 	bool IsProcessing() const { return ProcessingScopeCount > 0; }
 	FMassCommandBuffer& Defer() const { return *DeferredCommandBuffer.Get(); }
 
+	/**
+	 * Shared fragment creation methods
+	 */
+	template<typename T>
+	FConstSharedStruct& GetOrCreateConstSharedFragment(const uint32 Hash, const T& Fragment)
+	{
+		static_assert(TIsDerivedFrom<T, FMassSharedFragment>::IsDerived, "Given struct doesn't represent a valid shared fragment type. Make sure to inherit from FMassSharedFragment or one of its child-types.");
+		int32& Index = ConstSharedFragmentsMap.FindOrAddByHash(Hash, Hash, INDEX_NONE);
+		if (Index == INDEX_NONE)
+		{
+			Index = ConstSharedFragments.Add(FSharedStruct::Make(Fragment));
+		}
+		return ConstSharedFragments[Index];
+	}
+
+	template<typename T, typename... TArgs>
+	FSharedStruct& GetOrCreateSharedFragment(const uint32 Hash, TArgs&&... InArgs)
+	{
+		static_assert(TIsDerivedFrom<T, FMassSharedFragment>::IsDerived, "Given struct doesn't represent a valid shared fragment type. Make sure to inherit from FMassSharedFragment or one of its child-types.");
+
+		int32& Index = SharedFragmentsMap.FindOrAddByHash(Hash, Hash, INDEX_NONE);
+		if (Index == INDEX_NONE)
+		{
+			Index = SharedFragments.Add(FSharedStruct::Make<T>(Forward<TArgs>(InArgs)...));
+		}
+
+		return SharedFragments[Index];
+	}
+
+	template<typename T>
+	void ForEachSharedFragment(TFunction< void(T& /*SharedFragment*/) > ExecuteFunction)
+	{
+		FSameTypeScriptStructPredicate Predicate(T::StaticStruct());
+		for (const FSharedStruct& Struct : SharedFragments)
+		{
+			if (Predicate(Struct))
+			{
+				ExecuteFunction(Struct.GetMutable<T>());
+			}
+		}
+	}
+
 #if WITH_MASSENTITY_DEBUG
 	void DebugPrintEntity(int32 Index, FOutputDevice& Ar, const TCHAR* InPrefix = TEXT("")) const;
 	void DebugPrintEntity(FMassEntityHandle Entity, FOutputDevice& Ar, const TCHAR* InPrefix = TEXT("")) const;
@@ -308,6 +350,17 @@ private:
 
 	// Map to list of archetypes that contain the specified fragment type
 	TMap<const UScriptStruct*, TArray<TSharedPtr<FMassArchetypeData>>> FragmentTypeToArchetypeMap;
+
+	// Shared fragments
+	UPROPERTY(Transient)
+	TArray<FConstSharedStruct> ConstSharedFragments;
+	// Hash/Index in array pair
+	TMap<uint32, int32> ConstSharedFragmentsMap;
+
+	UPROPERTY(Transient)
+	TArray<FSharedStruct> SharedFragments;
+	// Hash/Index in array pair
+	TMap<uint32, int32> SharedFragmentsMap;
 };
 
 
@@ -317,30 +370,30 @@ private:
 struct MASSENTITY_API FMassExecutionContext
 {
 private:
-	struct FFragmentView 
+
+	template< typename ViewType >
+	struct TFragmentView 
 	{
 		FMassFragmentRequirement Requirement;
-		TArrayView<FMassFragment> FragmentView;
+		ViewType FragmentView;
 
-		FFragmentView() {}
-		explicit FFragmentView(const FMassFragmentRequirement& InRequirement) : Requirement(InRequirement) {}
+		TFragmentView() {}
+		explicit TFragmentView(const FMassFragmentRequirement& InRequirement) : Requirement(InRequirement) {}
 
 		bool operator==(const UScriptStruct* FragmentType) const { return Requirement.StructType == FragmentType; }
 	};
+	using FFragmentView=TFragmentView<TArrayView<FMassFragment>>;
 	TArray<FFragmentView, TInlineAllocator<8>> FragmentViews;
 
-	struct FChunkFragmentView
-	{
-		FMassFragmentRequirement Requirement;
-		FStructView ChunkFragmentView;
+	using FChunkFragmentView = TFragmentView<FStructView>;
+	TArray<FChunkFragmentView, TInlineAllocator<4>> ChunkFragmentViews;
 
-		FChunkFragmentView() {}
-		explicit FChunkFragmentView(const FMassFragmentRequirement& InRequirement) : Requirement(InRequirement)	{}
-
-		bool operator==(const UScriptStruct* FragmentType) const { return Requirement.StructType == FragmentType; }
-	};
-	TArray<FChunkFragmentView, TInlineAllocator<4>> ChunkFragments;
+	using FConstSharedFragmentView = TFragmentView<FConstStructView>;
+	TArray<FConstSharedFragmentView, TInlineAllocator<4>> ConstSharedFragmentViews;
 
+	using FSharedFragmentView = TFragmentView<FStructView>;
+	TArray<FSharedFragmentView, TInlineAllocator<4>> SharedFragmentViews;
+	
 	// mz@todo make this shared ptr thread-safe and never auto-flush in MT environment. 
 	TSharedPtr<FMassCommandBuffer> DeferredCommandBuffer;
 	TArrayView<FMassEntityHandle> EntityListView;
@@ -364,8 +417,10 @@ private:
 	bool bFlushDeferredCommands = true;
 
 	TArrayView<FFragmentView> GetMutableRequirements() { return FragmentViews; }
-	TArrayView<FChunkFragmentView> GetMutableChunkRequirements() { return ChunkFragments; }
-	
+	TArrayView<FChunkFragmentView> GetMutableChunkRequirements() { return ChunkFragmentViews; }
+	TArrayView<FConstSharedFragmentView> GetMutableConstSharedRequirements() { return ConstSharedFragmentViews; }
+	TArrayView<FSharedFragmentView> GetMutableSharedRequirements() { return SharedFragmentViews; }
+
 	friend FMassArchetypeData;
 	friend FMassEntityQuery;
 
@@ -413,7 +468,7 @@ public:
 		return CurrentArchetypesTagBitSet.Contains<T>();
 	}
 
-	/** Chunk related operation */
+	/** Chunk related operations */
 	void SetCurrentChunkSerialModificationNumber(const int32 SerialModificationNumber) { ChunkSerialModificationNumber = SerialModificationNumber; }
 	int32 GetChunkSerialModificationNumber() const { return ChunkSerialModificationNumber; }
 
@@ -423,8 +478,8 @@ public:
 		static_assert(TIsDerivedFrom<T, FMassChunkFragment>::IsDerived, "Given struct doesn't represent a valid chunk fragment type. Make sure to inherit from FMassChunkFragment or one of its child-types.");
 
 		const UScriptStruct* Type = T::StaticStruct();
-		const FChunkFragmentView* FoundChunkFragmentData = ChunkFragments.FindByPredicate([Type](const FChunkFragmentView& Element) { return Element.Requirement.StructType == Type; } );
-		return FoundChunkFragmentData ? FoundChunkFragmentData->ChunkFragmentView.GetMutablePtr<T>() : static_cast<T*>(nullptr);
+		const FChunkFragmentView* FoundChunkFragmentData = ChunkFragmentViews.FindByPredicate([Type](const FChunkFragmentView& Element) { return Element.Requirement.StructType == Type; } );
+		return FoundChunkFragmentData ? FoundChunkFragmentData->FragmentView.GetMutablePtr<T>() : static_cast<T*>(nullptr);
 	}
 	
 	template<typename T>
@@ -438,21 +493,68 @@ public:
 	template<typename T>
 	const T* GetChunkFragmentPtr() const
 	{
-		static_assert(TIsDerivedFrom<T, FMassChunkFragment>::IsDerived, "Given struct doesn't represent a valid chunk fragment type. Make sure to inherit from FMassChunkFragment or one of its child-types.");
-
-		const UScriptStruct* Type = T::StaticStruct();
-		const FChunkFragmentView* FoundChunkFragmentData = ChunkFragments.FindByPredicate([Type](const FChunkFragmentView& Element) { return Element.Requirement.StructType == Type; } );
-		return FoundChunkFragmentData ? FoundChunkFragmentData->ChunkFragmentView.GetPtr<T>() : static_cast<const T*>(nullptr);
+		return GetMutableChunkFragmentPtr<T>();
 	}
 	
 	template<typename T>
 	const T& GetChunkFragment() const
 	{
-		const T* ChunkFragment = GetChunkFragmentPtr<T>();
+		const T* ChunkFragment = GetMutableChunkFragmentPtr<T>();
 		checkf(ChunkFragment, TEXT("Chunk Fragment requirement not found: %s"), *T::StaticStruct()->GetName());
 		return *ChunkFragment;
 	}
-	
+
+	/** Shared fragment related operations */
+	template<typename T>
+	const T* GetConstSharedFragmentPtr() const
+	{
+		static_assert(TIsDerivedFrom<T, FMassSharedFragment>::IsDerived, "Given struct doesn't represent a valid Shared fragment type. Make sure to inherit from FMassSharedFragment or one of its child-types.");
+
+		const FConstSharedFragmentView* FoundSharedFragmentData = ConstSharedFragmentViews.FindByPredicate([](const FConstSharedFragmentView& Element) { return Element.Requirement.StructType == T::StaticStruct(); });
+		return FoundSharedFragmentData ? FoundSharedFragmentData->FragmentView.GetPtr<T>() : static_cast<const T*>(nullptr);
+	}
+
+	template<typename T>
+	const T& GetConstSharedFragment() const
+	{
+		const T* SharedFragment = GetConstSharedFragmentPtr<T>();
+		checkf(SharedFragment, TEXT("Shared Fragment requirement not found: %s"), *T::StaticStruct()->GetName());
+		return *SharedFragment;
+	}
+
+
+	template<typename T>
+	T* GetMutableSharedFragmentPtr() const
+	{
+		static_assert(TIsDerivedFrom<T, FMassSharedFragment>::IsDerived, "Given struct doesn't represent a valid shared fragment type. Make sure to inherit from FMassSharedFragment or one of its child-types.");
+
+		const FSharedFragmentView* FoundSharedFragmentData = SharedFragmentViews.FindByPredicate([](const FSharedFragmentView& Element) { return Element.Requirement.StructType == T::StaticStruct(); });
+		return FoundSharedFragmentData ? FoundSharedFragmentData->FragmentView.GetMutablePtr<T>() : static_cast<T*>(nullptr);
+	}
+
+	template<typename T>
+	T& GetMutableSharedFragment() const
+	{
+		T* SharedFragment = GetMutableSharedFragmentPtr<T>();
+		checkf(SharedFragment, TEXT("Shared Fragment requirement not found: %s"), *T::StaticStruct()->GetName());
+		return *SharedFragment;
+	}
+
+	template<typename T>
+	const T* GetSharedFragmentPtr() const
+	{
+		return GetMutableSharedFragmentPtr<T>();
+	}
+
+	template<typename T>
+	const T& GetSharedFragment() const
+	{
+		const T* SharedFragment = GetSharedFragmentPtr<T>();
+		checkf(SharedFragment, TEXT("Shared Fragment requirement not found: %s"), *T::StaticStruct()->GetName());
+		return *SharedFragment;
+	}
+
+	/* Fragments related operations */
 	template<typename TFragment>
 	TArrayView<TFragment> GetMutableFragmentView()
 	{
@@ -498,16 +600,28 @@ public:
 	void SetCurrentArchetypeData(FMassArchetypeData& ArchetypeData);
 
 protected:
-	void SetRequirements(TConstArrayView<FMassFragmentRequirement> InRequirements, TConstArrayView<FMassFragmentRequirement> InChunkRequirements);
+	void SetRequirements(TConstArrayView<FMassFragmentRequirement> InRequirements, 
+		TConstArrayView<FMassFragmentRequirement> InChunkRequirements, 
+		TConstArrayView<FMassFragmentRequirement> InConstSharedRequirements, 
+		TConstArrayView<FMassFragmentRequirement> InSharedRequirements);
+
 	void ClearFragmentViews()
 	{
 		for (FFragmentView& View : FragmentViews)
 		{
 			View.FragmentView = TArrayView<FMassFragment>();
 		}
-		for (FChunkFragmentView& View : ChunkFragments)
+		for (FChunkFragmentView& View : ChunkFragmentViews)
+		{
+			View.FragmentView.Reset();
+		}
+		for (FConstSharedFragmentView& View : ConstSharedFragmentViews)
+		{
+			View.FragmentView.Reset();
+		}
+		for (FSharedFragmentView& View : SharedFragmentViews)
 		{
-			View.ChunkFragmentView.Reset();
+			View.FragmentView.Reset();
 		}
 	}
 };
diff --git a/Engine/Plugins/Runtime/MassEntity/Source/MassEntity/Public/MassEntityTypes.h b/Engine/Plugins/Runtime/MassEntity/Source/MassEntity/Public/MassEntityTypes.h
--- a/Engine/Plugins/Runtime/MassEntity/Source/MassEntity/Public/MassEntityTypes.h
+++ b/Engine/Plugins/Runtime/MassEntity/Source/MassEntity/Public/MassEntityTypes.h
@@ -36,6 +36,14 @@ struct FMassChunkFragment
 	FMassChunkFragment() {}
 };
 
+USTRUCT()
+struct FMassSharedFragment
+{
+	GENERATED_BODY()
+
+	FMassSharedFragment() {}
+};
+
 // A handle to a lightweight entity.  An entity is used in conjunction with the UMassEntitySubsystem
 // for the current world and can contain lightweight fragments.
 USTRUCT()
@@ -101,26 +109,34 @@ template struct MASSENTITY_API TScriptStructTypeBitSet<FMassTag>;
 using FMassTagBitSet = TScriptStructTypeBitSet<FMassTag>;
 template struct MASSENTITY_API TScriptStructTypeBitSet<FMassChunkFragment>;
 using FMassChunkFragmentBitSet = TScriptStructTypeBitSet<FMassChunkFragment>;
+template struct MASSENTITY_API TScriptStructTypeBitSet<FMassSharedFragment>;
+using FMassSharedFragmentBitSet = TScriptStructTypeBitSet<FMassSharedFragment>;
 
 /** The type summarily describing a composition of an entity or an archetype. It contains information on both the
  *  fragments as well as tags */
 struct FMassArchetypeCompositionDescriptor
 {
 	FMassArchetypeCompositionDescriptor() = default;
-	FMassArchetypeCompositionDescriptor(const FMassFragmentBitSet& InFragments, const FMassTagBitSet& InTags, const FMassChunkFragmentBitSet& InChunkFragments)
-		: Fragments(InFragments), Tags(InTags), ChunkFragments(InChunkFragments)
+	FMassArchetypeCompositionDescriptor(const FMassFragmentBitSet& InFragments, const FMassTagBitSet& InTags, const FMassChunkFragmentBitSet& InChunkFragments, const FMassSharedFragmentBitSet& InSharedFragments)
+		: Fragments(InFragments)
+		, Tags(InTags)
+		, ChunkFragments(InChunkFragments)
+		, SharedFragments(InSharedFragments)
 	{}
 
-	FMassArchetypeCompositionDescriptor(TConstArrayView<const UScriptStruct*> InFragments, const FMassTagBitSet& InTags, const FMassChunkFragmentBitSet& InChunkFragments)
-		: FMassArchetypeCompositionDescriptor(FMassFragmentBitSet(InFragments), InTags, InChunkFragments)
+	FMassArchetypeCompositionDescriptor(TConstArrayView<const UScriptStruct*> InFragments, const FMassTagBitSet& InTags, const FMassChunkFragmentBitSet& InChunkFragments, const FMassSharedFragmentBitSet& InSharedFragments)
+		: FMassArchetypeCompositionDescriptor(FMassFragmentBitSet(InFragments), InTags, InChunkFragments, InSharedFragments)
 	{}
 
-	FMassArchetypeCompositionDescriptor(TConstArrayView<FInstancedStruct> InFragmentInstances, const FMassTagBitSet& InTags, const FMassChunkFragmentBitSet& InChunkFragments)
-		: FMassArchetypeCompositionDescriptor(FMassFragmentBitSet(InFragmentInstances), InTags, InChunkFragments)
+	FMassArchetypeCompositionDescriptor(TConstArrayView<FInstancedStruct> InFragmentInstances, const FMassTagBitSet& InTags, const FMassChunkFragmentBitSet& InChunkFragments, const FMassSharedFragmentBitSet& InSharedFragments)
+		: FMassArchetypeCompositionDescriptor(FMassFragmentBitSet(InFragmentInstances), InTags, InChunkFragments, InSharedFragments)
 	{}
 
-	FMassArchetypeCompositionDescriptor(FMassFragmentBitSet&& InFragments, FMassTagBitSet&& InTags, FMassChunkFragmentBitSet&& InChunkFragments)
-		: Fragments(MoveTemp(InFragments)), Tags(MoveTemp(InTags)), ChunkFragments(MoveTemp(InChunkFragments))
+	FMassArchetypeCompositionDescriptor(FMassFragmentBitSet&& InFragments, FMassTagBitSet&& InTags, FMassChunkFragmentBitSet&& InChunkFragments, FMassSharedFragmentBitSet&& InSharedFragments)
+		: Fragments(MoveTemp(InFragments))
+		, Tags(MoveTemp(InTags))
+		, ChunkFragments(MoveTemp(InChunkFragments))
+		, SharedFragments(MoveTemp(InSharedFragments))
 	{}
 
 	void Reset()
@@ -130,29 +146,34 @@ struct FMassArchetypeCompositionDescriptor
 		ChunkFragments.Reset();
 	}
 
-	bool IsEquivalent(const FMassFragmentBitSet& InFragmentBitSet, const FMassTagBitSet& InTagBitSet, const FMassChunkFragmentBitSet& InChunkFragmentsBitSet) const
-	{
-		return Fragments.IsEquivalent(InFragmentBitSet) && Tags.IsEquivalent(InTagBitSet) && ChunkFragments.IsEquivalent(InChunkFragmentsBitSet);
-	}
-
 	bool IsEquivalent(const FMassArchetypeCompositionDescriptor& OtherDescriptor) const
 	{
-		return IsEquivalent(OtherDescriptor.Fragments, OtherDescriptor.Tags, OtherDescriptor.ChunkFragments);
+		return Fragments.IsEquivalent(OtherDescriptor.Fragments) &&
+			Tags.IsEquivalent(OtherDescriptor.Tags) &&
+			ChunkFragments.IsEquivalent(OtherDescriptor.ChunkFragments) &&
+			SharedFragments.IsEquivalent(OtherDescriptor.SharedFragments);
 	}
 
-	bool IsEmpty() const { return Fragments.IsEmpty() && Tags.IsEmpty() && ChunkFragments.IsEmpty(); }
+	bool IsEmpty() const 
+	{ 
+		return Fragments.IsEmpty() &&
+			Tags.IsEmpty() &&
+			ChunkFragments.IsEmpty() &&
+			SharedFragments.IsEmpty();
+	}
 
-	static uint32 CalculateHash(const FMassFragmentBitSet& InFragments, const FMassTagBitSet& InTags, const FMassChunkFragmentBitSet& InChunkFragments)
+	static uint32 CalculateHash(const FMassFragmentBitSet& InFragments, const FMassTagBitSet& InTags, const FMassChunkFragmentBitSet& InChunkFragments, const FMassSharedFragmentBitSet& InSharedFragmentBitSet)
 	{
 		const uint32 FragmentsHash = GetTypeHash(InFragments);
 		const uint32 TagsHash = GetTypeHash(InTags);
 		const uint32 ChunkFragmentsHash = GetTypeHash(InChunkFragments);
-		return HashCombine(HashCombine(FragmentsHash, TagsHash), ChunkFragmentsHash);
+		const uint32 SharedFragmentsHash = GetTypeHash(InSharedFragmentBitSet);
+		return HashCombine(HashCombine(HashCombine(FragmentsHash, TagsHash), ChunkFragmentsHash), SharedFragmentsHash);
 	}	
 
 	uint32 CalculateHash() const 
 	{
-		return CalculateHash(Fragments, Tags, ChunkFragments);
+		return CalculateHash(Fragments, Tags, ChunkFragments, SharedFragments);
 	}
 
 	void DebugOutputDescription(FOutputDevice& Ar) const
@@ -176,14 +197,131 @@ struct FMassArchetypeCompositionDescriptor
 	FMassFragmentBitSet Fragments;
 	FMassTagBitSet Tags;
 	FMassChunkFragmentBitSet ChunkFragments;
+	FMassSharedFragmentBitSet SharedFragments;
 };
 
 struct FMassArchetypeFragmentsInitialValues
 {
 	FMassArchetypeFragmentsInitialValues() = default;
 
+	FORCEINLINE bool IsEquivalent(const FMassArchetypeFragmentsInitialValues& OtherInitialValues) const
+	{
+		return GetTypeHash(*this) == GetTypeHash(OtherInitialValues);
+	}
+
+	template<typename T>
+	FORCEINLINE T& AddFragment(const T& Fragment)
+	{
+		DirtyHashCache();
+		return Fragments.Add_GetRef(FConstStructView::Make(Fragment)).template GetMutable<T>();
+	}
+
+	FORCEINLINE void AddFragment(FConstStructView Fragment)
+	{
+		DirtyHashCache();
+		Fragments.Add(Fragment);
+	}
+
+	template<typename T>
+	FORCEINLINE T& AddChunkFragment(const T& Fragment)
+	{
+		DirtyHashCache();
+		return ChunkFragments.Add_GetRef(FConstStructView::Make(Fragment)).template GetMutable<T>();
+	}
+
+	FORCEINLINE void AddChunkFragment(FConstStructView Fragment)
+	{
+		DirtyHashCache();
+		ChunkFragments.Add(Fragment);
+	}
+
+	FORCEINLINE FConstSharedStruct& AddConstSharedFragment(const FConstSharedStruct& Fragment)
+	{
+		DirtyHashCache();
+		return ConstSharedFragments.Add_GetRef(Fragment);
+	}
+
+	FORCEINLINE FSharedStruct AddSharedFragment(const FSharedStruct& Fragment)
+	{
+		DirtyHashCache();
+		return SharedFragments.Add_GetRef(Fragment);
+	}
+
+	FORCEINLINE const TArray<FInstancedStruct>& GetFragments() const
+	{
+		return Fragments;
+	}
+
+	FORCEINLINE const TArray<FInstancedStruct>& GetChunkFragments() const
+	{
+		return ChunkFragments;
+	}
+
+	FORCEINLINE const TArray<FConstSharedStruct>& GetConstSharedFragments() const
+	{
+		return ConstSharedFragments;
+	}
+
+	FORCEINLINE const TArray<FSharedStruct>& GetSharedFragments() const
+	{
+		return SharedFragments;
+	}
+
+	FORCEINLINE void SetFragments(const TArray<FInstancedStruct>& InFragments)
+	{
+		DirtyHashCache();
+		Fragments = InFragments;
+	}
+
+	FORCEINLINE void SetChunkFragments(const TArray<FInstancedStruct>& InChunkFragments)
+	{
+		DirtyHashCache();
+		ChunkFragments = InChunkFragments;
+	}
+
+	FORCEINLINE void SetConstSharedFragments(const TArray<FConstSharedStruct>& InConstSharedFragments)
+	{
+		DirtyHashCache();
+		ConstSharedFragments = InConstSharedFragments;
+	}
+
+	FORCEINLINE void SetSharedFragments(const TArray<FSharedStruct>& InSharedFragments)
+	{
+		DirtyHashCache();
+		SharedFragments = InSharedFragments;
+	}
+
+	FORCEINLINE void DirtyHashCache()
+	{
+		HashCache = UINT32_MAX;
+	}
+
+	FORCEINLINE void CacheHash() const
+	{
+		if (HashCache == UINT32_MAX)
+		{
+			HashCache = CalculateHash();
+		}
+	}
+
+	friend FORCEINLINE uint32 GetTypeHash(const FMassArchetypeFragmentsInitialValues& InitialValues)
+	{
+		InitialValues.CacheHash();
+		return InitialValues.HashCache;
+	}
+
+	uint32 CalculateHash() const;
+
+	SIZE_T GetAllocatedSize() const
+	{
+		return Fragments.GetAllocatedSize() + ChunkFragments.GetAllocatedSize() + ConstSharedFragments.GetAllocatedSize() + SharedFragments.GetAllocatedSize();
+	}
+
+protected:
+	mutable uint32 HashCache = UINT32_MAX;
+	// Atodo remove the fragments and ChunkFragments as they are not used and it creating inconsistency in the uniqueness of the Archetypes.
 	TArray<FInstancedStruct> Fragments;
 	TArray<FInstancedStruct> ChunkFragments;
-
-	// @todo add the shared fragments values here
+	TArray<FConstSharedStruct> ConstSharedFragments;
+	TArray<FSharedStruct> SharedFragments;
 };
\ No newline at end of file
diff --git a/Engine/Plugins/Runtime/MassGameplay/Source/MassSpawner/Private/MassSpawnerSubsystem.cpp b/Engine/Plugins/Runtime/MassGameplay/Source/MassSpawner/Private/MassSpawnerSubsystem.cpp
--- a/Engine/Plugins/Runtime/MassGameplay/Source/MassSpawner/Private/MassSpawnerSubsystem.cpp
+++ b/Engine/Plugins/Runtime/MassGameplay/Source/MassSpawner/Private/MassSpawnerSubsystem.cpp
@@ -136,7 +136,7 @@ bool UMassSpawnerSubsystem::SpawnEntities(const FMassEntityTemplate& EntityTempl
 	{ 
 		UE_CVLOG(OutEntities.Num() != NumberToSpawn, this, LogMassSpawner, Warning, TEXT("Tried to batch-create %d entities but created %d."), NumberToSpawn, OutEntities.Num());
 
-		const TArrayView<const FInstancedStruct> FragmentInstances = EntityTemplate.GetArchetypeFragmentsInitialValues().Fragments;
+		const TArrayView<const FInstancedStruct> FragmentInstances = EntityTemplate.GetArchetypeFragmentsInitialValues().GetFragments();
 		const FArchetypeChunkCollection ChunkCollection(EntityTemplate.GetArchetype(), OutEntities);
 		EntitySystem->BatchSetEntityFragmentsValues(ChunkCollection, FragmentInstances);
 
@@ -301,7 +301,7 @@ void UMassSpawnerSubsystem::DoSpawning(const FMassEntityTemplate& EntityTemplate
 	TArray<FMassEntityHandle> SpawnedEntities;
 	EntitySystem->BatchCreateEntities(EntityTemplate.GetArchetype(), Data.MaxNumber, SpawnedEntities);
 
-	const TArrayView<const FInstancedStruct> FragmentInstances = EntityTemplate.GetArchetypeFragmentsInitialValues().Fragments;
+	const TArrayView<const FInstancedStruct> FragmentInstances = EntityTemplate.GetArchetypeFragmentsInitialValues().GetFragments();
 	const FArchetypeChunkCollection ChunkCollection(EntityTemplate.GetArchetype(), SpawnedEntities);
 	EntitySystem->BatchSetEntityFragmentsValues(ChunkCollection, FragmentInstances);
 	
diff --git a/Engine/Plugins/Runtime/MassGameplay/Source/MassSpawner/Public/MassEntityTemplate.h b/Engine/Plugins/Runtime/MassGameplay/Source/MassSpawner/Public/MassEntityTemplate.h
--- a/Engine/Plugins/Runtime/MassGameplay/Source/MassSpawner/Public/MassEntityTemplate.h
+++ b/Engine/Plugins/Runtime/MassGameplay/Source/MassSpawner/Public/MassEntityTemplate.h
@@ -128,11 +128,11 @@ struct MASSSPAWNER_API FMassEntityTemplate
 		if (!Composition.Fragments.Contains(*FragmentType))
 		{
 			Composition.Fragments.Add(*FragmentType);
-			InitialValues.Fragments.Add(Fragment);
+			InitialValues.AddFragment(Fragment);
 		}
-		else if (!InitialValues.Fragments.ContainsByPredicate(FSameTypeScriptStructPredicate(FragmentType)))
+		else if (!InitialValues.GetFragments().ContainsByPredicate(FSameTypeScriptStructPredicate(FragmentType)))
 		{
-			InitialValues.Fragments.Add(Fragment);
+			InitialValues.AddFragment(Fragment);
 		}
 	}
 
@@ -144,14 +144,13 @@ struct MASSSPAWNER_API FMassEntityTemplate
 		{
 			Composition.Fragments.Add<T>();
 		}
-		else if (FInstancedStruct* Fragment = InitialValues.Fragments.FindByPredicate(FSameTypeScriptStructPredicate(T::StaticStruct())))
+		else if (const FInstancedStruct* Fragment = InitialValues.GetFragments().FindByPredicate(FSameTypeScriptStructPredicate(T::StaticStruct())))
 		{
 			return Fragment->template GetMutable<T>();
 		}
 
 		// Add a default initial fragment value
-		const int32 Index = InitialValues.Fragments.Add(FInstancedStruct(T::StaticStruct()));
-		return InitialValues.Fragments[Index].template GetMutable<T>();
+		return InitialValues.AddFragment(T());
 	}
 
 	template<typename T>
@@ -183,14 +182,55 @@ struct MASSSPAWNER_API FMassEntityTemplate
 		{
 			Composition.ChunkFragments.Add<T>();
 		}
-		else if (FInstancedStruct* ChunkFragment = InitialValues.ChunkFragments.FindByPredicate(FSameTypeScriptStructPredicate(T::StaticStruct())))
+		else if (const FInstancedStruct* ChunkFragment = InitialValues.GetChunkFragments().FindByPredicate(FSameTypeScriptStructPredicate(T::StaticStruct())))
 		{
 			return ChunkFragment->template GetMutable<T>();
 		}
 
 		// Add a default initial chunk fragment value
-		const int32 Index = InitialValues.ChunkFragments.Add(FInstancedStruct(T::StaticStruct()));
-		return InitialValues.ChunkFragments[Index].template GetMutable<T>();
+		return InitialValues.AddChunkFragment(T());
+	}
+
+	void AddConstSharedFragment(const FConstSharedStruct& SharedFragment)
+	{
+		const UScriptStruct* FragmentType = SharedFragment.GetScriptStruct();
+		if(ensureMsgf(FragmentType && FragmentType->IsChildOf(FMassSharedFragment::StaticStruct()), TEXT("Given struct doesn't represent a valid shared fragment type. Make sure to inherit from FMassSharedFragment or one of its child-types.")))
+		{
+			if (!Composition.SharedFragments.Contains(*FragmentType))
+			{
+				Composition.SharedFragments.Add(*FragmentType);
+				InitialValues.AddConstSharedFragment(SharedFragment);
+			}
+#if DO_ENSURE
+			else
+			{
+				const FConstSharedStruct* Struct = InitialValues.GetConstSharedFragments().FindByPredicate(FSameTypeScriptStructPredicate(SharedFragment));
+				ensureMsgf(Struct && *Struct == SharedFragment, TEXT("Adding 2 different const shared fragment of the same type is not allowed"));
+
+			}
+#endif // DO_ENSURE
+		}
+	}
+
+	void AddSharedFragment(const FSharedStruct& SharedFragment)
+	{
+		const UScriptStruct* FragmentType = SharedFragment.GetScriptStruct();
+		if(ensureMsgf(FragmentType && FragmentType->IsChildOf(FMassSharedFragment::StaticStruct()), TEXT("Given struct doesn't represent a valid shared fragment type. Make sure to inherit from FMassSharedFragment or one of its child-types.")))
+		{
+			if (!Composition.SharedFragments.Contains(*FragmentType))
+			{
+				Composition.SharedFragments.Add(*FragmentType);
+				InitialValues.AddSharedFragment(SharedFragment);
+			}
+	#if DO_ENSURE
+			else
+			{
+				const FSharedStruct* Struct = InitialValues.GetSharedFragments().FindByPredicate(FSameTypeScriptStructPredicate(SharedFragment));
+				ensureMsgf(Struct && *Struct == SharedFragment, TEXT("Adding 2 different shared fragment of the same type is not allowed"));
+
+			}
+	#endif // DO_ENSURE
+		}
 	}
 
 	FString DebugGetDescription(UMassEntitySubsystem* EntitySubsystem = nullptr) const;
@@ -219,6 +259,12 @@ struct MASSSPAWNER_API FMassEntityTemplate
 		return Composition.ChunkFragments.Contains<T>();
 	}
 
+	template<typename T>
+	bool HasSharedFragment() const
+	{
+		return Composition.SharedFragments.Contains<T>();
+	}
+
 private:
 	FArchetypeHandle Archetype;
 
diff --git a/Engine/Plugins/Runtime/MassGameplay/Source/MassSpawner/Public/MassEntityTemplateRegistry.h b/Engine/Plugins/Runtime/MassGameplay/Source/MassSpawner/Public/MassEntityTemplateRegistry.h
--- a/Engine/Plugins/Runtime/MassGameplay/Source/MassSpawner/Public/MassEntityTemplateRegistry.h
+++ b/Engine/Plugins/Runtime/MassGameplay/Source/MassSpawner/Public/MassEntityTemplateRegistry.h
@@ -113,6 +113,16 @@ struct FMassEntityTemplateBuildContext
 		return Template.AddChunkFragment_GetRef<T>();
 	}
 
+	void AddConstSharedFragment(const FConstSharedStruct& InSharedFragment)
+	{
+		Template.AddConstSharedFragment(InSharedFragment);
+	}
+
+	void AddSharedFragment(const FSharedStruct& InSharedFragment)
+	{
+		Template.AddSharedFragment(InSharedFragment);
+	}
+
 	template<typename T>
 	bool HasFragment() const
 	{
@@ -136,6 +146,12 @@ struct FMassEntityTemplateBuildContext
 		return Template.HasChunkFragment<T>();
 	}
 
+	template<typename T>
+	bool HasSharedFragment() const
+	{
+		return Template.HasSharedFragment<T>();
+	}
+
 	//----------------------------------------------------------------------//
 	// Initializers
 	//----------------------------------------------------------------------//
