Moved the Struct operators into the StructUtilsTypes.h and renamed them to more generic consistent name.

#rnx

diff --git a/Engine/Plugins/Experimental/StructUtils/Source/StructUtils/Public/InstancedStruct.h b/Engine/Plugins/Experimental/StructUtils/Source/StructUtils/Public/InstancedStruct.h
--- a/Engine/Plugins/Experimental/StructUtils/Source/StructUtils/Public/InstancedStruct.h
+++ b/Engine/Plugins/Experimental/StructUtils/Source/StructUtils/Public/InstancedStruct.h
@@ -811,16 +811,4 @@ public:
 		static_assert(!TIsDerivedFrom<T, FConstBaseStruct>::IsDerived && !TIsDerivedFrom<T, FConstSharedStruct>::IsDerived, "It does not make sense to create a struct view over an other struct wrapper type");
 		return FStructView(T::StaticStruct(), reinterpret_cast<uint8*>(&InStruct));
 	}
-};
-
-///////////////////////////////////////////////////////////////// FSameTypeScriptStructPredicate /////////////////////////////////////////////////////////////////
-
-/* Predicate useful to find a struct of a specific type in an container */
-struct FSameTypeScriptStructPredicate
-{
-	const UScriptStruct* TypePtr;
-	FSameTypeScriptStructPredicate(const UScriptStruct* InTypePtr) : TypePtr(InTypePtr) {}
-	FSameTypeScriptStructPredicate(const FConstStructView& InRef) : TypePtr(InRef.GetScriptStruct()) {}
-
-	bool operator()(const FConstStructView& Other) const { return Other.GetScriptStruct() == TypePtr; }
-};
+};
\ No newline at end of file
diff --git a/Engine/Plugins/Experimental/StructUtils/Source/StructUtils/Public/StructUtilsTypes.h b/Engine/Plugins/Experimental/StructUtils/Source/StructUtils/Public/StructUtilsTypes.h
--- a/Engine/Plugins/Experimental/StructUtils/Source/StructUtils/Public/StructUtilsTypes.h
+++ b/Engine/Plugins/Experimental/StructUtils/Source/StructUtils/Public/StructUtilsTypes.h
@@ -3,6 +3,7 @@
 #pragma once
 
 #include "UObject/Class.h"
+#include "InstancedStruct.h"
 
 #ifndef WITH_STRUCTUTILS_DEBUG
 #define WITH_STRUCTUTILS_DEBUG (!(UE_BUILD_SHIPPING || UE_BUILD_SHIPPING_WITH_EDITOR || UE_BUILD_TEST) && 1)
@@ -19,3 +20,41 @@ namespace UE::StructUtils
 
 	extern STRUCTUTILS_API uint32 GetStructCrc32(const FConstStructView& StructView, const uint32 CRC = 0);
 }
+
+/* Predicate useful to find a struct of a specific type in an container */
+struct FStructTypeEqualOperator
+{
+	const UScriptStruct* TypePtr;
+	FStructTypeEqualOperator(const UScriptStruct* InTypePtr) : TypePtr(InTypePtr) {}
+	FStructTypeEqualOperator(const FConstStructView& InRef) : TypePtr(InRef.GetScriptStruct()) {}
+
+	bool operator()(const FConstStructView& Other) const { return Other.GetScriptStruct() == TypePtr; }
+};
+
+struct FScriptStructSortOperator
+{
+	template<typename T>
+	bool operator()(const T& A, const T& B) const
+	{
+		return (A.GetStructureSize() > B.GetStructureSize())
+			|| (A.GetStructureSize() == B.GetStructureSize() && B.GetFName().FastLess(A.GetFName()));
+	}
+};
+
+struct FStructTypeSortOperator
+{
+	bool operator()(const FConstStructView& A, const FConstStructView& B) const
+	{
+		const UScriptStruct* AScriptStruct = A.GetScriptStruct();
+		const UScriptStruct* BScriptStruct = B.GetScriptStruct();
+		if (!AScriptStruct)
+		{
+			return true;
+		}
+		else if (!BScriptStruct)
+		{
+			return false;
+		}
+		return FScriptStructSortOperator()(*AScriptStruct, *BScriptStruct);
+	}
+};
\ No newline at end of file
diff --git a/Engine/Plugins/Runtime/MassEntity/Source/MassEntity/Private/MassArchetypeData.cpp b/Engine/Plugins/Runtime/MassEntity/Source/MassEntity/Private/MassArchetypeData.cpp
--- a/Engine/Plugins/Runtime/MassEntity/Source/MassEntity/Private/MassArchetypeData.cpp
+++ b/Engine/Plugins/Runtime/MassEntity/Source/MassEntity/Private/MassArchetypeData.cpp
@@ -30,7 +30,7 @@ void FMassArchetypeData::Initialize(const FMassArchetypeCompositionDescriptor& I
 	TArray<const UScriptStruct*, TInlineAllocator<16>> SortedFragmentList;
 	InCompositionDescriptor.Fragments.ExportTypes(SortedFragmentList);
 
-	SortedFragmentList.Sort(FMassSorterOperator<UScriptStruct>());
+	SortedFragmentList.Sort(FScriptStructSortOperator());
 
 	// Figure out how many bytes all of the individual fragments (and metadata) will cost per entity
 	int32 FragmentSizeTallyBytes = 0;
@@ -68,11 +68,11 @@ void FMassArchetypeData::Initialize(const FMassArchetypeCompositionDescriptor& I
 	CompositionDescriptor.ChunkFragments = InCompositionDescriptor.ChunkFragments;
 	TArray<const UScriptStruct*, TInlineAllocator<16>> ChunkFragmentList;
 	CompositionDescriptor.ChunkFragments.ExportTypes(ChunkFragmentList);
-	ChunkFragmentList.Sort(FMassSorterOperator<UScriptStruct>());
+	ChunkFragmentList.Sort(FScriptStructSortOperator());
 	for (const UScriptStruct* ChunkFragmentType : ChunkFragmentList)
 	{
 		check(ChunkFragmentType);
-		if (const FInstancedStruct* Struct = InInitialValues.GetChunkFragments().FindByPredicate(FSameTypeScriptStructPredicate(ChunkFragmentType)))
+		if (const FInstancedStruct* Struct = InInitialValues.GetChunkFragments().FindByPredicate(FStructTypeEqualOperator(ChunkFragmentType)))
 		{
 			InitialValues.AddChunkFragment(FConstStructView(*Struct));
 		}
@@ -697,7 +697,7 @@ void FMassArchetypeData::GetRequirementsConstSharedFragmentMapping(TConstArrayVi
 	{
 		if (Requirement.RequiresBinding())
 		{
-			const int32 FragmentIndex = InitialValues.GetConstSharedFragments().IndexOfByPredicate(FSameTypeScriptStructPredicate(Requirement.StructType));
+			const int32 FragmentIndex = InitialValues.GetConstSharedFragments().IndexOfByPredicate(FStructTypeEqualOperator(Requirement.StructType));
 			check(FragmentIndex != INDEX_NONE || Requirement.IsOptional());
 			OutFragmentIndices.Add(FragmentIndex);
 		}
@@ -711,7 +711,7 @@ void FMassArchetypeData::GetRequirementsSharedFragmentMapping(TConstArrayView<FM
 	{
 		if (Requirement.RequiresBinding())
 		{
-			const int32 FragmentIndex = InitialValues.GetSharedFragments().IndexOfByPredicate(FSameTypeScriptStructPredicate(Requirement.StructType));
+			const int32 FragmentIndex = InitialValues.GetSharedFragments().IndexOfByPredicate(FStructTypeEqualOperator(Requirement.StructType));
 			check(FragmentIndex != INDEX_NONE || Requirement.IsOptional());
 			OutFragmentIndices.Add(FragmentIndex);
 		}
@@ -811,7 +811,7 @@ void FMassArchetypeData::BindConstSharedFragmentRequirements(FMassExecutionConte
 	{
 		for (FMassExecutionContext::FConstSharedFragmentView& Requirement : RunContext.GetMutableConstSharedRequirements())
 		{
-			const FConstSharedStruct* SharedFragment = InitialValues.GetConstSharedFragments().FindByPredicate( FSameTypeScriptStructPredicate(Requirement.Requirement.StructType) );
+			const FConstSharedStruct* SharedFragment = InitialValues.GetConstSharedFragments().FindByPredicate(FStructTypeEqualOperator(Requirement.Requirement.StructType) );
 			check(SharedFragment != nullptr || Requirement.Requirement.IsOptional());
 			Requirement.FragmentView = SharedFragment ? *SharedFragment : FConstSharedStruct();
 		}
@@ -837,7 +837,7 @@ void FMassArchetypeData::BindSharedFragmentRequirements(FMassExecutionContext& R
 	{
 		for (FMassExecutionContext::FSharedFragmentView& Requirement : RunContext.GetMutableSharedRequirements())
 		{
-			const FSharedStruct* SharedFragment = InitialValues.GetSharedFragments().FindByPredicate(FSameTypeScriptStructPredicate(Requirement.Requirement.StructType));
+			const FSharedStruct* SharedFragment = InitialValues.GetSharedFragments().FindByPredicate(FStructTypeEqualOperator(Requirement.Requirement.StructType));
 			check(SharedFragment != nullptr || Requirement.Requirement.IsOptional());
 			Requirement.FragmentView = SharedFragment ? *SharedFragment : FSharedStruct();
 		}
diff --git a/Engine/Plugins/Runtime/MassEntity/Source/MassEntity/Private/MassEntityQuery.cpp b/Engine/Plugins/Runtime/MassEntity/Source/MassEntity/Private/MassEntityQuery.cpp
--- a/Engine/Plugins/Runtime/MassEntity/Source/MassEntity/Private/MassEntityQuery.cpp
+++ b/Engine/Plugins/Runtime/MassEntity/Source/MassEntity/Private/MassEntityQuery.cpp
@@ -51,10 +51,10 @@ void FMassEntityQuery::SortRequirements()
 	// so that when we access ArchetypeData.FragmentConfigs in FMassArchetypeData::BindRequirementsWithMapping
 	// (via GetFragmentData call) the access is sequential (i.e. not random) and there's a higher chance the memory
 	// FragmentConfigs we want to access have already been fetched and are available in processor cache.
-	Requirements.Sort(FMassSorterOperator<FMassFragmentRequirement>());
-	ChunkRequirements.Sort(FMassSorterOperator<FMassFragmentRequirement>());
-	ConstSharedRequirements.Sort(FMassSorterOperator<FMassFragmentRequirement>());
-	SharedRequirements.Sort(FMassSorterOperator<FMassFragmentRequirement>());
+	Requirements.Sort(FScriptStructSortOperator());
+	ChunkRequirements.Sort(FScriptStructSortOperator());
+	ConstSharedRequirements.Sort(FScriptStructSortOperator());
+	SharedRequirements.Sort(FScriptStructSortOperator());
 }
 
 void FMassEntityQuery::CacheArchetypes(UMassEntitySubsystem& InEntitySubsystem)
diff --git a/Engine/Plugins/Runtime/MassEntity/Source/MassEntity/Public/MassArchetypeTypes.h b/Engine/Plugins/Runtime/MassEntity/Source/MassEntity/Public/MassArchetypeTypes.h
--- a/Engine/Plugins/Runtime/MassEntity/Source/MassEntity/Public/MassArchetypeTypes.h
+++ b/Engine/Plugins/Runtime/MassEntity/Source/MassEntity/Public/MassArchetypeTypes.h
@@ -148,13 +148,3 @@ struct FMassQueryRequirementIndicesMapping
 		return EntityFragments.Num() == 0 || ChunkFragments.Num() == 0;
 	}
 };
-
-template<typename T>
-struct FMassSorterOperator
-{
-	bool operator()(const T& A, const T& B) const
-	{
-		return (A.GetStructureSize() > B.GetStructureSize())
-			|| (A.GetStructureSize() == B.GetStructureSize() && B.GetFName().FastLess(A.GetFName()));
-	}
-};
diff --git a/Engine/Plugins/Runtime/MassEntity/Source/MassEntity/Public/MassEntityQuery.h b/Engine/Plugins/Runtime/MassEntity/Source/MassEntity/Public/MassEntityQuery.h
--- a/Engine/Plugins/Runtime/MassEntity/Source/MassEntity/Public/MassEntityQuery.h
+++ b/Engine/Plugins/Runtime/MassEntity/Source/MassEntity/Public/MassEntityQuery.h
@@ -71,7 +71,7 @@ public:
 
 	FString DebugGetDescription() const;
 
-	// these functions are used for sorting. See FMassSorterOperator
+	// these functions are used for sorting. See FScriptStructSortOperator
 	int32 GetStructureSize() const
 	{
 		return StructType->GetStructureSize();
diff --git a/Engine/Plugins/Runtime/MassEntity/Source/MassEntity/Public/MassEntitySubsystem.h b/Engine/Plugins/Runtime/MassEntity/Source/MassEntity/Public/MassEntitySubsystem.h
--- a/Engine/Plugins/Runtime/MassEntity/Source/MassEntity/Public/MassEntitySubsystem.h
+++ b/Engine/Plugins/Runtime/MassEntity/Source/MassEntity/Public/MassEntitySubsystem.h
@@ -7,6 +7,7 @@
 #include "MassProcessingTypes.h"
 #include "InstancedStruct.h"
 #include "MassEntityQuery.h"
+#include "StructUtilsTypes.h"
 #include "MassEntitySubsystem.generated.h"
 
 
@@ -286,7 +287,7 @@ public:
 	template<typename T>
 	void ForEachSharedFragment(TFunction< void(T& /*SharedFragment*/) > ExecuteFunction)
 	{
-		FSameTypeScriptStructPredicate Predicate(T::StaticStruct());
+		FStructTypeEqualOperator Predicate(T::StaticStruct());
 		for (const FSharedStruct& Struct : SharedFragments)
 		{
 			if (Predicate(Struct))
diff --git a/Engine/Plugins/Runtime/MassGameplay/Source/MassSpawner/Public/MassEntityTemplate.h b/Engine/Plugins/Runtime/MassGameplay/Source/MassSpawner/Public/MassEntityTemplate.h
--- a/Engine/Plugins/Runtime/MassGameplay/Source/MassSpawner/Public/MassEntityTemplate.h
+++ b/Engine/Plugins/Runtime/MassGameplay/Source/MassSpawner/Public/MassEntityTemplate.h
@@ -130,7 +130,7 @@ struct MASSSPAWNER_API FMassEntityTemplate
 			Composition.Fragments.Add(*FragmentType);
 			InitialValues.AddFragment(Fragment);
 		}
-		else if (!InitialValues.GetFragments().ContainsByPredicate(FSameTypeScriptStructPredicate(FragmentType)))
+		else if (!InitialValues.GetFragments().ContainsByPredicate(FStructTypeEqualOperator(FragmentType)))
 		{
 			InitialValues.AddFragment(Fragment);
 		}
@@ -144,7 +144,7 @@ struct MASSSPAWNER_API FMassEntityTemplate
 		{
 			Composition.Fragments.Add<T>();
 		}
-		else if (const FInstancedStruct* Fragment = InitialValues.GetFragments().FindByPredicate(FSameTypeScriptStructPredicate(T::StaticStruct())))
+		else if (const FInstancedStruct* Fragment = InitialValues.GetFragments().FindByPredicate(FStructTypeEqualOperator(T::StaticStruct())))
 		{
 			return Fragment->template GetMutable<T>();
 		}
@@ -182,7 +182,7 @@ struct MASSSPAWNER_API FMassEntityTemplate
 		{
 			Composition.ChunkFragments.Add<T>();
 		}
-		else if (const FInstancedStruct* ChunkFragment = InitialValues.GetChunkFragments().FindByPredicate(FSameTypeScriptStructPredicate(T::StaticStruct())))
+		else if (const FInstancedStruct* ChunkFragment = InitialValues.GetChunkFragments().FindByPredicate(FStructTypeEqualOperator(T::StaticStruct())))
 		{
 			return ChunkFragment->template GetMutable<T>();
 		}
@@ -204,7 +204,7 @@ struct MASSSPAWNER_API FMassEntityTemplate
 #if DO_ENSURE
 			else
 			{
-				const FConstSharedStruct* Struct = InitialValues.GetConstSharedFragments().FindByPredicate(FSameTypeScriptStructPredicate(SharedFragment));
+				const FConstSharedStruct* Struct = InitialValues.GetConstSharedFragments().FindByPredicate(FStructTypeEqualOperator(SharedFragment));
 				ensureMsgf(Struct && *Struct == SharedFragment, TEXT("Adding 2 different const shared fragment of the same type is not allowed"));
 
 			}
@@ -225,7 +225,7 @@ struct MASSSPAWNER_API FMassEntityTemplate
 	#if DO_ENSURE
 			else
 			{
-				const FSharedStruct* Struct = InitialValues.GetSharedFragments().FindByPredicate(FSameTypeScriptStructPredicate(SharedFragment));
+				const FSharedStruct* Struct = InitialValues.GetSharedFragments().FindByPredicate(FStructTypeEqualOperator(SharedFragment));
 				ensureMsgf(Struct && *Struct == SharedFragment, TEXT("Adding 2 different shared fragment of the same type is not allowed"));
 
 			}
