Renamed the struct from FArchetypeFragmentInitialValues to FArchetypeFragmentValues
Removed ability to set chunk chunk default values as we do not want it part of the archetype uniqueness
Move the default fragment values only in the EntityTemplate and outside of the Archetype as it is spawner job to set those values

#rnx

diff --git a/Engine/Plugins/Runtime/MassEntity/Source/MassEntity/Private/MassArchetypeData.cpp b/Engine/Plugins/Runtime/MassEntity/Source/MassEntity/Private/MassArchetypeData.cpp
--- a/Engine/Plugins/Runtime/MassEntity/Source/MassEntity/Private/MassArchetypeData.cpp
+++ b/Engine/Plugins/Runtime/MassEntity/Source/MassEntity/Private/MassArchetypeData.cpp
@@ -25,7 +25,7 @@ bool FMassArchetypeData::HasFragmentType(const UScriptStruct* FragmentType) cons
 	return (FragmentType && CompositionDescriptor.Fragments.Contains(*FragmentType));
 }
 
-void FMassArchetypeData::Initialize(const FMassArchetypeCompositionDescriptor& InCompositionDescriptor, const FMassArchetypeFragmentsInitialValues& InInitialValues)
+void FMassArchetypeData::Initialize(const FMassArchetypeCompositionDescriptor& InCompositionDescriptor, const FMassArchetypeSharedFragmentValues& InSharedFragmentValues)
 {
 	TArray<const UScriptStruct*, TInlineAllocator<16>> SortedFragmentList;
 	InCompositionDescriptor.Fragments.ExportTypes(SortedFragmentList);
@@ -58,8 +58,6 @@ void FMassArchetypeData::Initialize(const FMassArchetypeCompositionDescriptor& I
 			CompositionDescriptor.Fragments.Add(*FragmentType);
 		}
 	}
-	// @todo, remove this from the Initial values and rename this structure differently
-	InitialValues.SetFragments(InInitialValues.GetFragments());
 
 	// Tags
 	CompositionDescriptor.Tags = InCompositionDescriptor.Tags;
@@ -72,20 +70,12 @@ void FMassArchetypeData::Initialize(const FMassArchetypeCompositionDescriptor& I
 	for (const UScriptStruct* ChunkFragmentType : ChunkFragmentList)
 	{
 		check(ChunkFragmentType);
-		if (const FInstancedStruct* Struct = InInitialValues.GetChunkFragments().FindByPredicate(FStructTypeEqualOperator(ChunkFragmentType)))
-		{
-			InitialValues.AddChunkFragment(FConstStructView(*Struct));
-		}
-		else
-		{
-			InitialValues.AddChunkFragment(FConstStructView(ChunkFragmentType));
-		}
+		ChunkFragmentsTemplate.Add(FConstStructView(ChunkFragmentType));
 	}
 
 	// Share fragments
 	CompositionDescriptor.SharedFragments = InCompositionDescriptor.SharedFragments;
-	InitialValues.SetConstSharedFragments(InInitialValues.GetConstSharedFragments());
-	InitialValues.SetSharedFragments(InInitialValues.GetSharedFragments());
+	SharedFragmentValues = InSharedFragmentValues;
 
 	TotalBytesPerEntity = FragmentSizeTallyBytes;
 	int32 ChunkAvailableSize = GetChunkAllocSize() - AlignmentPadding;
@@ -115,11 +105,8 @@ void FMassArchetypeData::InitializeWithSibling(const FMassArchetypeData& Sibling
 	CompositionDescriptor.Tags = OverrideTags;
 	CompositionDescriptor.ChunkFragments = SiblingArchetype.CompositionDescriptor.ChunkFragments;
 	CompositionDescriptor.SharedFragments = SiblingArchetype.CompositionDescriptor.SharedFragments;
-	// @todo maybe we should implement a copy operator in the FMassArchetypeInitialValues struct
-	InitialValues.SetFragments(SiblingArchetype.InitialValues.GetFragments());
-	InitialValues.SetChunkFragments(SiblingArchetype.InitialValues.GetChunkFragments());
-	InitialValues.SetConstSharedFragments(SiblingArchetype.InitialValues.GetConstSharedFragments());
-	InitialValues.SetSharedFragments(SiblingArchetype.InitialValues.GetSharedFragments());
+	ChunkFragmentsTemplate = SiblingArchetype.ChunkFragmentsTemplate;
+	SharedFragmentValues = SiblingArchetype.GetSharedFragmentValues();
 
 	TotalBytesPerEntity = SiblingArchetype.TotalBytesPerEntity;
 	NumEntitiesPerChunk = SiblingArchetype.NumEntitiesPerChunk;
@@ -175,12 +162,12 @@ int32 FMassArchetypeData::AddEntityInternal(FMassEntityHandle Entity, const bool
 		if (EmptyChunkIndex != INDEX_NONE)
 		{
 			DestinationChunk = &Chunks[EmptyChunkIndex];
-			DestinationChunk->Recycle(InitialValues.GetChunkFragments());
+			DestinationChunk->Recycle(ChunkFragmentsTemplate);
 			AbsoluteIndex = EmptyAbsoluteIndex;
 		}
 		else
 		{
-			DestinationChunk = &Chunks.Emplace_GetRef(GetChunkAllocSize(), InitialValues.GetChunkFragments());
+			DestinationChunk = &Chunks.Emplace_GetRef(GetChunkAllocSize(), ChunkFragmentsTemplate);
 		}
 
 		check(DestinationChunk);
@@ -674,9 +661,9 @@ void FMassArchetypeData::GetRequirementsChunkFragmentMapping(TConstArrayView<FMa
 		{
 			int32 FragmentIndex = INDEX_NONE;
 			// mz@todo Add comment here as this code seems to be assuming a certain order for chunk fragments, please explain
-			for (int32 i = LastFoundFragmentIndex + 1; i < InitialValues.GetChunkFragments().Num(); ++i)
+			for (int32 i = LastFoundFragmentIndex + 1; i < ChunkFragmentsTemplate.Num(); ++i)
 			{
-				if (InitialValues.GetChunkFragments()[i].GetScriptStruct()->IsChildOf(Requirement.StructType))
+				if (ChunkFragmentsTemplate[i].GetScriptStruct()->IsChildOf(Requirement.StructType))
 				{
 					FragmentIndex = i;
 					break;
@@ -697,7 +684,7 @@ void FMassArchetypeData::GetRequirementsConstSharedFragmentMapping(TConstArrayVi
 	{
 		if (Requirement.RequiresBinding())
 		{
-			const int32 FragmentIndex = InitialValues.GetConstSharedFragments().IndexOfByPredicate(FStructTypeEqualOperator(Requirement.StructType));
+			const int32 FragmentIndex = SharedFragmentValues.GetConstSharedFragments().IndexOfByPredicate(FStructTypeEqualOperator(Requirement.StructType));
 			check(FragmentIndex != INDEX_NONE || Requirement.IsOptional());
 			OutFragmentIndices.Add(FragmentIndex);
 		}
@@ -711,7 +698,7 @@ void FMassArchetypeData::GetRequirementsSharedFragmentMapping(TConstArrayView<FM
 	{
 		if (Requirement.RequiresBinding())
 		{
-			const int32 FragmentIndex = InitialValues.GetSharedFragments().IndexOfByPredicate(FStructTypeEqualOperator(Requirement.StructType));
+			const int32 FragmentIndex = SharedFragmentValues.GetSharedFragments().IndexOfByPredicate(FStructTypeEqualOperator(Requirement.StructType));
 			check(FragmentIndex != INDEX_NONE || Requirement.IsOptional());
 			OutFragmentIndices.Add(FragmentIndex);
 		}
@@ -804,14 +791,14 @@ void FMassArchetypeData::BindConstSharedFragmentRequirements(FMassExecutionConte
 			const int32 FragmentIndex = FragmentsMapping[i];
 
 			check(FragmentIndex != INDEX_NONE || Requirement.Requirement.IsOptional());
-			Requirement.FragmentView = FragmentIndex != INDEX_NONE ? InitialValues.GetConstSharedFragments()[FragmentIndex] : FConstSharedStruct();
+			Requirement.FragmentView = FragmentIndex != INDEX_NONE ? SharedFragmentValues.GetConstSharedFragments()[FragmentIndex] : FConstSharedStruct();
 		}
 	}
 	else
 	{
 		for (FMassExecutionContext::FConstSharedFragmentView& Requirement : RunContext.GetMutableConstSharedRequirements())
 		{
-			const FConstSharedStruct* SharedFragment = InitialValues.GetConstSharedFragments().FindByPredicate(FStructTypeEqualOperator(Requirement.Requirement.StructType) );
+			const FConstSharedStruct* SharedFragment = SharedFragmentValues.GetConstSharedFragments().FindByPredicate(FStructTypeEqualOperator(Requirement.Requirement.StructType) );
 			check(SharedFragment != nullptr || Requirement.Requirement.IsOptional());
 			Requirement.FragmentView = SharedFragment ? *SharedFragment : FConstSharedStruct();
 		}
@@ -830,14 +817,14 @@ void FMassArchetypeData::BindSharedFragmentRequirements(FMassExecutionContext& R
 			const int32 FragmentIndex = FragmentsMapping[i];
 
 			check(FragmentIndex != INDEX_NONE || Requirement.Requirement.IsOptional());
-			Requirement.FragmentView = FragmentIndex != INDEX_NONE ? InitialValues.GetSharedFragments()[FragmentIndex] : FSharedStruct();
+			Requirement.FragmentView = FragmentIndex != INDEX_NONE ? SharedFragmentValues.GetSharedFragments()[FragmentIndex] : FSharedStruct();
 		}
 	}
 	else
 	{
 		for (FMassExecutionContext::FSharedFragmentView& Requirement : RunContext.GetMutableSharedRequirements())
 		{
-			const FSharedStruct* SharedFragment = InitialValues.GetSharedFragments().FindByPredicate(FStructTypeEqualOperator(Requirement.Requirement.StructType));
+			const FSharedStruct* SharedFragment = SharedFragmentValues.GetSharedFragments().FindByPredicate(FStructTypeEqualOperator(Requirement.Requirement.StructType));
 			check(SharedFragment != nullptr || Requirement.Requirement.IsOptional());
 			Requirement.FragmentView = SharedFragment ? *SharedFragment : FSharedStruct();
 		}
@@ -856,7 +843,8 @@ SIZE_T FMassArchetypeData::GetAllocatedSize() const
 	}
 
 	return sizeof(FMassArchetypeData) +
-		InitialValues.GetAllocatedSize() +
+		SharedFragmentValues.GetAllocatedSize() +
+		ChunkFragmentsTemplate.GetAllocatedSize() +
 		FragmentConfigs.GetAllocatedSize() +
 		Chunks.GetAllocatedSize() +
 		(NumAllocatedChunkBuffers * GetChunkAllocSize()) +
diff --git a/Engine/Plugins/Runtime/MassEntity/Source/MassEntity/Private/MassArchetypeData.h b/Engine/Plugins/Runtime/MassEntity/Source/MassEntity/Private/MassArchetypeData.h
--- a/Engine/Plugins/Runtime/MassEntity/Source/MassEntity/Private/MassArchetypeData.h
+++ b/Engine/Plugins/Runtime/MassEntity/Source/MassEntity/Private/MassArchetypeData.h
@@ -137,7 +137,10 @@ struct FMassArchetypeData
 private:
 	// One-stop-shop variable describing the archetype's fragment and tag composition 
 	FMassArchetypeCompositionDescriptor CompositionDescriptor;
-	FMassArchetypeFragmentsInitialValues InitialValues;
+	FMassArchetypeSharedFragmentValues SharedFragmentValues;
+
+	// Pre-created default chunk fragment templates
+	TArray<FInstancedStruct> ChunkFragmentsTemplate;
 
 	TArray<FMassArchetypeFragmentConfig, TInlineAllocator<16>> FragmentConfigs;
 	
@@ -165,19 +168,19 @@ public:
 	const FMassSharedFragmentBitSet& GetSharedFragmentBitSet() const { return CompositionDescriptor.SharedFragments; }
 
 	const FMassArchetypeCompositionDescriptor& GetCompositionDescriptor() const { return CompositionDescriptor; }
-	const FMassArchetypeFragmentsInitialValues& GetInitialValues() const { return InitialValues; }
+	const FMassArchetypeSharedFragmentValues& GetSharedFragmentValues() const { return SharedFragmentValues; }
 
 	/** Method to iterate on all the fragment types */
 	void ForEachFragmentType(TFunction< void(const UScriptStruct* /*FragmentType*/)> Function) const;
 	bool HasFragmentType(const UScriptStruct* FragmentType) const;
 	bool HasTagType(const UScriptStruct* FragmentType) const { check(FragmentType); return CompositionDescriptor.Tags.Contains(*FragmentType); }
 
-	bool IsEquivalent(const FMassArchetypeCompositionDescriptor& OtherCompositionDescriptor, const FMassArchetypeFragmentsInitialValues& OtherInitialValues) const
+	bool IsEquivalent(const FMassArchetypeCompositionDescriptor& OtherCompositionDescriptor, const FMassArchetypeSharedFragmentValues& OtherSharedFragmentValues) const
 	{
-		return CompositionDescriptor.IsEquivalent(OtherCompositionDescriptor) && InitialValues.IsEquivalent(OtherInitialValues);
+		return CompositionDescriptor.IsEquivalent(OtherCompositionDescriptor) && SharedFragmentValues.IsEquivalent(OtherSharedFragmentValues);
 	}
 
-	void Initialize(const FMassArchetypeCompositionDescriptor& InCompositionDescriptor, const FMassArchetypeFragmentsInitialValues& InInitialValues);
+	void Initialize(const FMassArchetypeCompositionDescriptor& InCompositionDescriptor, const FMassArchetypeSharedFragmentValues& InSharedFragmentValues);
 
 	/** 
 	 * A special way of initializing an archetype resulting in a copy of SiblingArchetype's setup with OverrideTags
diff --git a/Engine/Plugins/Runtime/MassEntity/Source/MassEntity/Private/MassEntitySubsystem.cpp b/Engine/Plugins/Runtime/MassEntity/Source/MassEntity/Private/MassEntitySubsystem.cpp
--- a/Engine/Plugins/Runtime/MassEntity/Source/MassEntity/Private/MassEntitySubsystem.cpp
+++ b/Engine/Plugins/Runtime/MassEntity/Source/MassEntity/Private/MassEntitySubsystem.cpp
@@ -106,7 +106,7 @@ FArchetypeHandle UMassEntitySubsystem::CreateArchetype(TConstArrayView<const USc
 	}
 
 	const FMassArchetypeCompositionDescriptor Composition(FMassFragmentBitSet(FragmentList), Tags, ChunkFragments, FMassSharedFragmentBitSet());
-	return CreateArchetype(Composition, FMassArchetypeFragmentsInitialValues());
+	return CreateArchetype(Composition, FMassArchetypeSharedFragmentValues());
 }
 
 FArchetypeHandle UMassEntitySubsystem::CreateArchetype(const TSharedPtr<FMassArchetypeData>& SourceArchetype, const FMassFragmentBitSet& NewFragments)
@@ -115,10 +115,10 @@ FArchetypeHandle UMassEntitySubsystem::CreateArchetype(const TSharedPtr<FMassArc
 	checkf(NewFragments.IsEmpty() == false, TEXT("%s Adding an empty fragment list to an archetype is not supported."), ANSI_TO_TCHAR(__FUNCTION__));
 
 	const FMassArchetypeCompositionDescriptor Composition(NewFragments + SourceArchetype->GetFragmentBitSet(), SourceArchetype->GetTagBitSet(), SourceArchetype->GetChunkFragmentBitSet(), SourceArchetype->GetSharedFragmentBitSet());
-	return CreateArchetype(Composition, SourceArchetype->GetInitialValues());
+	return CreateArchetype(Composition, SourceArchetype->GetSharedFragmentValues());
 }
 
-FArchetypeHandle UMassEntitySubsystem::CreateArchetype(const FMassArchetypeCompositionDescriptor& Composition, const FMassArchetypeFragmentsInitialValues& InitialValues)
+FArchetypeHandle UMassEntitySubsystem::CreateArchetype(const FMassArchetypeCompositionDescriptor& Composition, const FMassArchetypeSharedFragmentValues& SharedFragmentValues)
 {
 	const uint32 TypeHash = Composition.CalculateHash();
 
@@ -127,7 +127,7 @@ FArchetypeHandle UMassEntitySubsystem::CreateArchetype(const FMassArchetypeCompo
 	FArchetypeHandle Result;
 	for (TSharedPtr<FMassArchetypeData>& Ptr : HashRow)
 	{
-		if (Ptr->IsEquivalent(Composition, InitialValues))
+		if (Ptr->IsEquivalent(Composition, SharedFragmentValues))
 		{
 			Result.DataPtr = Ptr;
 			break;
@@ -138,7 +138,7 @@ FArchetypeHandle UMassEntitySubsystem::CreateArchetype(const FMassArchetypeCompo
 	{
 		// Create a new archetype
 		TSharedPtr<FMassArchetypeData> NewArchetype = MakeShareable(new FMassArchetypeData);
-		NewArchetype->Initialize(Composition, InitialValues);
+		NewArchetype->Initialize(Composition, SharedFragmentValues);
 		HashRow.Add(NewArchetype);
 
 		for (const FMassArchetypeFragmentConfig& FragmentConfig : NewArchetype->GetFragmentConfigs())
@@ -167,7 +167,7 @@ FArchetypeHandle UMassEntitySubsystem::InternalCreateSiblingArchetype(const TSha
 	FArchetypeHandle Result;
 	for (TSharedPtr<FMassArchetypeData>& Ptr : HashRow)
 	{
-		if (Ptr->IsEquivalent(NewComposition, SourceArchetypeRef.GetInitialValues()))
+		if (Ptr->IsEquivalent(NewComposition, SourceArchetypeRef.GetSharedFragmentValues()))
 		{
 			Result.DataPtr = Ptr;
 			break;
@@ -249,7 +249,7 @@ FMassEntityHandle UMassEntitySubsystem::CreateEntity(TConstArrayView<FInstancedS
 {
 	check(FragmentInstanceList.Num() > 0);
 
-	FArchetypeHandle Archetype = CreateArchetype(FMassArchetypeCompositionDescriptor(FragmentInstanceList, FMassTagBitSet(), FMassChunkFragmentBitSet(), FMassSharedFragmentBitSet()), FMassArchetypeFragmentsInitialValues());
+	FArchetypeHandle Archetype = CreateArchetype(FMassArchetypeCompositionDescriptor(FragmentInstanceList, FMassTagBitSet(), FMassChunkFragmentBitSet(), FMassSharedFragmentBitSet()), FMassArchetypeSharedFragmentValues());
 	check(Archetype.DataPtr.IsValid());
 
 	FMassEntityHandle Entity = ReserveEntity();
@@ -293,7 +293,7 @@ void UMassEntitySubsystem::BuildEntity(FMassEntityHandle Entity, TConstArrayView
 	checkf(!IsEntityBuilt(Entity), TEXT("Expecting an entity that is not already built"));
 
 	const FMassArchetypeCompositionDescriptor Composition(FragmentInstanceList, FMassTagBitSet(), FMassChunkFragmentBitSet(), FMassSharedFragmentBitSet());
-	FArchetypeHandle Archetype = CreateArchetype(Composition, FMassArchetypeFragmentsInitialValues());
+	FArchetypeHandle Archetype = CreateArchetype(Composition, FMassArchetypeSharedFragmentValues());
 	check(Archetype.DataPtr.IsValid());
 
 	InternalBuildEntity(Entity, Archetype);
@@ -439,7 +439,7 @@ void UMassEntitySubsystem::AddCompositionToEntity_GetDelta(FMassEntityHandle Ent
 		NewDescriptor.Fragments += InDescriptor.Fragments;
 		NewDescriptor.Tags+= InDescriptor.Tags;
 
-		const FArchetypeHandle NewArchetypeHandle = CreateArchetype(NewDescriptor, OldArchetype->GetInitialValues());
+		const FArchetypeHandle NewArchetypeHandle = CreateArchetype(NewDescriptor, OldArchetype->GetSharedFragmentValues());
 
 		if (ensure(NewArchetypeHandle.DataPtr != EntityData.CurrentArchetype))
 		{
@@ -471,7 +471,7 @@ void UMassEntitySubsystem::RemoveCompositionFromEntity(FMassEntityHandle Entity,
 
 		if (NewDescriptor.IsEquivalent(InDescriptor) == false)
 	{
-			const FArchetypeHandle NewArchetypeHandle = CreateArchetype(NewDescriptor, OldArchetype->GetInitialValues());
+			const FArchetypeHandle NewArchetypeHandle = CreateArchetype(NewDescriptor, OldArchetype->GetSharedFragmentValues());
 
 		if (ensure(NewArchetypeHandle.DataPtr != EntityData.CurrentArchetype))
 		{
@@ -568,7 +568,7 @@ void UMassEntitySubsystem::RemoveFragmentListFromEntity(FMassEntityHandle Entity
 	{
 		// If all the fragments got removed this will result in fetching of the empty archetype
 		const FMassArchetypeCompositionDescriptor NewComposition(OldArchetype->GetFragmentBitSet() - FragmentsToRemove, OldArchetype->GetTagBitSet(), OldArchetype->GetChunkFragmentBitSet(), OldArchetype->GetSharedFragmentBitSet());
-		const FArchetypeHandle NewArchetypeHandle = CreateArchetype(NewComposition, OldArchetype->GetInitialValues());
+		const FArchetypeHandle NewArchetypeHandle = CreateArchetype(NewComposition, OldArchetype->GetSharedFragmentValues());
 
 		// Move the entity over
 		FMassArchetypeData* NewArchetype = NewArchetypeHandle.DataPtr.Get();
diff --git a/Engine/Plugins/Runtime/MassEntity/Source/MassEntity/Private/MassEntityTypes.cpp b/Engine/Plugins/Runtime/MassEntity/Source/MassEntity/Private/MassEntityTypes.cpp
--- a/Engine/Plugins/Runtime/MassEntity/Source/MassEntity/Private/MassEntityTypes.cpp
+++ b/Engine/Plugins/Runtime/MassEntity/Source/MassEntity/Private/MassEntityTypes.cpp
@@ -3,29 +3,21 @@
 #include "MassEntityTypes.h"
 #include "StructUtilsTypes.h"
 
-uint32 FMassArchetypeFragmentsInitialValues::CalculateHash() const
+uint32 FMassArchetypeSharedFragmentValues::CalculateHash() const
 {
-	TArray<uint32, TInlineAllocator<16>> FragmentHashes;
-	FragmentHashes.Reserve(Fragments.Num() + ChunkFragments.Num() + ConstSharedFragments.Num() + SharedFragments.Num());
+	checkf(bSorted, TEXT("Expecting the containers to be sorted for the hash caluclation to be correct"));
 
-	// max@todo: Fragments and chunk fragments are not part of the uniqueness and should be removed from the initial values, also should think of a better name for this struct.
+	// Fragments are not part of the uniqueness 
+	uint32 Hash = 0;
 	for (const FConstSharedStruct& Fragment : ConstSharedFragments)
 	{
-		FragmentHashes.Add(PointerHash(Fragment.GetMemory()));
+		PointerHash(Fragment.GetMemory(), Hash);
 	}
 
 	for (const FSharedStruct& Fragment : SharedFragments)
 	{
-		FragmentHashes.Add(PointerHash(Fragment.GetMemory()));
+		PointerHash(Fragment.GetMemory(), Hash);
 	}
 
-	FragmentHashes.Sort();
-
-	uint32 CalculatedHash = 0;
-	for (const uint32 FragmentHash : FragmentHashes)
-	{
-		CalculatedHash = HashCombine(CalculatedHash, FragmentHash);
-	}
-
-	return CalculatedHash;
+	return Hash;
 }
diff --git a/Engine/Plugins/Runtime/MassEntity/Source/MassEntity/Public/MassEntitySubsystem.h b/Engine/Plugins/Runtime/MassEntity/Source/MassEntity/Public/MassEntitySubsystem.h
--- a/Engine/Plugins/Runtime/MassEntity/Source/MassEntity/Public/MassEntitySubsystem.h
+++ b/Engine/Plugins/Runtime/MassEntity/Source/MassEntity/Public/MassEntitySubsystem.h
@@ -87,10 +87,10 @@ public:
 	 * CreateArchetype from a composition descriptor and initial values
 	 *
 	 * @param Composition of fragment, tag and chunk fragment types
-	 * @param InitialValues like default chunk fragment values
+	 * @param SharedFragmentValues are the actual pointer to shared fragments
 	 * @return a handle of a new archetype 
 	 */
-	FArchetypeHandle CreateArchetype(const FMassArchetypeCompositionDescriptor& Composition, const FMassArchetypeFragmentsInitialValues& InitialValues);
+	FArchetypeHandle CreateArchetype(const FMassArchetypeCompositionDescriptor& Composition, const FMassArchetypeSharedFragmentValues& SharedFragmentValues);
 
 	/** 
 	 *  Creates an archetype like SourceArchetype + NewFragmentList. 
diff --git a/Engine/Plugins/Runtime/MassEntity/Source/MassEntity/Public/MassEntityTypes.h b/Engine/Plugins/Runtime/MassEntity/Source/MassEntity/Public/MassEntityTypes.h
--- a/Engine/Plugins/Runtime/MassEntity/Source/MassEntity/Public/MassEntityTypes.h
+++ b/Engine/Plugins/Runtime/MassEntity/Source/MassEntity/Public/MassEntityTypes.h
@@ -200,39 +200,47 @@ struct FMassArchetypeCompositionDescriptor
 	FMassSharedFragmentBitSet SharedFragments;
 };
 
-struct FMassArchetypeFragmentsInitialValues
+struct FMassArchetypeSharedFragmentValues
 {
-	FMassArchetypeFragmentsInitialValues() = default;
+	FMassArchetypeSharedFragmentValues() = default;
 
-	FORCEINLINE bool IsEquivalent(const FMassArchetypeFragmentsInitialValues& OtherInitialValues) const
+	FMassArchetypeSharedFragmentValues(const FMassArchetypeSharedFragmentValues& OtherFragmentValues)
 	{
-		return GetTypeHash(*this) == GetTypeHash(OtherInitialValues);
+		ConstSharedFragments = OtherFragmentValues.ConstSharedFragments;
+		SharedFragments = OtherFragmentValues.SharedFragments;
+		HashCache = OtherFragmentValues.HashCache;
+		bSorted = OtherFragmentValues.bSorted;
 	}
 
-	template<typename T>
-	FORCEINLINE T& AddFragment(const T& Fragment)
+	FMassArchetypeSharedFragmentValues(FMassArchetypeSharedFragmentValues&& OtherFragmentValues)
 	{
-		DirtyHashCache();
-		return Fragments.Add_GetRef(FConstStructView::Make(Fragment)).template GetMutable<T>();
+		ConstSharedFragments = MoveTemp(OtherFragmentValues.ConstSharedFragments);
+		SharedFragments = MoveTemp(OtherFragmentValues.SharedFragments);
+		HashCache = OtherFragmentValues.HashCache;
+		bSorted = OtherFragmentValues.bSorted;
 	}
 
-	FORCEINLINE void AddFragment(FConstStructView Fragment)
+	FMassArchetypeSharedFragmentValues& operator=(const FMassArchetypeSharedFragmentValues& OtherFragmentValues)
 	{
-		DirtyHashCache();
-		Fragments.Add(Fragment);
+		ConstSharedFragments = OtherFragmentValues.ConstSharedFragments;
+		SharedFragments = OtherFragmentValues.SharedFragments;
+		HashCache = OtherFragmentValues.HashCache;
+		bSorted = OtherFragmentValues.bSorted;
+		return *this;
 	}
 
-	template<typename T>
-	FORCEINLINE T& AddChunkFragment(const T& Fragment)
+	FMassArchetypeSharedFragmentValues& operator=(FMassArchetypeSharedFragmentValues&& OtherFragmentValues)
 	{
-		DirtyHashCache();
-		return ChunkFragments.Add_GetRef(FConstStructView::Make(Fragment)).template GetMutable<T>();
+		ConstSharedFragments = MoveTemp(OtherFragmentValues.ConstSharedFragments);
+		SharedFragments = MoveTemp(OtherFragmentValues.SharedFragments);
+		HashCache = OtherFragmentValues.HashCache;
+		bSorted = OtherFragmentValues.bSorted;
+		return *this;
 	}
 
-	FORCEINLINE void AddChunkFragment(FConstStructView Fragment)
+	FORCEINLINE bool IsEquivalent(const FMassArchetypeSharedFragmentValues& OtherSharedFragmentValues) const
 	{
-		DirtyHashCache();
-		ChunkFragments.Add(Fragment);
+		return GetTypeHash(*this) == GetTypeHash(OtherSharedFragmentValues);
 	}
 
 	FORCEINLINE FConstSharedStruct& AddConstSharedFragment(const FConstSharedStruct& Fragment)
@@ -247,16 +255,6 @@ struct FMassArchetypeFragmentsInitialValues
 		return SharedFragments.Add_GetRef(Fragment);
 	}
 
-	FORCEINLINE const TArray<FInstancedStruct>& GetFragments() const
-	{
-		return Fragments;
-	}
-
-	FORCEINLINE const TArray<FInstancedStruct>& GetChunkFragments() const
-	{
-		return ChunkFragments;
-	}
-
 	FORCEINLINE const TArray<FConstSharedStruct>& GetConstSharedFragments() const
 	{
 		return ConstSharedFragments;
@@ -267,33 +265,10 @@ struct FMassArchetypeFragmentsInitialValues
 		return SharedFragments;
 	}
 
-	FORCEINLINE void SetFragments(const TArray<FInstancedStruct>& InFragments)
-	{
-		DirtyHashCache();
-		Fragments = InFragments;
-	}
-
-	FORCEINLINE void SetChunkFragments(const TArray<FInstancedStruct>& InChunkFragments)
-	{
-		DirtyHashCache();
-		ChunkFragments = InChunkFragments;
-	}
-
-	FORCEINLINE void SetConstSharedFragments(const TArray<FConstSharedStruct>& InConstSharedFragments)
-	{
-		DirtyHashCache();
-		ConstSharedFragments = InConstSharedFragments;
-	}
-
-	FORCEINLINE void SetSharedFragments(const TArray<FSharedStruct>& InSharedFragments)
-	{
-		DirtyHashCache();
-		SharedFragments = InSharedFragments;
-	}
-
 	FORCEINLINE void DirtyHashCache()
 	{
 		HashCache = UINT32_MAX;
+		bSorted = false;
 	}
 
 	FORCEINLINE void CacheHash() const
@@ -304,24 +279,33 @@ struct FMassArchetypeFragmentsInitialValues
 		}
 	}
 
-	friend FORCEINLINE uint32 GetTypeHash(const FMassArchetypeFragmentsInitialValues& InitialValues)
+	friend FORCEINLINE uint32 GetTypeHash(const FMassArchetypeSharedFragmentValues& SharedFragmentValues)
 	{
-		InitialValues.CacheHash();
-		return InitialValues.HashCache;
+		SharedFragmentValues.CacheHash();
+		return SharedFragmentValues.HashCache;
 	}
 
 	uint32 CalculateHash() const;
 
 	SIZE_T GetAllocatedSize() const
 	{
-		return Fragments.GetAllocatedSize() + ChunkFragments.GetAllocatedSize() + ConstSharedFragments.GetAllocatedSize() + SharedFragments.GetAllocatedSize();
+		return ConstSharedFragments.GetAllocatedSize() + SharedFragments.GetAllocatedSize();
+	}
+
+	void Sort()
+	{
+		if(!bSorted)
+		{
+			ConstSharedFragments.Sort(FStructTypeSortOperator());
+			SharedFragments.Sort(FStructTypeSortOperator());
+			bSorted = true;
+		}
 	}
 
 protected:
 	mutable uint32 HashCache = UINT32_MAX;
-	// Atodo remove the fragments and ChunkFragments as they are not used and it creating inconsistency in the uniqueness of the Archetypes.
-	TArray<FInstancedStruct> Fragments;
-	TArray<FInstancedStruct> ChunkFragments;
+	mutable bool bSorted = true; // When no element in the array, consider already sorted
+	
 	TArray<FConstSharedStruct> ConstSharedFragments;
 	TArray<FSharedStruct> SharedFragments;
 };
\ No newline at end of file
diff --git a/Engine/Plugins/Runtime/MassGameplay/Source/MassSpawner/Private/MassEntityTemplateRegistry.cpp b/Engine/Plugins/Runtime/MassGameplay/Source/MassSpawner/Private/MassEntityTemplateRegistry.cpp
--- a/Engine/Plugins/Runtime/MassGameplay/Source/MassSpawner/Private/MassEntityTemplateRegistry.cpp
+++ b/Engine/Plugins/Runtime/MassGameplay/Source/MassSpawner/Private/MassEntityTemplateRegistry.cpp
@@ -255,7 +255,10 @@ void UMassEntityTemplateRegistry::InitializeEntityTemplate(FMassEntityTemplate&
 	UMassEntitySubsystem* EntitySys = UWorld::GetSubsystem<UMassEntitySubsystem>(World);
 	check(EntitySys);
 
-	const FArchetypeHandle ArchetypeHandle = EntitySys->CreateArchetype(OutTemplate.GetCompositionDescriptor(), OutTemplate.GetArchetypeFragmentsInitialValues());
+	// Sort anything there is to sort for later comparaison purposes
+	OutTemplate.Sort();
+
+	const FArchetypeHandle ArchetypeHandle = EntitySys->CreateArchetype(OutTemplate.GetCompositionDescriptor(), OutTemplate.GetSharedFragmentValues());
 	OutTemplate.SetArchetype(ArchetypeHandle);
 }
 
diff --git a/Engine/Plugins/Runtime/MassGameplay/Source/MassSpawner/Private/MassSpawnerSubsystem.cpp b/Engine/Plugins/Runtime/MassGameplay/Source/MassSpawner/Private/MassSpawnerSubsystem.cpp
--- a/Engine/Plugins/Runtime/MassGameplay/Source/MassSpawner/Private/MassSpawnerSubsystem.cpp
+++ b/Engine/Plugins/Runtime/MassGameplay/Source/MassSpawner/Private/MassSpawnerSubsystem.cpp
@@ -136,7 +136,7 @@ bool UMassSpawnerSubsystem::SpawnEntities(const FMassEntityTemplate& EntityTempl
 	{ 
 		UE_CVLOG(OutEntities.Num() != NumberToSpawn, this, LogMassSpawner, Warning, TEXT("Tried to batch-create %d entities but created %d."), NumberToSpawn, OutEntities.Num());
 
-		const TArrayView<const FInstancedStruct> FragmentInstances = EntityTemplate.GetArchetypeFragmentsInitialValues().GetFragments();
+		TConstArrayView<FInstancedStruct> FragmentInstances = EntityTemplate.GetInitialFragmentValues();
 		const FArchetypeChunkCollection ChunkCollection(EntityTemplate.GetArchetype(), OutEntities);
 		EntitySystem->BatchSetEntityFragmentsValues(ChunkCollection, FragmentInstances);
 
@@ -301,7 +301,7 @@ void UMassSpawnerSubsystem::DoSpawning(const FMassEntityTemplate& EntityTemplate
 	TArray<FMassEntityHandle> SpawnedEntities;
 	EntitySystem->BatchCreateEntities(EntityTemplate.GetArchetype(), Data.MaxNumber, SpawnedEntities);
 
-	const TArrayView<const FInstancedStruct> FragmentInstances = EntityTemplate.GetArchetypeFragmentsInitialValues().GetFragments();
+	TConstArrayView<FInstancedStruct> FragmentInstances = EntityTemplate.GetInitialFragmentValues();
 	const FArchetypeChunkCollection ChunkCollection(EntityTemplate.GetArchetype(), SpawnedEntities);
 	EntitySystem->BatchSetEntityFragmentsValues(ChunkCollection, FragmentInstances);
 	
diff --git a/Engine/Plugins/Runtime/MassGameplay/Source/MassSpawner/Public/MassEntityTemplate.h b/Engine/Plugins/Runtime/MassGameplay/Source/MassSpawner/Public/MassEntityTemplate.h
--- a/Engine/Plugins/Runtime/MassGameplay/Source/MassSpawner/Public/MassEntityTemplate.h
+++ b/Engine/Plugins/Runtime/MassGameplay/Source/MassSpawner/Public/MassEntityTemplate.h
@@ -106,7 +106,8 @@ struct MASSSPAWNER_API FMassEntityTemplate
 	FMassEntityTemplateID GetTemplateID() const { return TemplateID; }
 
 	const FMassArchetypeCompositionDescriptor& GetCompositionDescriptor() const { return Composition; }
-	const FMassArchetypeFragmentsInitialValues& GetArchetypeFragmentsInitialValues() const { return InitialValues; }
+	const FMassArchetypeSharedFragmentValues& GetSharedFragmentValues() const { return SharedFragmentValues; }
+	TConstArrayView<FInstancedStruct> GetInitialFragmentValues() const { return InitialFragmentValues; }
 
 	template<typename T>
 	void AddFragment()
@@ -128,11 +129,11 @@ struct MASSSPAWNER_API FMassEntityTemplate
 		if (!Composition.Fragments.Contains(*FragmentType))
 		{
 			Composition.Fragments.Add(*FragmentType);
-			InitialValues.AddFragment(Fragment);
+			InitialFragmentValues.Add(Fragment);
 		}
-		else if (!InitialValues.GetFragments().ContainsByPredicate(FStructTypeEqualOperator(FragmentType)))
+		else if (!InitialFragmentValues.ContainsByPredicate(FStructTypeEqualOperator(FragmentType)))
 		{
-			InitialValues.AddFragment(Fragment);
+			InitialFragmentValues.Add(Fragment);
 		}
 	}
 
@@ -144,13 +145,13 @@ struct MASSSPAWNER_API FMassEntityTemplate
 		{
 			Composition.Fragments.Add<T>();
 		}
-		else if (const FInstancedStruct* Fragment = InitialValues.GetFragments().FindByPredicate(FStructTypeEqualOperator(T::StaticStruct())))
+		else if (const FInstancedStruct* Fragment = InitialFragmentValues.FindByPredicate(FStructTypeEqualOperator(T::StaticStruct())))
 		{
 			return Fragment->template GetMutable<T>();
 		}
 
 		// Add a default initial fragment value
-		return InitialValues.AddFragment(T());
+		return InitialFragmentValues.Emplace_GetRef(T::StaticStruct()).template GetMutable<T>();
 	}
 
 	template<typename T>
@@ -174,23 +175,6 @@ struct MASSSPAWNER_API FMassEntityTemplate
 		Composition.ChunkFragments.Add<T>();
 	}
 
-	template<typename T>
-	T& AddChunkFragment_GetRef()
-	{
-		static_assert(TIsDerivedFrom<T, FMassChunkFragment>::IsDerived, "Given struct doesn't represent a valid chunk fragment type. Make sure to inherit from FMassChunkFragment or one of its child-types.");
-		if( !Composition.ChunkFragments.Contains<T>() )
-		{
-			Composition.ChunkFragments.Add<T>();
-		}
-		else if (const FInstancedStruct* ChunkFragment = InitialValues.GetChunkFragments().FindByPredicate(FStructTypeEqualOperator(T::StaticStruct())))
-		{
-			return ChunkFragment->template GetMutable<T>();
-		}
-
-		// Add a default initial chunk fragment value
-		return InitialValues.AddChunkFragment(T());
-	}
-
 	void AddConstSharedFragment(const FConstSharedStruct& SharedFragment)
 	{
 		const UScriptStruct* FragmentType = SharedFragment.GetScriptStruct();
@@ -199,12 +183,12 @@ struct MASSSPAWNER_API FMassEntityTemplate
 			if (!Composition.SharedFragments.Contains(*FragmentType))
 			{
 				Composition.SharedFragments.Add(*FragmentType);
-				InitialValues.AddConstSharedFragment(SharedFragment);
+				SharedFragmentValues.AddConstSharedFragment(SharedFragment);
 			}
 #if DO_ENSURE
 			else
 			{
-				const FConstSharedStruct* Struct = InitialValues.GetConstSharedFragments().FindByPredicate(FStructTypeEqualOperator(SharedFragment));
+				const FConstSharedStruct* Struct = SharedFragmentValues.GetConstSharedFragments().FindByPredicate(FStructTypeEqualOperator(SharedFragment));
 				ensureMsgf(Struct && *Struct == SharedFragment, TEXT("Adding 2 different const shared fragment of the same type is not allowed"));
 
 			}
@@ -220,12 +204,12 @@ struct MASSSPAWNER_API FMassEntityTemplate
 			if (!Composition.SharedFragments.Contains(*FragmentType))
 			{
 				Composition.SharedFragments.Add(*FragmentType);
-				InitialValues.AddSharedFragment(SharedFragment);
+				SharedFragmentValues.AddSharedFragment(SharedFragment);
 			}
 	#if DO_ENSURE
 			else
 			{
-				const FSharedStruct* Struct = InitialValues.GetSharedFragments().FindByPredicate(FStructTypeEqualOperator(SharedFragment));
+				const FSharedStruct* Struct = SharedFragmentValues.GetSharedFragments().FindByPredicate(FStructTypeEqualOperator(SharedFragment));
 				ensureMsgf(Struct && *Struct == SharedFragment, TEXT("Adding 2 different shared fragment of the same type is not allowed"));
 
 			}
@@ -265,12 +249,20 @@ struct MASSSPAWNER_API FMassEntityTemplate
 		return Composition.SharedFragments.Contains<T>();
 	}
 
+	void Sort()
+	{
+		SharedFragmentValues.Sort();
+	}
+
 private:
 	FArchetypeHandle Archetype;
 
 	FMassArchetypeCompositionDescriptor Composition;
-	FMassArchetypeFragmentsInitialValues InitialValues;
+	FMassArchetypeSharedFragmentValues SharedFragmentValues;
 	
+	// Initial fragment values, this is not part of the archetype as it is the spawner job to set them.
+	TArray<FInstancedStruct> InitialFragmentValues;
+
 	UPROPERTY()
 	FMassRuntimePipeline InitializationPipeline;
 	
diff --git a/Engine/Plugins/Runtime/MassGameplay/Source/MassSpawner/Public/MassEntityTemplateRegistry.h b/Engine/Plugins/Runtime/MassGameplay/Source/MassSpawner/Public/MassEntityTemplateRegistry.h
--- a/Engine/Plugins/Runtime/MassGameplay/Source/MassSpawner/Public/MassEntityTemplateRegistry.h
+++ b/Engine/Plugins/Runtime/MassGameplay/Source/MassSpawner/Public/MassEntityTemplateRegistry.h
@@ -107,12 +107,6 @@ struct FMassEntityTemplateBuildContext
 		Template.AddChunkFragment<T>();
 	}
 
-	template<typename T>
-	T& AddChunkFragment_GetRef()
-	{
-		return Template.AddChunkFragment_GetRef<T>();
-	}
-
 	void AddConstSharedFragment(const FConstSharedStruct& InSharedFragment)
 	{
 		Template.AddConstSharedFragment(InSharedFragment);
