[GameplayDebuggerCategory_Mass] do not display empty archetypes

diff --git a/Engine/Plugins/AI/MassAI/Source/MassAIDebug/Private/GameplayDebuggerCategory_Mass.cpp b/Engine/Plugins/AI/MassAI/Source/MassAIDebug/Private/GameplayDebuggerCategory_Mass.cpp
--- a/Engine/Plugins/AI/MassAI/Source/MassAIDebug/Private/GameplayDebuggerCategory_Mass.cpp
+++ b/Engine/Plugins/AI/MassAI/Source/MassAIDebug/Private/GameplayDebuggerCategory_Mass.cpp
@@ -214,7 +214,7 @@ void FGameplayDebuggerCategory_Mass::CollectData(APlayerController* OwnerPC, AAc
 		{
 			FStringOutputDevice Ar;
 			Ar.SetAutoEmitLineTerminator(true);
-			EntitySystem->DebugPrintArchetypes(Ar);
+			EntitySystem->DebugPrintArchetypes(Ar, /*bIncludeEmpty*/false);
 
 			AddTextLine(Ar);
 		}
diff --git a/Engine/Plugins/Runtime/MassEntity/Source/MassEntity/Private/MassEntitySubsystem.cpp b/Engine/Plugins/Runtime/MassEntity/Source/MassEntity/Private/MassEntitySubsystem.cpp
--- a/Engine/Plugins/Runtime/MassEntity/Source/MassEntity/Private/MassEntitySubsystem.cpp
+++ b/Engine/Plugins/Runtime/MassEntity/Source/MassEntity/Private/MassEntitySubsystem.cpp
@@ -890,7 +890,7 @@ FMassExecutionContext UMassEntitySubsystem::CreateExecutionContext(const float D
 }
 
 #if WITH_MASSENTITY_DEBUG
-void UMassEntitySubsystem::DebugPrintArchetypes(FOutputDevice& Ar) const
+void UMassEntitySubsystem::DebugPrintArchetypes(FOutputDevice& Ar, const bool bIncludeEmpty) const
 {
 	Ar.Logf(ELogVerbosity::Log, TEXT("Listing archetypes contained in %s"), *GetPathNameSafe(this));
 
@@ -901,7 +901,10 @@ void UMassEntitySubsystem::DebugPrintArchetypes(FOutputDevice& Ar) const
 	{
 		for (const TSharedPtr<FMassArchetypeData>& ArchetypePtr : KVP.Value)
 		{
-			ArchetypePtr->DebugPrintArchetype(Ar);
+			if (ArchetypePtr.IsValid() && (bIncludeEmpty == true || ArchetypePtr->GetChunkCount() > 0))
+			{
+				ArchetypePtr->DebugPrintArchetype(Ar);
+			}
 		}
 
 		const int32 NumArchetypesInBucket = KVP.Value.Num();
diff --git a/Engine/Plugins/Runtime/MassEntity/Source/MassEntity/Public/MassEntitySubsystem.h b/Engine/Plugins/Runtime/MassEntity/Source/MassEntity/Public/MassEntitySubsystem.h
--- a/Engine/Plugins/Runtime/MassEntity/Source/MassEntity/Public/MassEntitySubsystem.h
+++ b/Engine/Plugins/Runtime/MassEntity/Source/MassEntity/Public/MassEntitySubsystem.h
@@ -300,7 +300,7 @@ public:
 #if WITH_MASSENTITY_DEBUG
 	void DebugPrintEntity(int32 Index, FOutputDevice& Ar, const TCHAR* InPrefix = TEXT("")) const;
 	void DebugPrintEntity(FMassEntityHandle Entity, FOutputDevice& Ar, const TCHAR* InPrefix = TEXT("")) const;
-	void DebugPrintArchetypes(FOutputDevice& Ar) const;
+	void DebugPrintArchetypes(FOutputDevice& Ar, const bool bIncludeEmpty = true) const;
 	static void DebugGetStringDesc(const FArchetypeHandle& Archetype, FOutputDevice& Ar);
 	void DebugGetArchetypesStringDetails(FOutputDevice& Ar, const bool bIncludeEmpty = true);
 	void DebugGetArchetypeFragmentTypes(const FArchetypeHandle& Archetype, TArray<const UScriptStruct*>& InOutFragmentList) const;
