Replaced the "component overlap" ensure in UEntitySubsystem::InternalAddComponentListToEntityChecked with an UE_CLOG

diff --git a/Engine/Plugins/Experimental/StructUtils/Source/StructUtils/Public/ScriptStructTypeBitSet.h b/Engine/Plugins/Experimental/StructUtils/Source/StructUtils/Public/ScriptStructTypeBitSet.h
--- a/Engine/Plugins/Experimental/StructUtils/Source/StructUtils/Public/ScriptStructTypeBitSet.h
+++ b/Engine/Plugins/Experimental/StructUtils/Source/StructUtils/Public/ScriptStructTypeBitSet.h
@@ -411,6 +411,17 @@ public:
 		}
 	}
 
+	FString DebugGetStringDesc() const
+	{
+#if WITH_STRUCTUTILS_DEBUG
+		FStringOutputDevice Ar;
+		DebugGetStringDesc(Ar);
+		return static_cast<FString>(Ar);
+#else
+		return TEXT("DEBUG INFO COMPILED OUT");
+#endif //WITH_STRUCTUTILS_DEBUG
+	}
+
 #if WITH_STRUCTUTILS_DEBUG
 	void DebugGetStringDesc(FOutputDevice& Ar) const
 	{
diff --git a/Engine/Plugins/Runtime/MassEntity/Source/MassEntity/Private/MassEntitySubsystem.cpp b/Engine/Plugins/Runtime/MassEntity/Source/MassEntity/Private/MassEntitySubsystem.cpp
--- a/Engine/Plugins/Runtime/MassEntity/Source/MassEntity/Private/MassEntitySubsystem.cpp
+++ b/Engine/Plugins/Runtime/MassEntity/Source/MassEntity/Private/MassEntitySubsystem.cpp
@@ -504,8 +504,10 @@ void UMassEntitySubsystem::InternalAddFragmentListToEntityChecked(FMassEntityHan
 	FEntityData& EntityData = Entities[Entity.Index];
 	FMassArchetypeData* OldArchetype = EntityData.CurrentArchetype.Get();
 	check(OldArchetype);
-	
-	ensureMsgf(OldArchetype->GetFragmentBitSet().HasAny(InFragments) == false, TEXT("Trying to add a new fragment type to an entity, but it already has some of them."));
+
+	UE_CLOG(OldArchetype->GetFragmentBitSet().HasAny(InComponents), LogMass, Log
+		, TEXT("Trying to add a new fragment type to an entity, but it already has some of them. (%s)")
+		, *InComponents.GetOverlap(OldArchetype->GetFragmentBitSet()).DebugGetStringDesc());
 
 	const FMassFragmentBitSet NewFragments = InFragments - OldArchetype->GetFragmentBitSet();
 	if (NewFragments.IsEmpty() == false)
