Trivial changes (headers removed, comments spelling)

diff --git a/Engine/Plugins/Runtime/MassEntity/Source/MassEntity/Public/MassProcessingPhase.h b/Engine/Plugins/Runtime/MassEntity/Source/MassEntity/Public/MassProcessingPhase.h
--- a/Engine/Plugins/Runtime/MassEntity/Source/MassEntity/Public/MassProcessingPhase.h
+++ b/Engine/Plugins/Runtime/MassEntity/Source/MassEntity/Public/MassProcessingPhase.h
@@ -71,7 +71,7 @@ struct TStructOpsTypeTraits<FMassProcessingPhase> : public TStructOpsTypeTraitsB
 
 /** MassProcessingPhaseManager owns separate FMassProcessingPhase instances for every ETickingGroup. When activated
  *  via Start function it registers and enables the FMassProcessingPhase instances which themselves are tick functions 
- *  that host FMassRuntimePipeline which they trigger as part of their Tick function. 
+ *  that host UMassCompositeProcessor which they trigger as part of their Tick function. 
  *  MassProcessingPhaseManager serves as an interface to said FMassProcessingPhase instances and allows initialization
  *  with MassSchematics (via InitializePhases function) as well as registering arbitrary functions to be called 
  *  when a particular phase starts of ends (via GetOnPhaseStart and GetOnPhaseEnd functions). */
diff --git a/Engine/Plugins/Runtime/MassGameplay/Source/MassActors/Private/MassAgentTraits.cpp b/Engine/Plugins/Runtime/MassGameplay/Source/MassActors/Private/MassAgentTraits.cpp
--- a/Engine/Plugins/Runtime/MassGameplay/Source/MassActors/Private/MassAgentTraits.cpp
+++ b/Engine/Plugins/Runtime/MassGameplay/Source/MassActors/Private/MassAgentTraits.cpp
@@ -3,7 +3,6 @@
 #include "MassAgentTraits.h"
 #include "MassCommonTypes.h"
 #include "MassMovementFragments.h"
-#include "MassTranslatorRegistry.h"
 #include "MassEntityTemplate.h"
 #include "MassEntityTemplateRegistry.h"
 #include "Translators/MassCapsuleComponentTranslators.h"
