Extended FMassProcessingContext with a command buffer shared ptr to be used if set, and a flag instructing mass executor whether to flush commands after processing is done

diff --git a/Engine/Plugins/Runtime/MassEntity/Source/MassEntity/Private/MassExecutor.cpp b/Engine/Plugins/Runtime/MassEntity/Source/MassEntity/Private/MassExecutor.cpp
--- a/Engine/Plugins/Runtime/MassEntity/Source/MassEntity/Private/MassExecutor.cpp
+++ b/Engine/Plugins/Runtime/MassEntity/Source/MassEntity/Private/MassExecutor.cpp
@@ -89,8 +89,12 @@ void RunProcessorsView(TArrayView<UMassProcessor*> Processors, FMassProcessingCo
 	{
 		ExecutionContext.SetChunkCollection(*ChunkCollection);
 	}
-	// manually creating a new command buffer to let the default one still be used by code unaware of mass processing
-	ExecutionContext.SetDeferredCommandBuffer(MakeShareable(new FMassCommandBuffer()));
+	
+	// if ProcessingContext points at a valid CommandBuffer use that one, otherwise manually create a new command buffer 
+	// to let the default one still be used by code unaware of mass processing
+	TSharedPtr<FMassCommandBuffer> CommandBuffer = ProcessingContext.CommandBuffer 
+		? ProcessingContext.CommandBuffer : MakeShareable(new FMassCommandBuffer());
+	ExecutionContext.SetDeferredCommandBuffer(CommandBuffer);
 	ExecutionContext.SetFlushDeferredCommands(false);
 	ExecutionContext.SetAuxData(ProcessingContext.AuxData);
 
@@ -105,6 +109,7 @@ void RunProcessorsView(TArrayView<UMassProcessor*> Processors, FMassProcessingCo
 		}
 	}
 	
+	if (ProcessingContext.bFlushCommandBuffer)
 	{		
 		TRACE_CPUPROFILER_EVENT_SCOPE_STR("Flush Deferred Commands")
 		
@@ -167,7 +172,9 @@ FGraphEventRef TriggerParallelTasks(UMassProcessor& Processor, FMassProcessingCo
 
 	// not going through UMassEntitySubsystem::CreateExecutionContext on purpose - we do need a separate command buffer
 	FMassExecutionContext ExecutionContext(ProcessingContext.DeltaSeconds);
-	ExecutionContext.SetDeferredCommandBuffer(MakeShareable(new FMassCommandBuffer()));
+	TSharedPtr<FMassCommandBuffer> CommandBuffer = ProcessingContext.CommandBuffer
+		? ProcessingContext.CommandBuffer : MakeShareable(new FMassCommandBuffer());
+	ExecutionContext.SetDeferredCommandBuffer(CommandBuffer);
 	ExecutionContext.SetFlushDeferredCommands(false);
 	ExecutionContext.SetAuxData(ProcessingContext.AuxData);
 
diff --git a/Engine/Plugins/Runtime/MassEntity/Source/MassEntity/Public/MassProcessingTypes.h b/Engine/Plugins/Runtime/MassEntity/Source/MassEntity/Public/MassProcessingTypes.h
--- a/Engine/Plugins/Runtime/MassEntity/Source/MassEntity/Public/MassProcessingTypes.h
+++ b/Engine/Plugins/Runtime/MassEntity/Source/MassEntity/Public/MassProcessingTypes.h
@@ -12,6 +12,7 @@ class UMassEntitySubsystem;
 class UMassProcessor;
 class UMassSchematic;
 class UMassCompositeProcessor;
+struct FMassCommandBuffer;
 
 UENUM(meta = (Bitflags, UseEnumValuesAsMaskValuesInEditor = "true"))
 enum class EProcessorExecutionFlags : uint8
@@ -44,6 +45,15 @@ struct MASSENTITY_API FMassProcessingContext
 	UPROPERTY()
 	FInstancedStruct AuxData;
 
+	/** 
+	 * If set to "true" the MassExecutor will flush commands at the end of given execution function. 
+	 * If "false" the caller is responsible for manually flushing the commands.
+	 */
+	UPROPERTY()
+	bool bFlushCommandBuffer = true; 
+		
+	TSharedPtr<FMassCommandBuffer> CommandBuffer;
+	
 	FMassProcessingContext() = default;
 	FMassProcessingContext(UMassEntitySubsystem& InEntities, const float InDeltaSeconds);
 };
