Added a function to FMassRuntimePipeline that lets caller add a processor based on a given UClass

diff --git a/Engine/Plugins/Runtime/MassEntity/Source/MassEntity/Private/MassProcessingTypes.cpp b/Engine/Plugins/Runtime/MassEntity/Source/MassEntity/Private/MassProcessingTypes.cpp
--- a/Engine/Plugins/Runtime/MassEntity/Source/MassEntity/Private/MassProcessingTypes.cpp
+++ b/Engine/Plugins/Runtime/MassEntity/Source/MassEntity/Private/MassProcessingTypes.cpp
@@ -196,6 +196,13 @@ void FMassRuntimePipeline::AppendProcessor(UMassProcessor& Processor)
 	Processors.Add(&Processor);
 }
 
+void FMassRuntimePipeline::AppendProcessor(TSubclassOf<UMassProcessor> ProcessorClass, UObject& InOwner)
+{
+	check(ProcessorClass);
+	UMassProcessor* ProcInstance = NewObject<UMassProcessor>(&InOwner, ProcessorClass);
+	AppendProcessor(*ProcInstance);
+}
+
 UMassCompositeProcessor* FMassRuntimePipeline::FindTopLevelGroupByName(FName GroupName)
 {
 	for (UMassProcessor* Processor : Processors)
diff --git a/Engine/Plugins/Runtime/MassEntity/Source/MassEntity/Public/MassProcessingTypes.h b/Engine/Plugins/Runtime/MassEntity/Source/MassEntity/Public/MassProcessingTypes.h
--- a/Engine/Plugins/Runtime/MassEntity/Source/MassEntity/Public/MassProcessingTypes.h
+++ b/Engine/Plugins/Runtime/MassEntity/Source/MassEntity/Public/MassProcessingTypes.h
@@ -101,6 +101,9 @@ struct MASSENTITY_API FMassRuntimePipeline
 	/** Adds InProcessor to Processors without any additional checks */
 	void AppendProcessor(UMassProcessor& Processor);
 
+	/** Creates an instance of ProcessorClass and adds it to Processors without any additional checks */
+	void AppendProcessor(TSubclassOf<UMassProcessor> ProcessorClass, UObject& InOwner);
+
 	/** goes through Processor looking for a UMassCompositeProcessor instance which GroupName matches the one given as the parameter */
 	UMassCompositeProcessor* FindTopLevelGroupByName(const FName GroupName);
 
