Added FMassArchetypeCompositionDescriptor.HasAll

diff --git a/Engine/Plugins/Runtime/MassEntity/Source/MassEntity/Public/MassEntityTypes.h b/Engine/Plugins/Runtime/MassEntity/Source/MassEntity/Public/MassEntityTypes.h
--- a/Engine/Plugins/Runtime/MassEntity/Source/MassEntity/Public/MassEntityTypes.h
+++ b/Engine/Plugins/Runtime/MassEntity/Source/MassEntity/Public/MassEntityTypes.h
@@ -162,6 +162,14 @@ struct FMassArchetypeCompositionDescriptor
 			SharedFragments.IsEmpty();
 	}
 
+	bool HasAll(const FMassArchetypeCompositionDescriptor& OtherDescriptor) const
+	{
+		return Fragments.HasAll(OtherDescriptor.Fragments) &&
+			Tags.HasAll(OtherDescriptor.Tags) &&
+			ChunkFragments.HasAll(OtherDescriptor.ChunkFragments) &&
+			SharedFragments.HasAll(OtherDescriptor.SharedFragments);
+	}
+
 	static uint32 CalculateHash(const FMassFragmentBitSet& InFragments, const FMassTagBitSet& InTags, const FMassChunkFragmentBitSet& InChunkFragments, const FMassSharedFragmentBitSet& InSharedFragmentBitSet)
 	{
 		const uint32 FragmentsHash = GetTypeHash(InFragments);
