New approach to Mass' fragment initialization and deinitialization.

New observers triggering is a somewhat manual process, but it will be integrated into the EntitySubsystem's methods in the short-term future work.

* The CL contains a new concept, the FEntityCreationContext, that's being used to ensure the "on entities created" observers are notified only once we copied over instanced structs data over from the entity template.
* With the switch to new de/initialization the EntityTemplate no longer needs to store that data so it has been cut
* With the going away of the EntityTemplate-configured de/initializers all the traits had to be updated to remove the calls to the cut API
* A temp hack required to make randomized velocity trait to work as expected.
* Changes to CommandBuffer to utilize he new approach. The change adds operation type information to commands so that we know whether a command adds and/or removes fragments. This data is used to gather and distribute data to fragment addition/removal observers. This approach works but will be a subject to optimization in following work.
* MassSpawnerSubsystem updated to use new fragment de/initialization approach.

diff --git a/Engine/Plugins/AI/MassAI/Source/MassAIBehavior/Private/MassStateTreeTrait.cpp b/Engine/Plugins/AI/MassAI/Source/MassAIBehavior/Private/MassStateTreeTrait.cpp
--- a/Engine/Plugins/AI/MassAI/Source/MassAIBehavior/Private/MassStateTreeTrait.cpp
+++ b/Engine/Plugins/AI/MassAI/Source/MassAIBehavior/Private/MassStateTreeTrait.cpp
@@ -32,7 +32,7 @@ void UMassStateTreeTrait::BuildTemplate(FMassEntityTemplateBuildContext& BuildCo
 	const FMassStateTreeHandle Handle = MassStateTreeSubsystem->RegisterStateTreeAsset(StateTree);
 
 	// Add fragment describing which StateTree to run
-	FMassStateTreeFragment& StateTreeFragment = BuildContext.AddFragmentWithDefaultInitializer_GetRef<FMassStateTreeFragment>();
+	FMassStateTreeFragment& StateTreeFragment = BuildContext.AddFragment_GetRef<FMassStateTreeFragment>();
 	StateTreeFragment.StateTreeHandle = Handle;
 
 	// Add runtime storage as a fragment
diff --git a/Engine/Plugins/AI/MassAI/Source/MassAIBehavior/Private/MassZoneGraphAnnotationTrait.cpp b/Engine/Plugins/AI/MassAI/Source/MassAIBehavior/Private/MassZoneGraphAnnotationTrait.cpp
--- a/Engine/Plugins/AI/MassAI/Source/MassAIBehavior/Private/MassZoneGraphAnnotationTrait.cpp
+++ b/Engine/Plugins/AI/MassAI/Source/MassAIBehavior/Private/MassZoneGraphAnnotationTrait.cpp
@@ -7,6 +7,6 @@
 
 void UMassZoneGraphAnnotationTrait::BuildTemplate(FMassEntityTemplateBuildContext& BuildContext, UWorld& World) const
 {
-	BuildContext.AddFragmentWithDefaultInitializer<FMassZoneGraphAnnotationTagsFragment>();
+	BuildContext.AddFragment<FMassZoneGraphAnnotationTagsFragment>();
 	BuildContext.AddChunkFragment<FMassZoneGraphAnnotationVariableTickChunkFragment>();
 }
diff --git a/Engine/Plugins/AI/MassAI/Source/MassAIMovement/Private/MassAvoidanceObstacleTrait.cpp b/Engine/Plugins/AI/MassAI/Source/MassAIMovement/Private/MassAvoidanceObstacleTrait.cpp
--- a/Engine/Plugins/AI/MassAI/Source/MassAIMovement/Private/MassAvoidanceObstacleTrait.cpp
+++ b/Engine/Plugins/AI/MassAI/Source/MassAIMovement/Private/MassAvoidanceObstacleTrait.cpp
@@ -11,9 +11,8 @@
 void UMassAvoidanceObstacleTrait::BuildTemplate(FMassEntityTemplateBuildContext& BuildContext, UWorld& World) const
 {
 	// UAvoidanceObstacleProcessor
-	BuildContext.AddFragmentWithDefaultInitializer<FMassAvoidanceObstacleGridCellLocationFragment>();
-	BuildContext.AddDeinitializer(*GetDefault<UMassAvoidanceObstacleRemoverFragmentDestructor>());
-	BuildContext.AddFragmentWithDefaultInitializer<FDataFragment_AgentRadius>();
+	BuildContext.AddFragment<FMassAvoidanceObstacleGridCellLocationFragment>();
+	BuildContext.AddFragment<FDataFragment_AgentRadius>();
 
 	if (bExtendToEdgeObstacle)
 	{
diff --git a/Engine/Plugins/AI/MassAI/Source/MassAIMovement/Private/MassAvoidanceTrait.cpp b/Engine/Plugins/AI/MassAI/Source/MassAIMovement/Private/MassAvoidanceTrait.cpp
--- a/Engine/Plugins/AI/MassAI/Source/MassAIMovement/Private/MassAvoidanceTrait.cpp
+++ b/Engine/Plugins/AI/MassAI/Source/MassAIMovement/Private/MassAvoidanceTrait.cpp
@@ -11,10 +11,10 @@
 void UMassAvoidanceTrait::BuildTemplate(FMassEntityTemplateBuildContext& BuildContext, UWorld& World) const
 {
 	// UMassAvoidanceProcessor
-	BuildContext.AddFragmentWithDefaultInitializer<FDataFragment_AgentRadius>(); // From UMassProcessor_AgentMovement
-	BuildContext.AddFragmentWithDefaultInitializer<FMassNavigationEdgesFragment>();
-	BuildContext.AddFragmentWithDefaultInitializer<FDataFragment_Transform>(); // From UMassProcessor_Movement
-	BuildContext.AddFragmentWithDefaultInitializer<FMassVelocityFragment>(); // From UMassProcessor_Movement
-	BuildContext.AddFragmentWithDefaultInitializer<FMassSteeringFragment>();
-	BuildContext.AddFragmentWithDefaultInitializer<FMassMoveTargetFragment>();
+	BuildContext.AddFragment<FDataFragment_AgentRadius>(); // From UMassProcessor_AgentMovement
+	BuildContext.AddFragment<FMassNavigationEdgesFragment>();
+	BuildContext.AddFragment<FDataFragment_Transform>(); // From UMassProcessor_Movement
+	BuildContext.AddFragment<FMassVelocityFragment>(); // From UMassProcessor_Movement
+	BuildContext.AddFragment<FMassSteeringFragment>();
+	BuildContext.AddFragment<FMassMoveTargetFragment>();
 }
diff --git a/Engine/Plugins/AI/MassAI/Source/MassAIMovement/Private/MassZoneGraphMovementTrait.cpp b/Engine/Plugins/AI/MassAI/Source/MassAIMovement/Private/MassZoneGraphMovementTrait.cpp
--- a/Engine/Plugins/AI/MassAI/Source/MassAIMovement/Private/MassZoneGraphMovementTrait.cpp
+++ b/Engine/Plugins/AI/MassAI/Source/MassAIMovement/Private/MassZoneGraphMovementTrait.cpp
@@ -19,13 +19,13 @@ void UMassZoneGraphMovementTrait::BuildTemplate(FMassEntityTemplateBuildContext&
 		return;
 	}
 
-	BuildContext.AddFragmentWithDefaultInitializer<FDataFragment_AgentRadius>();
-	BuildContext.AddFragmentWithDefaultInitializer<FDataFragment_Transform>();
-	BuildContext.AddFragmentWithDefaultInitializer<FMassVelocityFragment>();
+	BuildContext.AddFragment<FDataFragment_AgentRadius>();
+	BuildContext.AddFragment<FDataFragment_Transform>();
+	BuildContext.AddFragment<FMassVelocityFragment>();
 
 	BuildContext.AddFragment<FMassMoveTargetFragment>();
 	
-	BuildContext.AddFragmentWithDefaultInitializer<FMassZoneGraphLaneLocationFragment>();
+	BuildContext.AddFragment<FMassZoneGraphLaneLocationFragment>();
 	
 	const UMassMovementSettings* Settings = GetDefault<UMassMovementSettings>();
 	check(Settings);
diff --git a/Engine/Plugins/AI/MassCrowd/Source/MassCrowd/Private/MassCrowdMemberTrait.cpp b/Engine/Plugins/AI/MassCrowd/Source/MassCrowd/Private/MassCrowdMemberTrait.cpp
--- a/Engine/Plugins/AI/MassCrowd/Source/MassCrowd/Private/MassCrowdMemberTrait.cpp
+++ b/Engine/Plugins/AI/MassCrowd/Source/MassCrowd/Private/MassCrowdMemberTrait.cpp
@@ -7,5 +7,5 @@
 void UMassCrowdMemberTrait::BuildTemplate(FMassEntityTemplateBuildContext& BuildContext, UWorld& World) const
 {
 	BuildContext.AddTag<FTagFragment_MassCrowd>();
-	BuildContext.AddFragmentWithDefaultInitializer<FMassCrowdLaneTrackingFragment>();	
+	BuildContext.AddFragment<FMassCrowdLaneTrackingFragment>();	
 }
diff --git a/Engine/Plugins/AI/MassCrowd/Source/MassCrowd/Private/MassCrowdServerRepresentationTrait.cpp b/Engine/Plugins/AI/MassCrowd/Source/MassCrowd/Private/MassCrowdServerRepresentationTrait.cpp
--- a/Engine/Plugins/AI/MassCrowd/Source/MassCrowd/Private/MassCrowdServerRepresentationTrait.cpp
+++ b/Engine/Plugins/AI/MassCrowd/Source/MassCrowd/Private/MassCrowdServerRepresentationTrait.cpp
@@ -22,7 +22,7 @@ void UMassCrowdServerRepresentationTrait::BuildTemplate(FMassEntityTemplateBuild
 
 	if (UMassCrowdRepresentationSubsystem* RepresentationSubsystem = World.GetSubsystem<UMassCrowdRepresentationSubsystem>())
 	{
-		FMassRepresentationFragment& RepresentationFragment = BuildContext.AddFragmentWithDefaultInitializer_GetRef<FMassRepresentationFragment>();
+		FMassRepresentationFragment& RepresentationFragment = BuildContext.AddFragment_GetRef<FMassRepresentationFragment>();
 		RepresentationFragment.StaticMeshDescIndex = INDEX_NONE;
 		RepresentationFragment.HighResTemplateActorIndex = TemplateActor.Get() ? RepresentationSubsystem->FindOrAddTemplateActor(TemplateActor.Get()) : INDEX_NONE;
 		RepresentationFragment.LowResTemplateActorIndex = INDEX_NONE;
diff --git a/Engine/Plugins/Runtime/MassEntity/Source/MassEntity/Private/MassCommandBuffer.cpp b/Engine/Plugins/Runtime/MassEntity/Source/MassEntity/Private/MassCommandBuffer.cpp
--- a/Engine/Plugins/Runtime/MassEntity/Source/MassEntity/Private/MassCommandBuffer.cpp
+++ b/Engine/Plugins/Runtime/MassEntity/Source/MassEntity/Private/MassCommandBuffer.cpp
@@ -1,8 +1,12 @@
 // Copyright Epic Games, Inc. All Rights Reserved.
 
 #include "MassCommandBuffer.h"
+#include "MassObserverManager.h"
+#include "MassEntityUtils.h"
 #include "HAL/IConsoleManager.h"
 #include "ProfilingDebugging/CsvProfiler.h"
+#include "VisualLogger/VisualLogger.h"
+
 
 //////////////////////////////////////////////////////////////////////
 // FMassCommandBuffer
@@ -66,12 +70,59 @@ void GetCommandStatNames(FStructView Entry, FString& OutName, ANSIName*& OutANSI
 #endif
 } // UE::FLWCCommand
 
-void FMassCommandBuffer::ReplayBufferAgainstSystem(UMassEntitySubsystem* System)
+void FMassCommandBuffer::ReplayBufferAgainstSystem(UMassEntitySubsystem* EntitySystem)
 {
-	check(System);
+	check(EntitySystem);
+
+	// short-circuit exit
+	if (EntitiesToDestroy.Num() == 0 && PendingCommands.IsEmpty())
+	{
+		return;
+	}
+
+	FMassObserverManager& ObserverManager = EntitySystem->GetObserverManager();
+
+	const FMassFragmentBitSet& ObservedAddFragments = ObserverManager.GetObservedAddFragmentsBitSet();
+	const FMassFragmentBitSet& ObservedRemoveFragments = ObserverManager.GetObservedRemoveFragmentsBitSet();
+
+	if (ObservedRemoveFragments.IsEmpty() == false)
+	{
+		for (auto It : FragmentsToRemove)
+		{
+			check(It.Key);
+			if (ObservedRemoveFragments.Contains(*It.Key))
+			{
+				TArray<FArchetypeChunkCollection> ChunkCollections;
+				UE::Mass::Utils::CreateSparseChunks(*EntitySystem, It.Value, ChunkCollections);
+				for (FArchetypeChunkCollection& Collection : ChunkCollections)
+				{
+					check(It.Key);
+					ObserverManager.OnPreFragmentRemoved(*It.Key, Collection);
+				}
+			}
+		}
+	
+		TArray<FArchetypeChunkCollection> EntityChunksToDestroy;
+		if (EntitiesToDestroy.Num())
+		{
+			UE::Mass::Utils::CreateSparseChunks(*EntitySystem, EntitiesToDestroy, EntityChunksToDestroy);
+			for (FArchetypeChunkCollection& Collection : EntityChunksToDestroy)
+			{
+				ObserverManager.OnPreEntitiesDestroyed(Collection);
+				EntitySystem->BatchDestroyEntityChunks(Collection);
+			}
+		}
+	}
+	else if (EntitiesToDestroy.Num())
+	{
+		TRACE_CPUPROFILER_EVENT_SCOPE_STR("ECS Deferred Commands Destroy Entities");
+		EntitySystem->BatchDestroyEntities(MakeArrayView(EntitiesToDestroy));
+	}
+	EntitiesToDestroy.Reset();
+
 
 	UE_MT_SCOPED_WRITE_ACCESS(PendingCommandsDetector);
-	PendingCommands.ForEach([System](FStructView Entry)
+	PendingCommands.ForEach([EntitySystem](FStructView Entry)
 	{
 		const FCommandBufferEntryBase* Command = Entry.GetPtr<FCommandBufferEntryBase>();
 		checkf(Command, TEXT("Either the entry is null or the command does not derive from FCommandBufferEntryBase"));
@@ -89,25 +140,36 @@ void FMassCommandBuffer::ReplayBufferAgainstSystem(UMassEntitySubsystem* System)
 		FCsvProfiler::RecordCustomStat(*Name, CSV_CATEGORY_INDEX(MassEntitiesCounters), 1, ECsvCustomStatOp::Accumulate);
 #endif // CSV_PROFILER
 
-		Command->Execute(*System);
+		Command->Execute(*EntitySystem);
 	});
 
 	// Using Clear() instead of Reset(), as otherwise the chunks moved into the PendingCommands in MoveAppend() can accumulate.
 	PendingCommands.Clear();
 
-	if (EntitiesToDestroy.Num())
+	for (auto It : FragmentsToAdd)
 	{
-		TRACE_CPUPROFILER_EVENT_SCOPE_STR("ECS Deferred Commands Destroy Entities");
-		System->BatchDestroyEntities(MakeArrayView(EntitiesToDestroy));
-		EntitiesToDestroy.Reset();
+		check(It.Key);
+		if (ObservedAddFragments.Contains(*It.Key))
+		{
+			TArray<FArchetypeChunkCollection> ChunkCollections;
+			UE::Mass::Utils::CreateSparseChunks(*EntitySystem, It.Value, ChunkCollections);
+			for (FArchetypeChunkCollection& Collection : ChunkCollections)
+			{
+				check(It.Key);
+				ObserverManager.OnPostFragmentAdded(*It.Key, Collection);
+			}
+		}
 	}
+
+	FragmentsToAdd.Reset();
+	FragmentsToRemove.Reset();
 }
 
 void FMassCommandBuffer::MoveAppend(FMassCommandBuffer& Other)
 {
 	// @todo optimize, there surely a way to do faster then this.
 	UE_MT_SCOPED_READ_ACCESS(Other.PendingCommandsDetector);
-	if (Other.GetPendingCommandsCount() || Other.EntitiesToDestroy.Num())
+	if (Other.HasPendingCommands() || Other.EntitiesToDestroy.Num())
 	{
 		FScopeLock Lock(&AppendingCommandsCS);
 		UE_MT_SCOPED_WRITE_ACCESS(PendingCommandsDetector);
@@ -178,7 +240,7 @@ void FCommandRemoveFragmentList::Execute(UMassEntitySubsystem& System) const
 void FCommandAddTag::Execute(UMassEntitySubsystem & System) const 
 {
 	if (System.IsEntityValid(TargetEntity) == false)
-{
+	{
 		return;
 	}
 
diff --git a/Engine/Plugins/Runtime/MassEntity/Source/MassEntity/Private/MassEntitySubsystem.cpp b/Engine/Plugins/Runtime/MassEntity/Source/MassEntity/Private/MassEntitySubsystem.cpp
--- a/Engine/Plugins/Runtime/MassEntity/Source/MassEntity/Private/MassEntitySubsystem.cpp
+++ b/Engine/Plugins/Runtime/MassEntity/Source/MassEntity/Private/MassEntitySubsystem.cpp
@@ -18,6 +18,7 @@ const FMassEntityHandle UMassEntitySubsystem::InvalidEntity;
 // UMassEntitySubsystem
 
 UMassEntitySubsystem::UMassEntitySubsystem()
+	: ObserverManager(*this)
 {
 }
 
@@ -72,6 +73,13 @@ void UMassEntitySubsystem::Initialize(FSubsystemCollectionBase& Collection)
 	}
 }
 
+void UMassEntitySubsystem::PostInitialize()
+{
+	// this needs to be done after all the subsystems have been initialized since some processors might want to access
+	// them during processors' initialization
+	ObserverManager.Initialize();
+}
+
 void UMassEntitySubsystem::Deinitialize()
 {
 	//@TODO: Should actually do this...
@@ -302,11 +310,11 @@ void UMassEntitySubsystem::BuildEntity(FMassEntityHandle Entity, TConstArrayView
 	EntityData.CurrentArchetype->SetFragmentsData(Entity, FragmentInstanceList);
 }
 
-int32 UMassEntitySubsystem::BatchCreateEntities(const FArchetypeHandle Archetype, const int32 Count, TArray<FMassEntityHandle>& OutEntities)
+TSharedRef<UMassEntitySubsystem::FEntityCreationContext> UMassEntitySubsystem::BatchCreateEntities(const FArchetypeHandle Archetype, const int32 Count, TArray<FMassEntityHandle>& OutEntities)
 {
 	FMassArchetypeData* ArchetypePtr = Archetype.DataPtr.Get();
 	check(ArchetypePtr);
-	check(Count >= 0);
+	check(Count > 0);
 	
 	int32 Index = OutEntities.Num();
 	OutEntities.AddDefaulted(Count);
@@ -324,8 +332,19 @@ int32 UMassEntitySubsystem::BatchCreateEntities(const FArchetypeHandle Archetype
 		
 		ArchetypePtr->AddEntity(Result);
 	}
+		
+	FEntityCreationContext* CreationContext = new FEntityCreationContext(Count);
+	// @todo this could probably be optimized since one would assume we're adding elements to OutEntities in order.
+	// Then again, if that's the case, the sorting will be almost instant
+	new (&CreationContext->ChunkCollection)FArchetypeChunkCollection(Archetype, MakeArrayView(&OutEntities[OutEntities.Num() - Count], Count));
+	if (ObserverManager.HasOnAddedObserversForFragments(ArchetypePtr->GetCompositionDescriptor().Fragments))
+	{
+		CreationContext->OnSpawningFinished = [this](FEntityCreationContext& Context){
+			ObserverManager.OnPostEntitiesCreated(Context.ChunkCollection);
+		};
+	}
 
-	return Count;
+	return MakeShareable(CreationContext);
 }
 
 void UMassEntitySubsystem::DestroyEntity(FMassEntityHandle Entity)
@@ -384,6 +403,11 @@ void UMassEntitySubsystem::BatchDestroyEntityChunks(const FArchetypeChunkCollect
 
 	checkf(IsProcessing() == false, TEXT("Synchronous API function %s called during mass processing. Use asynchronous API instead."), ANSI_TO_TCHAR(__FUNCTION__));
 
+	FMassProcessingContext ProcessingContext(*this, /*TimeDelta=*/0.0f);
+	ProcessingContext.bFlushCommandBuffer = false;
+	ProcessingContext.CommandBuffer = MakeShareable(new FMassCommandBuffer());
+	ObserverManager.OnPreEntitiesDestroyed(ProcessingContext, Chunks);
+
 	check(Chunks.GetArchetype().IsValid());
 	TArray<FMassEntityHandle> EntitiesRemoved;
 	Chunks.GetArchetype().DataPtr->BatchDestroyEntityChunks(Chunks, EntitiesRemoved);
@@ -401,6 +425,8 @@ void UMassEntitySubsystem::BatchDestroyEntityChunks(const FArchetypeChunkCollect
 			EntityFreeIndexList.Add(Entity.Index);
 		}
 	}
+
+	ensure(ProcessingContext.CommandBuffer->HasPendingCommands() == false);
 }
 
 void UMassEntitySubsystem::AddFragmentToEntity(FMassEntityHandle Entity, const UScriptStruct* FragmentType)
@@ -448,6 +474,8 @@ void UMassEntitySubsystem::AddCompositionToEntity_GetDelta(FMassEntityHandle Ent
 			check(NewArchetype);
 			EntityData.CurrentArchetype->MoveEntityToAnotherArchetype(Entity, *NewArchetype);
 			EntityData.CurrentArchetype = NewArchetypeHandle.DataPtr;
+
+			ObserverManager.OnPostCompositionAdded(Entity, InDescriptor);
 		}
 	}
 }
@@ -469,20 +497,28 @@ void UMassEntitySubsystem::RemoveCompositionFromEntity(FMassEntityHandle Entity,
 		ensureMsgf(InDescriptor.ChunkFragments.IsEmpty(), TEXT("Removing chunk fragments is not supported"));
 		ensureMsgf(InDescriptor.SharedFragments.IsEmpty(), TEXT("Removing shared fragments is not supported"));
 
-		if (NewDescriptor.IsEquivalent(InDescriptor) == false)
-	{
+		if (NewDescriptor.IsEquivalent(OldArchetype->GetCompositionDescriptor()) == false)
+		{
+			ensureMsgf(OldArchetype->GetCompositionDescriptor().HasAll(InDescriptor), TEXT("Some of the elements being removed are already missing from entity\'s composition."));
+			ObserverManager.OnPreCompositionRemoved(Entity, InDescriptor);
+
 			const FArchetypeHandle NewArchetypeHandle = CreateArchetype(NewDescriptor, OldArchetype->GetSharedFragmentValues());
 
-		if (ensure(NewArchetypeHandle.DataPtr != EntityData.CurrentArchetype))
-		{
-			// Move the entity over
-			FMassArchetypeData* NewArchetype = NewArchetypeHandle.DataPtr.Get();
-			check(NewArchetype);
-			EntityData.CurrentArchetype->MoveEntityToAnotherArchetype(Entity, *NewArchetype);
-			EntityData.CurrentArchetype = NewArchetypeHandle.DataPtr;
+			if (ensure(NewArchetypeHandle.DataPtr != EntityData.CurrentArchetype))
+			{
+				// Move the entity over
+				FMassArchetypeData* NewArchetype = NewArchetypeHandle.DataPtr.Get();
+				check(NewArchetype);
+				EntityData.CurrentArchetype->MoveEntityToAnotherArchetype(Entity, *NewArchetype);
+				EntityData.CurrentArchetype = NewArchetypeHandle.DataPtr;
+			}
 		}
 	}
 }
+
+const FMassArchetypeCompositionDescriptor& UMassEntitySubsystem::GetArchetypeComposition(const FArchetypeHandle& ArchetypeHandle) const
+{
+	return ArchetypeHandle.DataPtr->GetCompositionDescriptor();
 }
 
 void UMassEntitySubsystem::InternalBuildEntity(FMassEntityHandle Entity, const FArchetypeHandle Archetype)
@@ -997,6 +1033,11 @@ int32 UMassEntitySubsystem::DebugGetArchetypeEntitiesCount(const FArchetypeHandl
 	return Archetype.DataPtr.IsValid() ? Archetype.DataPtr->GetNumEntities() : 0;
 }
 
+int32 UMassEntitySubsystem::DebugGetArchetypeEntitiesCountPerChunk(const FArchetypeHandle& Archetype) const
+{
+	return Archetype.DataPtr.IsValid() ? Archetype.DataPtr->GetNumEntitiesPerChunk() : 0;
+}
+
 void UMassEntitySubsystem::DebugRemoveAllEntities()
 {
 	for (int EntityIndex = NumReservedEntities; EntityIndex < Entities.Num(); ++EntityIndex)
diff --git a/Engine/Plugins/Runtime/MassEntity/Source/MassEntity/Private/MassObserverManager.cpp b/Engine/Plugins/Runtime/MassEntity/Source/MassEntity/Private/MassObserverManager.cpp
new file mode 100644
--- /dev/null
+++ b/Engine/Plugins/Runtime/MassEntity/Source/MassEntity/Private/MassObserverManager.cpp
@@ -0,0 +1,186 @@
+// Copyright Epic Games, Inc. All Rights Reserved.
+
+#include "MassObserverManager.h"
+#include "MassEntitySubsystem.h"
+#include "MassExecutor.h"
+#include "MassProcessingTypes.h"
+#include "MassObserverRegistry.h"
+
+//----------------------------------------------------------------------//
+// FMassObserverManager
+//----------------------------------------------------------------------//
+FMassObserverManager::FMassObserverManager()
+	: EntitySubsystem(*GetMutableDefault<UMassEntitySubsystem>())
+{
+
+}
+
+FMassObserverManager::FMassObserverManager(UMassEntitySubsystem& Owner)
+	: EntitySubsystem(Owner)
+{
+
+}
+
+void FMassObserverManager::Initialize()
+{
+	// instantiate initializers
+	const UMassObserverRegistry& Registry = UMassObserverRegistry::Get();
+
+	ObservedAddFragments.Reset();
+	ObservedRemoveFragments.Reset();
+	OnFragmentAddedObservers.Reset();
+	OnFragmentRemovedObservers.Reset();
+
+	for (auto It : Registry.FragmentInitializersMap)
+	{
+		if (It.Value.ClassCollection.Num() == 0)
+		{
+			continue;
+		}
+
+		ObservedAddFragments.Add(*It.Key);
+		FMassRuntimePipeline& Pipeline = OnFragmentAddedObservers.FindOrAdd(It.Key);
+
+		for (const TSubclassOf<UMassProcessor>& ProcessorClass : It.Value.ClassCollection)
+		{
+			Pipeline.AppendProcessor(ProcessorClass, EntitySubsystem);
+		}
+		Pipeline.Initialize(EntitySubsystem);
+	}
+
+	for (auto It : Registry.FragmentDeinitializersMap)
+	{
+		if (It.Value.ClassCollection.Num() == 0)
+		{
+			continue;
+		}
+
+		ObservedRemoveFragments.Add(*It.Key);
+		FMassRuntimePipeline& Pipeline = OnFragmentRemovedObservers.FindOrAdd(It.Key);
+
+		for (const TSubclassOf<UMassProcessor>& ProcessorClass : It.Value.ClassCollection)
+		{
+			Pipeline.AppendProcessor(ProcessorClass, EntitySubsystem);
+		}
+		Pipeline.Initialize(EntitySubsystem);
+	}
+}
+
+bool FMassObserverManager::OnPostEntitiesCreated(const FArchetypeChunkCollection& ChunkCollection)
+{
+	FMassProcessingContext ProcessingContext(EntitySubsystem, /*DeltaSeconds=*/0.f);
+	// requesting not to flush commands since handling creation of new entities can result in multiple collections of
+	// processors being executed and flushing commands between these runs would ruin ChunkCollection since entities could
+	// get their composition changed and get moved to new archetypes
+	ProcessingContext.bFlushCommandBuffer = false;
+	ProcessingContext.CommandBuffer = MakeShareable(new FMassCommandBuffer());
+
+	if (OnPostEntitiesCreated(ProcessingContext, ChunkCollection))
+	{
+		ProcessingContext.CommandBuffer->ReplayBufferAgainstSystem(&EntitySubsystem);
+		return true;
+	}
+	return false;
+}
+
+bool FMassObserverManager::OnPostEntitiesCreated(FMassProcessingContext& ProcessingContext, const FArchetypeChunkCollection& ChunkCollection)
+{
+	TRACE_CPUPROFILER_EVENT_SCOPE_STR("OnPostEntitiesCreated")
+
+	check(ProcessingContext.EntitySubsystem);
+	const FMassArchetypeCompositionDescriptor& ArchetypeComposition = ProcessingContext.EntitySubsystem->GetArchetypeComposition(ChunkCollection.GetArchetype());
+	const FMassFragmentBitSet Overlap = ObservedAddFragments.GetOverlap(ArchetypeComposition.Fragments);
+
+	if (Overlap.IsEmpty() == false)
+	{
+		HandleFragmentsImpl(ProcessingContext, ChunkCollection, Overlap, OnFragmentAddedObservers);
+		return true;
+	}
+
+	return false;
+}
+
+bool FMassObserverManager::OnPreEntitiesDestroyed(const FArchetypeChunkCollection& ChunkCollection)
+{
+	FMassProcessingContext ProcessingContext(EntitySubsystem, /*DeltaSeconds=*/0.f);
+	ProcessingContext.bFlushCommandBuffer = false;
+	ProcessingContext.CommandBuffer = MakeShareable(new FMassCommandBuffer());
+
+	if (OnPreEntitiesDestroyed(ProcessingContext, ChunkCollection))
+	{
+		ProcessingContext.CommandBuffer->ReplayBufferAgainstSystem(&EntitySubsystem);
+		return true;
+	}
+	return false;
+}
+
+bool FMassObserverManager::OnPreEntitiesDestroyed(FMassProcessingContext& ProcessingContext, const FArchetypeChunkCollection& ChunkCollection)
+{
+	TRACE_CPUPROFILER_EVENT_SCOPE_STR("OnPreEntitiesDestroyed")
+
+	check(ProcessingContext.EntitySubsystem);
+	const FMassArchetypeCompositionDescriptor& ArchetypeComposition = ProcessingContext.EntitySubsystem->GetArchetypeComposition(ChunkCollection.GetArchetype());
+	const FMassFragmentBitSet Overlap = ObservedRemoveFragments.GetOverlap(ArchetypeComposition.Fragments);
+
+	if (Overlap.IsEmpty() == false)
+	{
+		HandleFragmentsImpl(ProcessingContext, ChunkCollection, Overlap, OnFragmentRemovedObservers);
+		return true;
+	}
+
+	return false;
+}
+
+bool FMassObserverManager::OnPostCompositionAdded(const FMassEntityHandle Entity, const FMassArchetypeCompositionDescriptor& Composition)
+{
+	const FMassFragmentBitSet Overlap = ObservedAddFragments.GetOverlap(Composition.Fragments);
+	if (Overlap.IsEmpty() == false)
+	{
+		FMassProcessingContext ProcessingContext(EntitySubsystem, /*DeltaSeconds=*/0.f);
+		const FArchetypeHandle ArchetypeHandle = EntitySubsystem.GetArchetypeForEntity(Entity);
+		HandleFragmentsImpl(ProcessingContext, FArchetypeChunkCollection(ArchetypeHandle, MakeArrayView(&Entity, 1)), Overlap, OnFragmentAddedObservers);
+		return true;
+	}
+
+	return false;
+}
+
+bool FMassObserverManager::OnPreCompositionRemoved(const FMassEntityHandle Entity, const FMassArchetypeCompositionDescriptor& Composition)
+{
+	const FMassFragmentBitSet Overlap = ObservedRemoveFragments.GetOverlap(Composition.Fragments);
+	if (Overlap.IsEmpty() == false)
+	{
+		FMassProcessingContext ProcessingContext(EntitySubsystem, /*DeltaSeconds=*/0.f);
+		const FArchetypeHandle ArchetypeHandle = EntitySubsystem.GetArchetypeForEntity(Entity);
+		HandleFragmentsImpl(ProcessingContext, FArchetypeChunkCollection(ArchetypeHandle, MakeArrayView(&Entity, 1)), Overlap, OnFragmentRemovedObservers);
+		return true;
+	}
+
+	return false;
+}
+
+void FMassObserverManager::HandleFragmentsImpl(FMassProcessingContext& ProcessingContext, const FArchetypeChunkCollection& ChunkCollection, const FMassFragmentBitSet& FragmentsBitSet, TMap<const UScriptStruct*, FMassRuntimePipeline>& HandlersContainer)
+{
+	TArray<const UScriptStruct*> Fragments;
+	FragmentsBitSet.ExportTypes(Fragments);
+	
+	for (const UScriptStruct* FragmentType : Fragments)
+	{		
+		ProcessingContext.AuxData.InitializeAs(FragmentType);
+		FMassRuntimePipeline& Pipeline = HandlersContainer.FindChecked(FragmentType);
+
+		UE::Mass::Executor::RunProcessorsView(Pipeline.Processors, ProcessingContext, &ChunkCollection);
+	}
+}
+
+void FMassObserverManager::HandleSingleFragmentImpl(const UScriptStruct& FragmentType, const FArchetypeChunkCollection& ChunkCollection, const FMassFragmentBitSet& FragmentFilterBitSet, TMap<const UScriptStruct*, FMassRuntimePipeline>& HandlersContainer)
+{
+	if (FragmentFilterBitSet.Contains(FragmentType))
+	{
+		FMassProcessingContext ProcessingContext(EntitySubsystem, /*DeltaSeconds=*/0.f);
+		ProcessingContext.AuxData.InitializeAs(&FragmentType);
+		FMassRuntimePipeline& Pipeline = HandlersContainer.FindChecked(&FragmentType	);
+
+		UE::Mass::Executor::RunProcessorsView(Pipeline.Processors, ProcessingContext, &ChunkCollection);
+	}
+}
diff --git a/Engine/Plugins/Runtime/MassEntity/Source/MassEntity/Private/MassObserverProcessor.cpp b/Engine/Plugins/Runtime/MassEntity/Source/MassEntity/Private/MassObserverProcessor.cpp
new file mode 100644
--- /dev/null
+++ b/Engine/Plugins/Runtime/MassEntity/Source/MassEntity/Private/MassObserverProcessor.cpp
@@ -0,0 +1,43 @@
+// Copyright Epic Games, Inc. All Rights Reserved.
+
+#include "MassObserverProcessor.h"
+#include "MassObserverRegistry.h"
+
+//----------------------------------------------------------------------//
+// UMassObserverProcessor
+//----------------------------------------------------------------------//
+UMassObserverProcessor::UMassObserverProcessor()
+{
+	bAutoRegisterWithProcessingPhases = false;
+#if WITH_EDITORONLY_DATA
+	bCanShowUpInSettings = false;
+#endif // WITH_EDITORONLY_DATA
+}
+
+void UMassObserverProcessor::PostInitProperties()
+{
+	Super::PostInitProperties();
+
+	if (HasAnyFlags(RF_ClassDefaultObject) && GetClass()->HasAnyClassFlags(CLASS_Abstract) == false && FragmentType != nullptr)
+	{
+		Register();
+	}
+}
+
+//----------------------------------------------------------------------//
+//  UMassFragmentInitializer
+//----------------------------------------------------------------------//
+void UMassFragmentInitializer::Register()
+{
+	check(FragmentType);
+	UMassObserverRegistry::GetMutable().RegisterFragmentInitializer(*FragmentType, GetClass());
+}
+
+//----------------------------------------------------------------------//
+//  UMassFragmentDestructor
+//----------------------------------------------------------------------//
+void UMassFragmentDestructor::Register()
+{
+	check(FragmentType);
+	UMassObserverRegistry::GetMutable().RegisterFragmentDeinitializer(*FragmentType, GetClass());
+}
diff --git a/Engine/Plugins/Runtime/MassEntity/Source/MassEntity/Private/MassObserverRegistry.cpp b/Engine/Plugins/Runtime/MassEntity/Source/MassEntity/Private/MassObserverRegistry.cpp
new file mode 100644
--- /dev/null
+++ b/Engine/Plugins/Runtime/MassEntity/Source/MassEntity/Private/MassObserverRegistry.cpp
@@ -0,0 +1,36 @@
+// Copyright Epic Games, Inc. All Rights Reserved.
+
+#include "MassObserverRegistry.h"
+
+//----------------------------------------------------------------------//
+// UMassObserverRegistry
+//----------------------------------------------------------------------//
+UMassObserverRegistry::UMassObserverRegistry()
+{
+	// there can be only one!
+	check(HasAnyFlags(RF_ClassDefaultObject));
+}
+
+void UMassObserverRegistry::RegisterFragmentInitializer(const UScriptStruct& FragmentType, TSubclassOf<UMassProcessor> FragmentInitializerClass)
+{
+	check(FragmentInitializerClass);
+	FragmentInitializersMap.FindOrAdd(&FragmentType).ClassCollection.AddUnique(FragmentInitializerClass);
+}
+
+void UMassObserverRegistry::RegisterFragmentDeinitializer(const UScriptStruct& FragmentType, TSubclassOf<UMassProcessor> FragmentDeinitializerClass)
+{
+	check(FragmentDeinitializerClass);
+	FragmentDeinitializersMap.FindOrAdd(&FragmentType).ClassCollection.AddUnique(FragmentDeinitializerClass);
+}
+
+TConstArrayView<TSubclassOf<UMassProcessor>> UMassObserverRegistry::GetFragmentInitializers(const UScriptStruct& FragmentType) const
+{
+	const FMassProcessorClassCollection* InitializersCollection = FragmentInitializersMap.Find(&FragmentType);
+	return InitializersCollection ? InitializersCollection->ClassCollection : TConstArrayView<TSubclassOf<UMassProcessor>>();
+}
+
+TConstArrayView<TSubclassOf<UMassProcessor>> UMassObserverRegistry::GetFragmentDeinitializers(const UScriptStruct& FragmentType) const
+{
+	const FMassProcessorClassCollection* DeinitializersCollection = FragmentDeinitializersMap.Find(&FragmentType);
+	return DeinitializersCollection ? DeinitializersCollection->ClassCollection : TConstArrayView<TSubclassOf<UMassProcessor>>();
+}
diff --git a/Engine/Plugins/Runtime/MassEntity/Source/MassEntity/Public/MassCommandBuffer.h b/Engine/Plugins/Runtime/MassEntity/Source/MassEntity/Public/MassCommandBuffer.h
--- a/Engine/Plugins/Runtime/MassEntity/Source/MassEntity/Public/MassCommandBuffer.h
+++ b/Engine/Plugins/Runtime/MassEntity/Source/MassEntity/Public/MassCommandBuffer.h
@@ -9,6 +9,14 @@
 
 #include "MassCommandBuffer.generated.h"
 
+
+namespace ECommandBufferOperationType
+{
+	constexpr int None = 1 << 0;
+	constexpr int Add = 1 << 1;
+	constexpr int Remove = 1 << 2;
+};
+
 //@TODO: Consider debug information in case there is an assert when replaying the command buffer
 // (e.g., which system added the command, or even file/line number in development builds for the specific call via a macro)
 //@TODO: Support more commands
@@ -27,6 +35,7 @@ struct MASSENTITY_API FCommandBufferEntryBase
 		: TargetEntity(InEntity)
 	{}
 
+	void GetAffectedTypes(TMap<const UScriptStruct*, TArray<FMassEntityHandle>>& OutAdded, TMap<const UScriptStruct*, TArray<FMassEntityHandle>>& OutRemoved) { ensure(false); }
 	virtual void Execute(UMassEntitySubsystem& System) const PURE_VIRTUAL(FCommandBufferEntryBase::Execute, );
 };
 template<> struct TStructOpsTypeTraits<FCommandBufferEntryBase> : public TStructOpsTypeTraitsBase2<FCommandBufferEntryBase> { enum { WithPureVirtual = true, }; };
@@ -39,6 +48,11 @@ struct MASSENTITY_API FDeferredCommand : public FCommandBufferEntryBase
 {
 	GENERATED_BODY()
 
+	enum 
+	{
+		Type = ECommandBufferOperationType::None
+	};
+
 	FDeferredCommand() = default;
 
 	FDeferredCommand(const TFunction<void(UMassEntitySubsystem& System)>& InDeferredFunction)
@@ -62,12 +76,22 @@ struct MASSENTITY_API FBuildEntityFromFragmentInstance : public FCommandBufferEn
 {
 	GENERATED_BODY()
 
+	enum 
+	{
+		Type = ECommandBufferOperationType::Add
+	};
+
 	FBuildEntityFromFragmentInstance() = default;
 	FBuildEntityFromFragmentInstance(const FMassEntityHandle Entity, FStructView InStruct)
 		: FCommandBufferEntryBase(Entity)
 		, Struct(InStruct)
 	{}
 
+	void GetAffectedTypes(TMap<const UScriptStruct*, TArray<FMassEntityHandle>>& OutAdded, TMap<const UScriptStruct*, TArray<FMassEntityHandle>>& OutRemoved)
+	{
+		OutAdded.FindOrAdd(Struct.GetScriptStruct()).AddUnique(TargetEntity);
+	}
+
 protected:
 	virtual void Execute(UMassEntitySubsystem& System) const override;
 
@@ -82,12 +106,25 @@ struct MASSENTITY_API FBuildEntityFromFragmentInstances : public FCommandBufferE
 {
 	GENERATED_BODY()
 
+	enum
+	{
+		Type = ECommandBufferOperationType::Add
+	};
+
 	FBuildEntityFromFragmentInstances() = default;
 	FBuildEntityFromFragmentInstances(const FMassEntityHandle Entity, TConstArrayView<FInstancedStruct> InInstances)
 		: FCommandBufferEntryBase(Entity)
 		, Instances(InInstances)
 	{}
 
+	void GetAffectedTypes(TMap<const UScriptStruct*, TArray<FMassEntityHandle>>& OutAdded, TMap<const UScriptStruct*, TArray<FMassEntityHandle>>& OutRemoved)
+	{
+		for (const FInstancedStruct& Struct : Instances)
+		{
+			OutAdded.FindOrAdd(Struct.GetScriptStruct()).AddUnique(TargetEntity);
+		}
+	}
+
 protected:
 	virtual void Execute(UMassEntitySubsystem& System) const override;
 
@@ -102,12 +139,22 @@ struct MASSENTITY_API FCommandAddFragment : public FCommandBufferEntryBase
 {
 	GENERATED_BODY()
 
+	enum
+	{
+		Type = ECommandBufferOperationType::Add
+	};
+
 	FCommandAddFragment() = default;
 	FCommandAddFragment(FMassEntityHandle InEntity, UScriptStruct* InStruct)
 		: FCommandBufferEntryBase(InEntity)
 		, StructParam(InStruct)
 	{}
 
+	void GetAffectedTypes(TMap<const UScriptStruct*, TArray<FMassEntityHandle>>& OutAdded, TMap<const UScriptStruct*, TArray<FMassEntityHandle>>& OutRemoved)
+	{
+		OutAdded.FindOrAdd(StructParam).AddUnique(TargetEntity);
+	}
+
 protected:
 	virtual void Execute(UMassEntitySubsystem& System) const override;
 
@@ -122,12 +169,22 @@ struct MASSENTITY_API FCommandAddFragmentInstance : public FCommandBufferEntryBa
 {
 	GENERATED_BODY()
 
+	enum
+	{
+		Type = ECommandBufferOperationType::Add
+	};
+
 	FCommandAddFragmentInstance() = default;
 	FCommandAddFragmentInstance(const FMassEntityHandle InEntity, FStructView InStruct)
         : FCommandBufferEntryBase(InEntity)
         , Struct(InStruct)
 	{}
 
+	void GetAffectedTypes(TMap<const UScriptStruct*, TArray<FMassEntityHandle>>& OutAdded, TMap<const UScriptStruct*, TArray<FMassEntityHandle>>& OutRemoved)
+	{
+		OutAdded.FindOrAdd(Struct.GetScriptStruct()).AddUnique(TargetEntity);
+	}
+
 protected:
 	virtual void Execute(UMassEntitySubsystem& System) const override;
 
@@ -142,12 +199,22 @@ struct MASSENTITY_API FCommandRemoveFragment : public FCommandBufferEntryBase
 {
 	GENERATED_BODY()
 
+	enum
+	{
+		Type = ECommandBufferOperationType::Remove
+	};
+	
 	FCommandRemoveFragment() = default;
 	FCommandRemoveFragment(FMassEntityHandle InEntity, UScriptStruct* InStruct)
 		: FCommandBufferEntryBase(InEntity)
 		, StructParam(InStruct)
 	{}
 
+	void GetAffectedTypes(TMap<const UScriptStruct*, TArray<FMassEntityHandle>>& OutAdded, TMap<const UScriptStruct*, TArray<FMassEntityHandle>>& OutRemoved)
+	{
+		OutRemoved.FindOrAdd(StructParam).AddUnique(TargetEntity);
+	}
+
 protected:
 	virtual void Execute(UMassEntitySubsystem& System) const override;
 
@@ -162,12 +229,25 @@ struct MASSENTITY_API FCommandAddFragmentList : public FCommandBufferEntryBase
 {
 	GENERATED_BODY()
 
+	enum
+	{
+		Type = ECommandBufferOperationType::Add
+	};
+
 	FCommandAddFragmentList() = default;
 	FCommandAddFragmentList(FMassEntityHandle InEntity, TConstArrayView<const UScriptStruct*> InFragmentList)
 		: FCommandBufferEntryBase(InEntity)
 		, FragmentList(InFragmentList)
 	{}
 
+	void GetAffectedTypes(TMap<const UScriptStruct*, TArray<FMassEntityHandle>>& OutAdded, TMap<const UScriptStruct*, TArray<FMassEntityHandle>>& OutRemoved)
+	{
+		for (const UScriptStruct* StructParam : FragmentList)
+		{
+			OutAdded.FindOrAdd(StructParam).AddUnique(TargetEntity);
+		}
+	}
+
 protected:
 	virtual void Execute(UMassEntitySubsystem& System) const override;
 
@@ -182,12 +262,25 @@ struct MASSENTITY_API FCommandRemoveFragmentList : public FCommandBufferEntryBas
 {
 	GENERATED_BODY()
 
+	enum
+	{
+		Type = ECommandBufferOperationType::Remove
+	};
+
 	FCommandRemoveFragmentList() = default;
 	FCommandRemoveFragmentList(FMassEntityHandle InEntity, TConstArrayView<const UScriptStruct*> InFragmentList)
 		: FCommandBufferEntryBase(InEntity)
 		, FragmentList(InFragmentList)
 	{}
 
+	void GetAffectedTypes(TMap<const UScriptStruct*, TArray<FMassEntityHandle>>& OutAdded, TMap<const UScriptStruct*, TArray<FMassEntityHandle>>& OutRemoved)
+	{
+		for (const UScriptStruct* StructParam : FragmentList)
+		{
+			OutRemoved.FindOrAdd(StructParam).AddUnique(TargetEntity);
+		}
+	}
+
 protected:
 	virtual void Execute(UMassEntitySubsystem& System) const override;
 
@@ -202,6 +295,11 @@ struct MASSENTITY_API FCommandAddTag: public FCommandBufferEntryBase
 {
 	GENERATED_BODY()
 
+	enum
+	{
+		Type = ECommandBufferOperationType::None
+	};
+
 	FCommandAddTag() = default;
 	FCommandAddTag(FMassEntityHandle InEntity, UScriptStruct* InStruct)
 		: FCommandBufferEntryBase(InEntity)
@@ -222,6 +320,11 @@ struct MASSENTITY_API FCommandRemoveTag : public FCommandBufferEntryBase
 {
 	GENERATED_BODY()
 
+	enum
+	{
+		Type = ECommandBufferOperationType::None
+	};
+
 	FCommandRemoveTag() = default;
 	FCommandRemoveTag(FMassEntityHandle InEntity, UScriptStruct* InStruct)
 		: FCommandBufferEntryBase(InEntity)
@@ -243,8 +346,12 @@ struct MASSENTITY_API FCommandSwapTags : public FCommandBufferEntryBase
 {
 	GENERATED_BODY()
 
-	FCommandSwapTags() = default;
+	enum
+	{
+		Type = ECommandBufferOperationType::None //ECommandBufferOperationType::Add | ECommandBufferOperationType::Remove
+	};
 
+	FCommandSwapTags() = default;
 	FCommandSwapTags(const FMassEntityHandle InEntity, const UScriptStruct* InOldTagType, const UScriptStruct* InNewTagType)
 		: FCommandBufferEntryBase(InEntity)
 		, OldTagType(InOldTagType)
@@ -269,12 +376,32 @@ struct MASSENTITY_API FCommandRemoveComposition : public FCommandBufferEntryBase
 {
 	GENERATED_BODY()
 
+	enum
+	{
+		Type = ECommandBufferOperationType::Remove
+	};
+
 	FCommandRemoveComposition() = default;
 	FCommandRemoveComposition(FMassEntityHandle InEntity, const FMassArchetypeCompositionDescriptor& InDescriptor)
 		: FCommandBufferEntryBase(InEntity)
 		, Descriptor(InDescriptor)
 	{}
 
+	void GetAffectedTypes(TMap<const UScriptStruct*, TArray<FMassEntityHandle>>& OutAdded, TMap<const UScriptStruct*, TArray<FMassEntityHandle>>& OutRemoved)
+	{
+		if (Descriptor.Fragments.IsEmpty() == false)
+		{
+			// @todo this is way too slow and I hate it. I'm considering adding abother flag to EOperationType to indicate
+			// a composition-based operation and call a different function or process those at flush time. Thoughts?
+			TArray<const UScriptStruct*> Fragments;
+			Descriptor.Fragments.ExportTypes(Fragments);
+			for (const UScriptStruct* StructParam : Fragments)
+			{
+				OutRemoved.FindOrAdd(StructParam).AddUnique(TargetEntity);
+			}
+		}
+	}
+
 protected:
 	virtual void Execute(UMassEntitySubsystem& System) const override;
 
@@ -284,6 +411,7 @@ protected:
 struct MASSENTITY_API FMassCommandBuffer
 {
 public:
+	FMassCommandBuffer() = default;
 
 	/** Push any command, requires to derive from FCommandBufferEntryBase */
 	template< typename T, typename = typename TEnableIf<TIsDerivedFrom<typename TRemoveReference<T>::Type, FCommandBufferEntryBase>::IsDerived, void>::Type >
@@ -296,8 +424,7 @@ public:
 	template< typename T, typename = typename TEnableIf<TIsDerivedFrom<typename TRemoveReference<T>::Type, FCommandBufferEntryBase>::IsDerived, void>::Type, typename... TArgs >
 	void EmplaceCommand(TArgs&&... InArgs)
 	{
-		UE_MT_SCOPED_WRITE_ACCESS(PendingCommandsDetector);
-		return PendingCommands.Emplace<T>(Forward<TArgs>(InArgs)...);
+		EmplaceCommand_GetRef<T>(Forward<TArgs>(InArgs)...);
 	}
 
 	/** Emplace any command and return its ref, requires to derive from FCommandBufferEntryBase */
@@ -305,9 +432,15 @@ public:
 	T& EmplaceCommand_GetRef(TArgs&&... InArgs)
 	{
 		UE_MT_SCOPED_WRITE_ACCESS(PendingCommandsDetector);
-		return PendingCommands.Emplace_GetRef<T>(Forward<TArgs>(InArgs)...);
+		T& Command = PendingCommands.Emplace_GetRef<T>(Forward<TArgs>(InArgs)...); 
+		if (constexpr bool bIsModifyingComposition = ((T::Type & (ECommandBufferOperationType::Add | ECommandBufferOperationType::Remove)) != 0))
+		{	
+			Command.GetAffectedTypes(FragmentsToAdd, FragmentsToRemove);
+		}
+		return Command;
 	}
-
+	
+public:
 	template<typename T>
 	void AddFragment(FMassEntityHandle Entity)
 	{
@@ -357,7 +490,7 @@ public:
 	 */
 	void MoveAppend(FMassCommandBuffer& InOutOther);
 
-	int32 GetPendingCommandsCount() const { return PendingCommands.Num(); }
+	bool HasPendingCommands() const { return PendingCommands.Num() > 0 || EntitiesToDestroy.Num() > 0; }
 
 private:
 	FInstancedStructStream PendingCommands;
@@ -365,4 +498,7 @@ private:
 	FCriticalSection AppendingCommandsCS;
 
 	TArray<FMassEntityHandle> EntitiesToDestroy;
+
+	TMap<const UScriptStruct*, TArray<FMassEntityHandle>> FragmentsToAdd;
+	TMap<const UScriptStruct*, TArray<FMassEntityHandle>> FragmentsToRemove;
 };
diff --git a/Engine/Plugins/Runtime/MassEntity/Source/MassEntity/Public/MassEntitySubsystem.h b/Engine/Plugins/Runtime/MassEntity/Source/MassEntity/Public/MassEntitySubsystem.h
--- a/Engine/Plugins/Runtime/MassEntity/Source/MassEntity/Public/MassEntitySubsystem.h
+++ b/Engine/Plugins/Runtime/MassEntity/Source/MassEntity/Public/MassEntitySubsystem.h
@@ -8,6 +8,7 @@
 #include "InstancedStruct.h"
 #include "MassEntityQuery.h"
 #include "StructUtilsTypes.h"
+#include "MassObserverManager.h"
 #include "MassEntitySubsystem.generated.h"
 
 
@@ -74,6 +75,7 @@ public:
 
 	//~USubsystem interface
 	virtual void Initialize(FSubsystemCollectionBase& Collection) override;
+	virtual void PostInitialize() override;
 	virtual void Deinitialize() override;
 	//~End of USubsystem interface
 	
@@ -126,12 +128,30 @@ public:
 	 * @return FMassEntityHandle id of the newly created entity */
 	FMassEntityHandle CreateEntity(TConstArrayView<FInstancedStruct> FragmentInstanceList);
 
+	/**
+	 * A dedicated structure for ensuring the "on entities creation" observers get notified only once all other 
+	 * initialization operations are done and this creation context instance gets released. */
+	struct FEntityCreationContext
+	{
+		explicit FEntityCreationContext(const int32 InNumSpawned = 0)
+			: NumberSpawned(InNumSpawned)
+		{}
+		~FEntityCreationContext() { if (OnSpawningFinished) OnSpawningFinished(*this); }
+
+		const FArchetypeChunkCollection& GetChunkCollection() const { return ChunkCollection; }
+	private:
+		friend UMassEntitySubsystem;
+		int32 NumberSpawned;
+		FArchetypeChunkCollection ChunkCollection;
+		TFunction<void(FEntityCreationContext&)> OnSpawningFinished;
+	};
+
 	/** A version of CreateEntity that's creating a number of entities (Count) at one go
 	 *  @param Archetype you want this entity to be
 	 *  @param Count number of entities to create
 	 *  @param OutEntities the newly created entities are appended to given array, i.e. the pre-existing content of OutEntities won't be affected by the call
-	 *  @return number of entities created */
-	int32 BatchCreateEntities(const FArchetypeHandle Archetype, const int32 Count, TArray<FMassEntityHandle>& OutEntities);
+	 *  @return a creation context that will notify all the interested observers about newly created fragments once the context is released */
+	TSharedRef<FEntityCreationContext> BatchCreateEntities(const FArchetypeHandle Archetype, const int32 Count, TArray<FMassEntityHandle>& OutEntities);
 
 	/**
 	 * Destroys a fully built entity, use ReleaseReservedEntity if entity was not yet built.
@@ -193,6 +213,8 @@ public:
 	void AddCompositionToEntity_GetDelta(FMassEntityHandle Entity, FMassArchetypeCompositionDescriptor& InOutDescriptor);
 	void RemoveCompositionFromEntity(FMassEntityHandle Entity, const FMassArchetypeCompositionDescriptor& InDescriptor);
 
+	const FMassArchetypeCompositionDescriptor& GetArchetypeComposition(const FArchetypeHandle& ArchetypeHandle) const;
+
 	/** 
 	 * Moves an entity over to a new archetype by copying over fragments common to both archetypes
 	 * @param Entity is the entity to move 
@@ -297,6 +319,8 @@ public:
 		}
 	}
 
+	FMassObserverManager& GetObserverManager() { return ObserverManager; }
+
 #if WITH_MASSENTITY_DEBUG
 	void DebugPrintEntity(int32 Index, FOutputDevice& Ar, const TCHAR* InPrefix = TEXT("")) const;
 	void DebugPrintEntity(FMassEntityHandle Entity, FOutputDevice& Ar, const TCHAR* InPrefix = TEXT("")) const;
@@ -305,6 +329,7 @@ public:
 	void DebugGetArchetypesStringDetails(FOutputDevice& Ar, const bool bIncludeEmpty = true);
 	void DebugGetArchetypeFragmentTypes(const FArchetypeHandle& Archetype, TArray<const UScriptStruct*>& InOutFragmentList) const;
 	int32 DebugGetArchetypeEntitiesCount(const FArchetypeHandle& Archetype) const;	
+	int32 DebugGetArchetypeEntitiesCountPerChunk(const FArchetypeHandle& Archetype) const;
 	int32 DebugGetEntityCount() const { return Entities.Num() - NumReservedEntities - EntityFreeIndexList.Num(); }
 	int32 DebugGetArchetypesCount() const { return FragmentHashToArchetypeMap.Num(); }
 	void DebugRemoveAllEntities();
@@ -362,6 +387,9 @@ private:
 	TArray<FSharedStruct> SharedFragments;
 	// Hash/Index in array pair
 	TMap<uint32, int32> SharedFragmentsMap;
+
+	UPROPERTY(Transient)
+	FMassObserverManager ObserverManager;
 };
 
 
diff --git a/Engine/Plugins/Runtime/MassEntity/Source/MassEntity/Public/MassObserverManager.h b/Engine/Plugins/Runtime/MassEntity/Source/MassEntity/Public/MassObserverManager.h
new file mode 100644
--- /dev/null
+++ b/Engine/Plugins/Runtime/MassEntity/Source/MassEntity/Public/MassObserverManager.h
@@ -0,0 +1,89 @@
+// Copyright Epic Games, Inc. All Rights Reserved.
+
+#pragma once
+
+#include "MassEntityTypes.h"
+#include "MassObserverManager.generated.h"
+
+
+class UMassEntitySubsystem;
+struct FArchetypeChunkCollection;
+
+/** 
+ * A type that encapsulates logic related to notifying interested parties of entity composition changes. Upon creation it
+ * reads information from UMassObserverRegistry and instantiates processors interested in handling given fragment
+ * type addition or removal.
+ */
+USTRUCT()
+struct MASSENTITY_API FMassObserverManager
+{
+	GENERATED_BODY()
+
+public:
+	FMassObserverManager();	
+
+	DECLARE_MULTICAST_DELEGATE_TwoParams(FOnObservedFragmentTypesChanged, const FMassFragmentBitSet& /*NewObservedAddSet*/, const FMassFragmentBitSet& /*NewObservedRemoveSet*/);
+
+	FOnObservedFragmentTypesChanged& GetOnObservedFragmentTypesChangedDelegate() { return OnObservedFragmentTypesChangedDelegate; }
+
+	const FMassFragmentBitSet& GetObservedAddFragmentsBitSet() const { return ObservedAddFragments; }
+	const FMassFragmentBitSet& GetObservedRemoveFragmentsBitSet() const { return ObservedRemoveFragments; }
+
+	bool HasOnAddedObserversForFragments(const FMassFragmentBitSet& InQueriedBitSet) const { return ObservedAddFragments.HasAny(InQueriedBitSet); }
+	bool HasOnRemovedObserversForFragments(const FMassFragmentBitSet& InQueriedBitSet) const { return ObservedRemoveFragments.HasAny(InQueriedBitSet); }
+
+	bool OnPostEntitiesCreated(const FArchetypeChunkCollection& ChunkCollection);
+	bool OnPostEntitiesCreated(FMassProcessingContext& ProcessingContext, const FArchetypeChunkCollection& ChunkCollection);
+
+	bool OnPreEntitiesDestroyed(const FArchetypeChunkCollection& ChunkCollection);
+	bool OnPreEntitiesDestroyed(FMassProcessingContext& ProcessingContext, const FArchetypeChunkCollection& ChunkCollection);
+
+	bool OnPostCompositionAdded(const FMassEntityHandle Entity, const FMassArchetypeCompositionDescriptor& Composition);
+	bool OnPreCompositionRemoved(const FMassEntityHandle Entity, const FMassArchetypeCompositionDescriptor& Composition);
+
+	void OnPostFragmentAdded(const UScriptStruct& FragmentType, const FArchetypeChunkCollection& ChunkCollection)
+	{
+		HandleSingleFragmentImpl(FragmentType, ChunkCollection, ObservedAddFragments, OnFragmentAddedObservers);
+	}
+
+	void OnPreFragmentRemoved(const UScriptStruct& FragmentType, const FArchetypeChunkCollection& ChunkCollection)
+	{
+		HandleSingleFragmentImpl(FragmentType, ChunkCollection, ObservedRemoveFragments, OnFragmentRemovedObservers);
+	}
+
+protected:
+	friend UMassEntitySubsystem;
+	explicit FMassObserverManager(UMassEntitySubsystem& Owner);
+
+	void Initialize();
+	void HandleFragmentsImpl(FMassProcessingContext& ProcessingContext, const FArchetypeChunkCollection& ChunkCollection
+		, const FMassFragmentBitSet& FragmentsBitSet, TMap<const UScriptStruct*, FMassRuntimePipeline>& HandlersContainer);
+	void HandleSingleFragmentImpl(const UScriptStruct& FragmentType, const FArchetypeChunkCollection& ChunkCollection
+		, const FMassFragmentBitSet& FragmentFilterBitSet, TMap<const UScriptStruct*, FMassRuntimePipeline>& HandlersContainer);
+
+	FOnObservedFragmentTypesChanged OnObservedFragmentTypesChangedDelegate;
+
+	FMassFragmentBitSet ObservedAddFragments;
+	FMassFragmentBitSet ObservedRemoveFragments;
+
+	UPROPERTY()
+	TMap<const UScriptStruct*, FMassRuntimePipeline> OnFragmentAddedObservers;
+
+	UPROPERTY()
+	TMap<const UScriptStruct*, FMassRuntimePipeline> OnFragmentRemovedObservers;
+
+	/** 
+	 * The owning EntitySubsystem. No need for it to be a UPROPERTY since by design we don't support creation of 
+	 * FMassObserverManager outside of an UMassEntitySubsystem instance 
+	 */
+	UMassEntitySubsystem& EntitySubsystem;
+};
+
+template<>
+struct TStructOpsTypeTraits<FMassObserverManager> : public TStructOpsTypeTraitsBase2<FMassObserverManager>
+{
+	enum
+	{
+		WithCopy = false,
+	};
+};
\ No newline at end of file
diff --git a/Engine/Plugins/Runtime/MassEntity/Source/MassEntity/Public/MassObserverProcessor.h b/Engine/Plugins/Runtime/MassEntity/Source/MassEntity/Public/MassObserverProcessor.h
new file mode 100644
--- /dev/null
+++ b/Engine/Plugins/Runtime/MassEntity/Source/MassEntity/Public/MassObserverProcessor.h
@@ -0,0 +1,51 @@
+// Copyright Epic Games, Inc. All Rights Reserved.
+
+#pragma once
+
+#include "MassProcessor.h"
+#include "MassObserverProcessor.generated.h"
+
+
+class UMassEntitySubsystem;
+
+UCLASS(abstract)
+class MASSENTITY_API UMassObserverProcessor : public UMassProcessor
+{
+	GENERATED_BODY()
+
+public:
+	UMassObserverProcessor();
+
+protected:
+	virtual void PostInitProperties() override;
+
+	/** Override to register this processor with the observer registry */
+	virtual void Register() PURE_VIRTUAL(UMassObserverProcessor::Register, );
+
+protected:
+	/** Set in class' constructor and determines for which Fragment type this given UMassObserverProcessor will be used 
+	 *  as the default initializer/deinitializer (as returned by the MassInitializersRegistry). If set to null will
+	 *  require the user to manually add this UMassObserverProcessor as initializer/deinitializer */
+	UPROPERTY()
+	UScriptStruct* FragmentType;
+};
+
+UCLASS(abstract)
+class MASSENTITY_API UMassFragmentInitializer : public UMassObserverProcessor
+{
+	GENERATED_BODY()
+
+protected:
+	/** Adds this FragmentInitializer to the MassInitializersRegistry to be used to initialize fragments of FragmentType */
+	virtual void Register() override;
+};
+
+UCLASS(abstract)
+class MASSENTITY_API UMassFragmentDestructor : public UMassObserverProcessor
+{
+	GENERATED_BODY()
+
+protected:
+	/** Adds this FragmentDestructor to the MassInitializersRegistry to be used to deinitialize fragments of FragmentType */
+	virtual void Register() override;
+};
diff --git a/Engine/Plugins/Runtime/MassEntity/Source/MassEntity/Public/MassObserverRegistry.h b/Engine/Plugins/Runtime/MassEntity/Source/MassEntity/Public/MassObserverRegistry.h
new file mode 100644
--- /dev/null
+++ b/Engine/Plugins/Runtime/MassEntity/Source/MassEntity/Public/MassObserverRegistry.h
@@ -0,0 +1,63 @@
+// Copyright Epic Games, Inc. All Rights Reserved.
+
+#pragma once
+
+#include "MassEntityTypes.h"
+#include "MassProcessor.h"
+#include "MassObserverRegistry.generated.h"
+
+
+struct FMassObserverManager;
+
+USTRUCT()
+struct FMassProcessorClassCollection
+{
+	GENERATED_BODY()
+
+	UPROPERTY()
+	TArray<TSubclassOf<UMassProcessor>> ClassCollection;
+};
+
+
+UCLASS()
+class MASSENTITY_API UMassObserverRegistry : public UObject
+{
+	GENERATED_BODY()
+
+public:
+	UMassObserverRegistry();
+
+	static UMassObserverRegistry& GetMutable() { return *GetMutableDefault<UMassObserverRegistry>(); }
+	static const UMassObserverRegistry& Get() { return *GetDefault<UMassObserverRegistry>(); }
+
+	template<typename TFragment>
+	static void RegisterInitializer(const UMassProcessor& CDO)
+	{
+		ensureMsgf(CDO.HasAnyFlags(RF_ClassDefaultObject) && CDO.GetClass()->HasAnyClassFlags(CLASS_Abstract) == false
+			, TEXT("Only CDOs are expected to be registered with UMassObserverRegistry::RegisterInitializer"));
+		GetMutable().RegisterFragmentInitializer(*TFragment::StaticStruct(), CDO.GetClass());
+	}
+
+	template<typename TFragment>
+	static void RegisterDeinitializer(const UMassProcessor& CDO)
+	{
+		ensureMsgf(CDO.HasAnyFlags(RF_ClassDefaultObject) && CDO.GetClass()->HasAnyClassFlags(CLASS_Abstract) == false
+			, TEXT("Only CDOs are expected to be registered with UMassObserverRegistry::RegisterDeinitializer"));
+		GetMutable().RegisterFragmentDeinitializer(*TFragment::StaticStruct(), CDO.GetClass());
+	}
+
+	void RegisterFragmentInitializer(const UScriptStruct& FragmentType, TSubclassOf<UMassProcessor> FragmentInitializerClass);
+	void RegisterFragmentDeinitializer(const UScriptStruct& FragmentType, TSubclassOf<UMassProcessor> FragmentDeinitializerClass);
+
+protected:
+	friend FMassObserverManager;
+
+	TConstArrayView<TSubclassOf<UMassProcessor>> GetFragmentInitializers(const UScriptStruct& FragmentType) const;
+	TConstArrayView<TSubclassOf<UMassProcessor>> GetFragmentDeinitializers(const UScriptStruct& FragmentType) const;
+
+	UPROPERTY()
+	TMap<const UScriptStruct*, FMassProcessorClassCollection> FragmentInitializersMap;
+
+	UPROPERTY()
+	TMap<const UScriptStruct*, FMassProcessorClassCollection> FragmentDeinitializersMap;
+};
diff --git a/Engine/Plugins/Runtime/MassGameplay/Source/MassGameplayDebug/Private/MassDebugVisualizationTrait.cpp b/Engine/Plugins/Runtime/MassGameplay/Source/MassGameplayDebug/Private/MassDebugVisualizationTrait.cpp
--- a/Engine/Plugins/Runtime/MassGameplay/Source/MassGameplayDebug/Private/MassDebugVisualizationTrait.cpp
+++ b/Engine/Plugins/Runtime/MassGameplay/Source/MassGameplayDebug/Private/MassDebugVisualizationTrait.cpp
@@ -18,7 +18,7 @@ void UMassDebugVisualizationTrait::BuildTemplate(FMassEntityTemplateBuildContext
 	if (DebugMesh)
 	{
 #if WITH_EDITORONLY_DATA
-		FSimDebugVisComponent& DebugVisFragment = BuildContext.AddFragmentWithDefaultInitializer_GetRef<FSimDebugVisComponent>();
+		FSimDebugVisComponent& DebugVisFragment = BuildContext.AddFragment_GetRef<FSimDebugVisComponent>();
 		UMassDebuggerSubsystem* Debugger = World.GetSubsystem<UMassDebuggerSubsystem>();
 		if (ensure(Debugger))
 		{
@@ -29,22 +29,20 @@ void UMassDebugVisualizationTrait::BuildTemplate(FMassEntityTemplateBuildContext
 			}
 			// @todo this path requires a fragment destructor that will remove the mesh from the debugger.
 		}
-
-		BuildContext.AddDefaultInitializer<FSimDebugVisComponent>();
 #endif // WITH_EDITORONLY_DATA
 	}
 	// add fragments needed whenever we have debugging capabilities
 #if !(UE_BUILD_SHIPPING || UE_BUILD_TEST)
 	BuildContext.AddTag<FMassDebuggableTag>();
 #if WITH_EDITORONLY_DATA
-	BuildContext.AddFragmentWithDefaultInitializer_GetRef<FDataFragment_DebugVis>().Shape = DebugShape.WireShape;
+	BuildContext.AddFragment_GetRef<FDataFragment_DebugVis>().Shape = DebugShape.WireShape;
 #else
 	// DebugShape unavailable, will used default instead
-	BuildContext.AddFragmentWithDefaultInitializer<FDataFragment_DebugVis>();
+	BuildContext.AddFragment<FDataFragment_DebugVis>();
 #endif // WITH_EDITORONLY_DATA
-	BuildContext.AddFragmentWithDefaultInitializer<FDataFragment_AgentRadius>();
+	BuildContext.AddFragment<FDataFragment_AgentRadius>();
 
-	BuildContext.AddFragmentWithDefaultInitializer<FDataFragment_Transform>();
+	BuildContext.AddFragment<FDataFragment_Transform>();
 #endif // if !(UE_BUILD_SHIPPING || UE_BUILD_TEST)
 
 }
\ No newline at end of file
diff --git a/Engine/Plugins/Runtime/MassGameplay/Source/MassLOD/Private/MassLODTrait.cpp b/Engine/Plugins/Runtime/MassGameplay/Source/MassLOD/Private/MassLODTrait.cpp
--- a/Engine/Plugins/Runtime/MassGameplay/Source/MassLOD/Private/MassLODTrait.cpp
+++ b/Engine/Plugins/Runtime/MassGameplay/Source/MassLOD/Private/MassLODTrait.cpp
@@ -38,7 +38,7 @@ void UMassSimulationLODTrait::BuildTemplate(FMassEntityTemplateBuildContext& Bui
 	// Variable ticking from simulation LOD
 	if(bEnableVariableTicking)
 	{
-		BuildContext.AddFragmentWithDefaultInitializer<FMassSimulationVariableTickFragment>();
+		BuildContext.AddFragment<FMassSimulationVariableTickFragment>();
 		BuildContext.AddChunkFragment<FMassSimulationVariableTickChunkFragment>();
 
 		uint32 VariableTickConfigHash = UE::StructUtils::GetStructCrc32(FConstStructView::Make(VariableTickConfig));
@@ -47,4 +47,4 @@ void UMassSimulationLODTrait::BuildTemplate(FMassEntityTemplateBuildContext& Bui
 		FSharedStruct VariableTickSharedFragment = EntitySubsystem->GetOrCreateSharedFragment<FMassSimulationVariableTickSharedFragment>(VariableTickConfigHash, VariableTickConfig);
 		BuildContext.AddSharedFragment(VariableTickSharedFragment);
 	}
-}
\ No newline at end of file
+}
diff --git a/Engine/Plugins/Runtime/MassGameplay/Source/MassMovement/Private/MassVelocityRandomizerTrait.cpp b/Engine/Plugins/Runtime/MassGameplay/Source/MassMovement/Private/MassVelocityRandomizerTrait.cpp
--- a/Engine/Plugins/Runtime/MassGameplay/Source/MassMovement/Private/MassVelocityRandomizerTrait.cpp
+++ b/Engine/Plugins/Runtime/MassGameplay/Source/MassMovement/Private/MassVelocityRandomizerTrait.cpp
@@ -11,12 +11,12 @@
 //----------------------------------------------------------------------//
 void UMassVelocityRandomizerTrait::BuildTemplate(FMassEntityTemplateBuildContext& BuildContext, UWorld& World) const
 {
-	BuildContext.AddFragment<FMassVelocityFragment>();
-	
-	UMassRandomVelocityInitializer* VelocityInitializer = NewObject<UMassRandomVelocityInitializer>(&World);
-	check(VelocityInitializer);
-	VelocityInitializer->SetParameters(MinSpeed, MaxSpeed, bSetZComponent);
-	BuildContext.AddInitializer(*VelocityInitializer);
+	FMassVelocityFragment& VelocityTemplate = BuildContext.AddFragment_GetRef<FMassVelocityFragment>();
+	// This is a small @hack to support sending parameters to initializer.
+	// A proper solution will allow users to specify a "lambda" initializer that will be used during creation
+	VelocityTemplate.Value.X = MinSpeed;
+	VelocityTemplate.Value.Y = MaxSpeed;
+	VelocityTemplate.Value.Z = bSetZComponent ? 1.f : 0.f;
 }
 
 //----------------------------------------------------------------------//
@@ -27,13 +27,6 @@ UMassRandomVelocityInitializer::UMassRandomVelocityInitializer()
 	FragmentType = FMassVelocityFragment::StaticStruct();
 }
 
-void UMassRandomVelocityInitializer::SetParameters(const float InMinSpeed, const float InMaxSpeed, const bool bInSetZComponent)
-{
-	MinSpeed = InMinSpeed;
-	MaxSpeed = InMaxSpeed;
-	bSetZComponent = bInSetZComponent;
-}
-
 void UMassRandomVelocityInitializer::ConfigureQueries()
 {
 	EntityQuery.AddRequirement<FMassVelocityFragment>(EMassFragmentAccess::ReadWrite);
@@ -42,30 +35,18 @@ void UMassRandomVelocityInitializer::ConfigureQueries()
 void UMassRandomVelocityInitializer::Execute(UMassEntitySubsystem& EntitySubsystem, FMassExecutionContext& Context)
 {
 	// note: the author is aware that the vectors produced below are not distributed uniformly, but it's good enough
-	if (bSetZComponent)
-	{
-		EntityQuery.ForEachEntityChunk(EntitySubsystem, Context, ([this](FMassExecutionContext& Context)
+	EntityQuery.ForEachEntityChunk(EntitySubsystem, Context, ([this](FMassExecutionContext& Context)
+		{
+			const TArrayView<FMassVelocityFragment> VelocitiesList = Context.GetMutableFragmentView<FMassVelocityFragment>();
+			for (FMassVelocityFragment& VelocityFragment : VelocitiesList)
 			{
-				const TArrayView<FMassVelocityFragment> VelocitiesList = Context.GetMutableFragmentView<FMassVelocityFragment>();
-				for (FMassVelocityFragment& VelocityFragment : VelocitiesList)
-				{
-					const FVector RandomVector = FVector(FMath::FRandRange(-1.f, 1.f), FMath::FRandRange(-1.f, 1.f), FMath::FRandRange(-1.f, 1.f)).GetSafeNormal();
-
-					VelocityFragment.Value = RandomVector * FMath::FRandRange(MinSpeed, MaxSpeed);
-				}
-			}));
-	}
-	else
-	{
-		EntityQuery.ForEachEntityChunk(EntitySubsystem, Context, ([this](FMassExecutionContext& Context)
-			{
-				const TArrayView<FMassVelocityFragment> VelocitiesList = Context.GetMutableFragmentView<FMassVelocityFragment>();
-				for (FMassVelocityFragment& VelocityFragment : VelocitiesList)
-				{
-					const FVector RandomVector = FVector(FMath::FRandRange(-1.f, 1.f), FMath::FRandRange(-1.f, 1.f), 0).GetSafeNormal2D();
-
-					VelocityFragment.Value = RandomVector * FMath::FRandRange(MinSpeed, MaxSpeed);
-				}
-			}));
-	}	
+				// the given VelocityFragment's value is encoding the initialization parameters, as per comment in
+				// UMassVelocityRandomizerTrait::BuildTemplate
+				const FVector RandomVector = VelocityFragment.Value.Z != 0.f 
+					? FVector(FMath::FRandRange(-1.f, 1.f), FMath::FRandRange(-1.f, 1.f), FMath::FRandRange(-1.f, 1.f)).GetSafeNormal()
+					: FVector(FMath::FRandRange(-1.f, 1.f), FMath::FRandRange(-1.f, 1.f), 0).GetSafeNormal();
+
+				VelocityFragment.Value = RandomVector * FMath::FRandRange(VelocityFragment.Value.X, VelocityFragment.Value.Y);
+			}
+		}));
 }
diff --git a/Engine/Plugins/Runtime/MassGameplay/Source/MassReplication/Private/MassReplicationTrait.cpp b/Engine/Plugins/Runtime/MassGameplay/Source/MassReplication/Private/MassReplicationTrait.cpp
--- a/Engine/Plugins/Runtime/MassGameplay/Source/MassReplication/Private/MassReplicationTrait.cpp
+++ b/Engine/Plugins/Runtime/MassGameplay/Source/MassReplication/Private/MassReplicationTrait.cpp
@@ -15,10 +15,10 @@ void UMassReplicationTrait::BuildTemplate(FMassEntityTemplateBuildContext& Build
 		return;
 	}
 
-	FDataFragment_ReplicationTemplateID& TemplateIDFragment = BuildContext.AddFragmentWithDefaultInitializer_GetRef<FDataFragment_ReplicationTemplateID>();
+	FDataFragment_ReplicationTemplateID& TemplateIDFragment = BuildContext.AddFragment_GetRef<FDataFragment_ReplicationTemplateID>();
 	TemplateIDFragment.ID = BuildContext.GetTemplateID();
 
-	BuildContext.AddFragmentWithDefaultInitializer<FMassNetworkIDFragment>();
+	BuildContext.AddFragment<FMassNetworkIDFragment>();
 	BuildContext.AddFragment<FMassReplicatedAgentFragment>();
 	BuildContext.AddFragment<FMassReplicationViewerInfoFragment>();
 	BuildContext.AddFragment<FMassReplicationLODFragment>();
diff --git a/Engine/Plugins/Runtime/MassGameplay/Source/MassRepresentation/Private/MassVisualizationTrait.cpp b/Engine/Plugins/Runtime/MassGameplay/Source/MassRepresentation/Private/MassVisualizationTrait.cpp
--- a/Engine/Plugins/Runtime/MassGameplay/Source/MassRepresentation/Private/MassVisualizationTrait.cpp
+++ b/Engine/Plugins/Runtime/MassGameplay/Source/MassRepresentation/Private/MassVisualizationTrait.cpp
@@ -11,8 +11,6 @@
 UMassVisualizationTrait::UMassVisualizationTrait()
 {
 	RepresentationSubsystemClass = UMassRepresentationSubsystem::StaticClass();
-	// @todo the following line will be cut once new de/initializers are in
-	RepresentationFragmentDeinitializerClass = UMassRepresentationFragmentDestructor::StaticClass();
 	RepresentationDestructorTag = FMassRepresentationDefaultDestructorTag::StaticStruct();
 }
 
@@ -37,8 +35,6 @@ void UMassVisualizationTrait::BuildTemplate(FMassEntityTemplateBuildContext& Bui
 		RepresentationFragment.HighResTemplateActorIndex = HighResTemplateActor.Get() ? RepresentationSubsystem->FindOrAddTemplateActor(HighResTemplateActor.Get()) : INDEX_NONE;
 		RepresentationFragment.LowResTemplateActorIndex = LowResTemplateActor.Get() ? RepresentationSubsystem->FindOrAddTemplateActor(LowResTemplateActor.Get()) : INDEX_NONE;
 
-		// @todo the following line will be cut once new de/initializers are in
-		BuildContext.AddDeinitializer(*(RepresentationFragmentDeinitializerClass->GetDefaultObject<UMassProcessor>()));
 		if (ensureMsgf(RepresentationDestructorTag, TEXT("RepresentationDestructorTag is never expected to be empty")))
 		{
 			BuildContext.AddTag(*RepresentationDestructorTag);
diff --git a/Engine/Plugins/Runtime/MassGameplay/Source/MassSmartObjects/Private/MassSmartObjectUserTrait.cpp b/Engine/Plugins/Runtime/MassGameplay/Source/MassSmartObjects/Private/MassSmartObjectUserTrait.cpp
--- a/Engine/Plugins/Runtime/MassGameplay/Source/MassSmartObjects/Private/MassSmartObjectUserTrait.cpp
+++ b/Engine/Plugins/Runtime/MassGameplay/Source/MassSmartObjects/Private/MassSmartObjectUserTrait.cpp
@@ -7,5 +7,5 @@
 
 void UMassSmartObjectUserTrait::BuildTemplate(FMassEntityTemplateBuildContext& BuildContext, UWorld& World) const
 {
-	BuildContext.AddFragmentWithDefaultInitializer<FMassSmartObjectUserFragment>();
+	BuildContext.AddFragment<FMassSmartObjectUserFragment>();
 }
diff --git a/Engine/Plugins/Runtime/MassGameplay/Source/MassSpawner/Private/MassAssortedFragmentsTrait.cpp b/Engine/Plugins/Runtime/MassGameplay/Source/MassSpawner/Private/MassAssortedFragmentsTrait.cpp
--- a/Engine/Plugins/Runtime/MassGameplay/Source/MassSpawner/Private/MassAssortedFragmentsTrait.cpp
+++ b/Engine/Plugins/Runtime/MassGameplay/Source/MassSpawner/Private/MassAssortedFragmentsTrait.cpp
@@ -14,7 +14,7 @@ void UMassAssortedFragmentsTrait::BuildTemplate(FMassEntityTemplateBuildContext&
 			CA_ASSUME(Type);
 			if (Type->IsChildOf(FMassFragment::StaticStruct()))
 			{
-				BuildContext.AddFragmentWithDefaultInitializer(Fragment);
+				BuildContext.AddFragment(Fragment);
 			}
 			else if (Type->IsChildOf(FMassTag::StaticStruct()))
 			{
diff --git a/Engine/Plugins/Runtime/MassGameplay/Source/MassSpawner/Private/MassEntityConfigAsset.cpp b/Engine/Plugins/Runtime/MassGameplay/Source/MassSpawner/Private/MassEntityConfigAsset.cpp
--- a/Engine/Plugins/Runtime/MassGameplay/Source/MassSpawner/Private/MassEntityConfigAsset.cpp
+++ b/Engine/Plugins/Runtime/MassGameplay/Source/MassSpawner/Private/MassEntityConfigAsset.cpp
@@ -75,7 +75,6 @@ const FMassEntityTemplate* FMassEntityConfig::GetOrCreateEntityTemplate(AActor&
 
 	if (ensureMsgf(!Template.IsEmpty(), TEXT("Need at least one fragment to create an Archetype")))
 	{
-		Template.SetUpProcessors(BuildContext.Handlers, OwnerActor);
 		TemplateRegistry.InitializeEntityTemplate(Template);
 	}
 
diff --git a/Engine/Plugins/Runtime/MassGameplay/Source/MassSpawner/Private/MassEntityTemplate.cpp b/Engine/Plugins/Runtime/MassGameplay/Source/MassSpawner/Private/MassEntityTemplate.cpp
--- a/Engine/Plugins/Runtime/MassGameplay/Source/MassSpawner/Private/MassEntityTemplate.cpp
+++ b/Engine/Plugins/Runtime/MassGameplay/Source/MassSpawner/Private/MassEntityTemplate.cpp
@@ -22,18 +22,6 @@ void FMassEntityTemplate::SetArchetype(const FArchetypeHandle& InArchetype)
 	Archetype = InArchetype;
 }
 
-void FMassEntityTemplate::SetUpProcessors(const FMassObjectHandlers& ObjectHandlers, UObject& PipelineOwner)
-{
-	TArray<const UMassProcessor*> Initializers;
-	Initializers.Reserve(ObjectHandlers.DefaultInitializers.Num() + ObjectHandlers.Initializers.Num());
-	Initializers.Append(ObjectHandlers.DefaultInitializers);
-	Initializers.Append(ObjectHandlers.Initializers);
-
-	// @todo do we want to copy the Processors or just use the CDOs?
-	InitializationPipeline.InitializeFromArray(Initializers, PipelineOwner);
-	DeinitializationPipeline.InitializeFromArray(ObjectHandlers.Deinitializers, PipelineOwner);
-}
-
 FString FMassEntityTemplate::DebugGetDescription(UMassEntitySubsystem* EntitySubsystem) const
 { 
 	FStringOutputDevice Ar;
@@ -51,9 +39,6 @@ FString FMassEntityTemplate::DebugGetDescription(UMassEntitySubsystem* EntitySub
 		Composition.DebugOutputDescription(Ar);
 	}
 
-	Ar += TEXT("\nInitialization Processors:\n");
-	InitializationPipeline.DebugOutputDescription(Ar);
-	
 #endif // WITH_MASSGAMEPLAY_DEBUG
 	return MoveTemp(Ar);
 }
diff --git a/Engine/Plugins/Runtime/MassGameplay/Source/MassSpawner/Private/MassEntityTemplateRegistry.cpp b/Engine/Plugins/Runtime/MassGameplay/Source/MassSpawner/Private/MassEntityTemplateRegistry.cpp
--- a/Engine/Plugins/Runtime/MassGameplay/Source/MassSpawner/Private/MassEntityTemplateRegistry.cpp
+++ b/Engine/Plugins/Runtime/MassGameplay/Source/MassSpawner/Private/MassEntityTemplateRegistry.cpp
@@ -128,7 +128,6 @@ bool UMassEntityTemplateRegistry::BuildTemplateImpl(const FStructToTemplateBuild
 	Builder.Execute(World, StructInstance, Context);
 	if (ensure(!OutTemplate.IsEmpty())) // need at least one fragment to create an Archetype
 	{
-		OutTemplate.SetUpProcessors(Context.Handlers, *this);
 		InitializeEntityTemplate(OutTemplate);
 
 		UE_VLOG(this, LogMassSpawner, Log, TEXT("Created entity template for %s:\n%s"), *GetNameSafe(StructInstance.GetScriptStruct())
@@ -149,7 +148,7 @@ void UMassEntityTemplateRegistry::InitializeEntityTemplate(FMassEntityTemplate&
 	UMassEntitySubsystem* EntitySys = UWorld::GetSubsystem<UMassEntitySubsystem>(World);
 	check(EntitySys);
 
-	// Sort anything there is to sort for later comparaison purposes
+	// Sort anything there is to sort for later comparison purposes
 	OutTemplate.Sort();
 
 	const FArchetypeHandle ArchetypeHandle = EntitySys->CreateArchetype(OutTemplate.GetCompositionDescriptor(), OutTemplate.GetSharedFragmentValues());
diff --git a/Engine/Plugins/Runtime/MassGameplay/Source/MassSpawner/Private/MassSpawnerSubsystem.cpp b/Engine/Plugins/Runtime/MassGameplay/Source/MassSpawner/Private/MassSpawnerSubsystem.cpp
--- a/Engine/Plugins/Runtime/MassGameplay/Source/MassSpawner/Private/MassSpawnerSubsystem.cpp
+++ b/Engine/Plugins/Runtime/MassGameplay/Source/MassSpawner/Private/MassSpawnerSubsystem.cpp
@@ -10,54 +10,12 @@
 #include "InstancedStruct.h"
 #include "VisualLogger/VisualLogger.h"
 #include "MassSpawner.h"
-#include "MassTranslator.h"
-#include "MassTranslatorRegistry.h"
+#include "MassObserverProcessor.h"
 #include "ProfilingDebugging/CpuProfilerTrace.h"
 #include "MassSimulationSubsystem.h"
+#include "MassProcessor.h"
 #include "MassEntityUtils.h"
 
-
-namespace UE::Mass::SpawnerSubsystem 
-{
-
-	bool RunFragmentDestructors(const FArchetypeChunkCollection& Chunks, UMassSpawnerSubsystem& SpawnerSubSystem, UMassEntitySubsystem& EntitySystem)
-	{
-		const FArchetypeHandle ArchetypeHandle = Chunks.GetArchetype();
-		// @todo this is a temporary measure to skip entities we've destroyed without the
-		// UMassSpawnerSubsystem::DestroyEntities' caller knowledge. Should be removed once that's addressed.
-		if (ArchetypeHandle.IsValid() == false)
-		{
-			return false;
-		}
-
-		// @Todo: should queue up the entity destruction and apply in batch
-		const UMassTranslatorRegistry& Registry = UMassTranslatorRegistry::Get();
-		TArray<const UMassProcessor*> Destructors;
-
-		EntitySystem.ForEachArchetypeFragmentType(ArchetypeHandle, [&Registry, &Destructors](const UScriptStruct* FragmentType)
-		{
-			check(FragmentType);
-			const UMassFragmentDestructor* Destructor = Registry.GetFragmentDestructor(*FragmentType);
-			if (Destructor)
-			{
-				Destructors.AddUnique(Destructor);
-			}
-		});
-
-		if (Destructors.Num())
-		{
-			FMassRuntimePipeline DestructionPipeline;
-			DestructionPipeline.InitializeFromArray(Destructors, SpawnerSubSystem);
-
-			FMassProcessingContext ProcessingContext(EntitySystem, /*TimeDelta=*/0.0f);
-
-			UE::Mass::Executor::RunSparse(DestructionPipeline, ProcessingContext, Chunks);
-		}
-
-		return Destructors.Num() > 0;
-	}
-} // namespace UE::Mass::SpawnerSubsystem
-
 //----------------------------------------------------------------------//
 //  UMassSpawnerSubsystem
 //----------------------------------------------------------------------//
@@ -68,8 +26,9 @@ void UMassSpawnerSubsystem::Initialize(FSubsystemCollectionBase& Collection)
 	// making sure UMassSimulationSubsystem gets created before the MassActorManager
 	Collection.InitializeDependency<UMassSimulationSubsystem>();
 
-	SimulationSystem = UWorld::GetSubsystem<UMassSimulationSubsystem>(GetWorld());
-	EntitySystem = UWorld::GetSubsystem<UMassEntitySubsystem>(GetWorld());
+	UWorld* World = GetWorld();
+	SimulationSystem = UWorld::GetSubsystem<UMassSimulationSubsystem>(World);
+	EntitySystem = UWorld::GetSubsystem<UMassEntitySubsystem>(World);
 	check(EntitySystem);
 }
 
@@ -156,61 +115,13 @@ void UMassSpawnerSubsystem::DestroyEntities(const FMassEntityTemplateID Template
 	UWorld* World = GetWorld();
 	check(World);
 
-	// Run deinitializers and destructor processors on all entities.
-	const FMassEntityTemplate* EntityTemplate = TemplateRegistryInstance->FindTemplateFromTemplateID(TemplateID);
-	if (EntityTemplate)
-	{
-		// There's no certainty all Entities still belong to the original archetype, the MutableEntityTemplate.GetArchetype().
-		// We need to figure out the separate, per-archetype sparse chunks
-		TArray<FArchetypeChunkCollection> ChunkCollections;
-		UE::Mass::Utils::CreateSparseChunks(*EntitySystem, Entities, ChunkCollections);
 
-		FMassProcessingContext ProcessingContext(*EntitySystem, /*TimeDelta=*/0.0f);
+	TArray<FArchetypeChunkCollection> ChunkCollections;
+	UE::Mass::Utils::CreateSparseChunks(*EntitySystem, Entities, ChunkCollections);
 
-		// Run destructors
-		bool bDestructorsRun = false;
-		for (const FArchetypeChunkCollection& Chunks : ChunkCollections)
-		{
-			bDestructorsRun = UE::Mass::SpawnerSubsystem::RunFragmentDestructors(Chunks, *this, *EntitySystem) || bDestructorsRun;
-		}
-
-		// if any destructors were called it's possible the entities changed archetypes. Need to 
-		// recalculate the chunks
-		if (bDestructorsRun)
-		{
-			ChunkCollections.Reset();
-			UE::Mass::Utils::CreateSparseChunks(*EntitySystem, Entities, ChunkCollections);
-		}
-		
-		// Run template deinitializers
-		FMassEntityTemplate& MutableEntityTemplate = const_cast<FMassEntityTemplate&>(*EntityTemplate);
-		if (MutableEntityTemplate.GetDeinitializationPipeline().Processors.Num())
-		{
-			for (const FArchetypeChunkCollection& Chunks : ChunkCollections)
-			{
-				if (Chunks.GetArchetype().IsValid())
-				{
-					UE::Mass::Executor::RunSparse(MutableEntityTemplate.GetMutableDeinitializationPipeline(), ProcessingContext, Chunks);
-				}
-			}
-		}
-
-		// if any deinitializer were called it's possible the entities changed archetypes. Need to 
-		// recalculate the chunks
-		if (MutableEntityTemplate.GetDeinitializationPipeline().Processors.Num())
-		{
-			ChunkCollections.Reset();
-			UE::Mass::Utils::CreateSparseChunks(*EntitySystem, Entities, ChunkCollections);
-		}
-
-		for (const FArchetypeChunkCollection& Chunks : ChunkCollections)
-		{
-			EntitySystem->BatchDestroyEntityChunks(Chunks);
-		}
-	}
-	else
+	for (const FArchetypeChunkCollection& Chunks : ChunkCollections)
 	{
-		EntitySystem->BatchDestroyEntities(Entities);
+		EntitySystem->BatchDestroyEntityChunks(Chunks);
 	}
 }
 
@@ -244,19 +155,25 @@ void UMassSpawnerSubsystem::DoSpawning(const FMassEntityTemplate& EntityTemplate
 	check(EntityTemplate.GetArchetype().IsValid());
 	UE_VLOG(this, LogMassSpawner, Log, TEXT("Spawning with EntityTemplate:\n%s"), *EntityTemplate.DebugGetDescription(EntitySystem));
 
+	if (NumToSpawn <= 0)
+	{
+		UE_VLOG(this, LogMassSpawner, Warning, TEXT("%s: Trying to spawn %d entities. Ignoring."), ANSI_TO_TCHAR(__FUNCTION__), NumToSpawn);
+		return;
+	}
+
 	//TRACE_CPUPROFILER_EVENT_SCOPE_STR("MassSpawnerSubsystem DoSpawning");
 
 	// 1. Create required number of entities with EntityTemplate.Archetype
 	// 2. Copy data from FMassEntityTemplate.Fragments.
 	//		a. @todo, could be done as part of creation?
-	// 3. Run all EntityTemplate.Translators to configure remaining values (i.e. not set via FMassEntityTemplate.Fragments).
+	// 3. Run SpawlLocationInitializer if set
+	// 4. "OnEntitiesCreated" notifies will be sent out once the CreationContext gets destroyed (via its destructor).
 
 	TArray<FMassEntityHandle> SpawnedEntities;
-	EntitySystem->BatchCreateEntities(EntityTemplate.GetArchetype(), NumToSpawn, SpawnedEntities);
+	TSharedRef<UMassEntitySubsystem::FEntityCreationContext> CreationContext = EntitySystem->BatchCreateEntities(EntityTemplate.GetArchetype(), NumToSpawn, SpawnedEntities);
 
 	TConstArrayView<FInstancedStruct> FragmentInstances = EntityTemplate.GetInitialFragmentValues();
-	const FArchetypeChunkCollection ChunkCollection(EntityTemplate.GetArchetype(), SpawnedEntities);
-	EntitySystem->BatchSetEntityFragmentsValues(ChunkCollection, FragmentInstances);
+	EntitySystem->BatchSetEntityFragmentsValues(CreationContext->GetChunkCollection(), FragmentInstances);
 	
 	UMassProcessor* SpawnLocationInitializer = SpawnData.IsValid() ? GetSpawnLocationInitializer(InitializerClass) : nullptr;
 
@@ -264,15 +181,7 @@ void UMassSpawnerSubsystem::DoSpawning(const FMassEntityTemplate& EntityTemplate
 	{
 		FMassProcessingContext ProcessingContext(*EntitySystem, /*TimeDelta=*/0.0f);
 		ProcessingContext.AuxData = SpawnData;
-		UE::Mass::Executor::RunProcessorsView(MakeArrayView(&SpawnLocationInitializer, 1), ProcessingContext, &ChunkCollection);
-	}
-
-	if (EntityTemplate.GetInitializationPipeline().Processors.Num())
-	{
-		FMassEntityTemplate& MutableEntityTemplate = const_cast<FMassEntityTemplate&>(EntityTemplate);
-
-		FMassProcessingContext ProcessingContext(*EntitySystem, /*TimeDelta=*/0.0f);
-		UE::Mass::Executor::RunSparse(MutableEntityTemplate.GetMutableInitializationPipeline(), ProcessingContext, ChunkCollection);
+		UE::Mass::Executor::RunProcessorsView(MakeArrayView(&SpawnLocationInitializer, 1), ProcessingContext, &CreationContext->GetChunkCollection());
 	}
 
 	OutEntities.Append(MoveTemp(SpawnedEntities));
diff --git a/Engine/Plugins/Runtime/MassGameplay/Source/MassSpawner/Private/MassTranslator.cpp b/Engine/Plugins/Runtime/MassGameplay/Source/MassSpawner/Private/MassTranslator.cpp
--- a/Engine/Plugins/Runtime/MassGameplay/Source/MassSpawner/Private/MassTranslator.cpp
+++ b/Engine/Plugins/Runtime/MassGameplay/Source/MassSpawner/Private/MassTranslator.cpp
@@ -2,7 +2,6 @@
 
 #include "MassTranslator.h"
 #include "MassCommonTypes.h"
-#include "MassTranslatorRegistry.h"
 
 //----------------------------------------------------------------------//
 //  UMassTranslator
@@ -17,46 +16,3 @@ void UMassTranslator::AddRequiredTagsToQuery(FMassEntityQuery& EntityQuery)
 {
 	EntityQuery.AddTagRequirements<EMassFragmentPresence::All>(RequiredTags);
 }
-
-//----------------------------------------------------------------------//
-// UMassFragmentBuilder
-//----------------------------------------------------------------------//
-UMassFragmentBuilder::UMassFragmentBuilder()
-{
-	bAutoRegisterWithProcessingPhases = false;
-#if WITH_EDITORONLY_DATA
-	bCanShowUpInSettings = false;
-#endif // WITH_EDITORONLY_DATA
-}
-
-void UMassFragmentBuilder::PostInitProperties()
-{
-	Super::PostInitProperties();
-
-	if (GetClass()->HasAnyClassFlags(CLASS_Abstract) == false && FragmentType != nullptr)
-	{
-		Register();
-	}
-}
-
-//----------------------------------------------------------------------//
-//  UMassFragmentInitializer
-//----------------------------------------------------------------------//
-void UMassFragmentInitializer::Register()
-{
-	if (FragmentType)
-	{
-		UMassTranslatorRegistry::GetMutable().RegisterFragmentInitializer(*FragmentType, GetClass());
-	}
-}
-
-//----------------------------------------------------------------------//
-//  UMassFragmentDestructor
-//----------------------------------------------------------------------//
-void UMassFragmentDestructor::Register()
-{
-	if (FragmentType)
-	{
-		UMassTranslatorRegistry::GetMutable().RegisterFragmentDestructor(*FragmentType, GetClass());
-	}
-}
\ No newline at end of file
diff --git a/Engine/Plugins/Runtime/MassGameplay/Source/MassSpawner/Private/MassTranslatorRegistry.cpp b/Engine/Plugins/Runtime/MassGameplay/Source/MassSpawner/Private/MassTranslatorRegistry.cpp
deleted file mode 100644
diff --git a/Engine/Plugins/Runtime/MassGameplay/Source/MassSpawner/Public/MassEntityTemplate.h b/Engine/Plugins/Runtime/MassGameplay/Source/MassSpawner/Public/MassEntityTemplate.h
--- a/Engine/Plugins/Runtime/MassGameplay/Source/MassSpawner/Public/MassEntityTemplate.h
+++ b/Engine/Plugins/Runtime/MassGameplay/Source/MassSpawner/Public/MassEntityTemplate.h
@@ -22,18 +22,6 @@ enum class EMassEntityTemplateIDType : uint8
 	Instance
 };
 
-USTRUCT()
-struct MASSSPAWNER_API FMassObjectHandlers
-{
-	GENERATED_BODY()
-	UPROPERTY()
-	TArray<const UMassFragmentInitializer*> DefaultInitializers;
-	UPROPERTY()
-	TArray<const UMassProcessor*> Initializers;
-	UPROPERTY()
-	TArray<const UMassProcessor*> Deinitializers;
-};
-
 //ID of the template an entity is using
 USTRUCT()
 struct MASSSPAWNER_API FMassEntityTemplateID
@@ -87,14 +75,6 @@ struct MASSSPAWNER_API FMassEntityTemplate
 	/** InArchetype is expected to be valid. The function will crash-check it. */
 	void SetArchetype(const FArchetypeHandle& InArchetype);
 	const FArchetypeHandle& GetArchetype() const { return Archetype; }
-	/** Configures initialization and deinitialization pipeline from ObjectHandlers */
-	void SetUpProcessors(const FMassObjectHandlers& ObjectHandlers, UObject& PipelineOwner);
-	
-	const FMassRuntimePipeline& GetInitializationPipeline() const { return InitializationPipeline; }
-	FMassRuntimePipeline& GetMutableInitializationPipeline() { return InitializationPipeline; }
-
-	const FMassRuntimePipeline& GetDeinitializationPipeline() const { return DeinitializationPipeline; }
-	FMassRuntimePipeline& GetMutableDeinitializationPipeline() { return DeinitializationPipeline; }
 
 	TConstArrayView<FObjectFragmentInitializerFunction> GetObjectFragmentInitializers() const { return ObjectInitializers; }
 	TArray<FObjectFragmentInitializerFunction>& GetMutableObjectFragmentInitializers() { return ObjectInitializers; }
@@ -263,12 +243,6 @@ private:
 	// Initial fragment values, this is not part of the archetype as it is the spawner job to set them.
 	TArray<FInstancedStruct> InitialFragmentValues;
 
-	UPROPERTY()
-	FMassRuntimePipeline InitializationPipeline;
-	
-	UPROPERTY()
-	FMassRuntimePipeline DeinitializationPipeline;
-
 	// These functions will be called to initialize entity's UObject-based fragments
 	TArray<FObjectFragmentInitializerFunction> ObjectInitializers;
 
diff --git a/Engine/Plugins/Runtime/MassGameplay/Source/MassSpawner/Public/MassEntityTemplateRegistry.h b/Engine/Plugins/Runtime/MassGameplay/Source/MassSpawner/Public/MassEntityTemplateRegistry.h
--- a/Engine/Plugins/Runtime/MassGameplay/Source/MassSpawner/Public/MassEntityTemplateRegistry.h
+++ b/Engine/Plugins/Runtime/MassGameplay/Source/MassSpawner/Public/MassEntityTemplateRegistry.h
@@ -6,7 +6,6 @@
 #include "UObject/Object.h"
 #include "MassCommonTypes.h"
 #include "MassEntityTemplate.h"
-#include "MassTranslatorRegistry.h"
 #include "MassTranslator.h"
 #include "MassEntityTemplateRegistry.generated.h"
 
@@ -35,60 +34,17 @@ struct FMassEntityTemplateBuildContext
 		return Template.AddFragment_GetRef<T>();
 	}
 
-	template<typename T>
-	T& AddFragmentWithDefaultInitializer_GetRef()
-	{
-		AddDefaultInitializer<T>();
-		return AddFragment_GetRef<T>();
-	}
-
-	template<typename T>
-	T& AddFragmentWithInitializer_GetRef(const UMassProcessor& Initializer)
-	{
-		AddInitializer(Initializer);
-		return AddFragment_GetRef<T>();
-	}
-
 	template<typename T>
 	void AddFragment()
 	{
 		Template.AddFragment<T>();
 	}
 
-	template<typename T>
-	void AddFragmentWithDefaultInitializer()
-	{
-		AddFragment<T>();
-		AddDefaultInitializer<T>();
-	}
-
-	template<typename T>
-	void AddFragmentWithInitializer(const UMassProcessor& Initializer)
-	{
-		AddFragment<T>();
-		AddInitializer(Initializer);
-	}
-
 	void AddFragment(FConstStructView InFragment)
 	{ 
 		Template.AddFragment(InFragment);
 	}
 
-	void AddFragmentWithDefaultInitializer(FConstStructView InFragment)
-	{
-		AddFragment(InFragment);
-		if (InFragment.GetScriptStruct())
-		{
-			AddDefaultInitializer(*InFragment.GetScriptStruct());
-		}
-	}
-
-	void AddFragmentWithInitializer(FConstStructView InFragment, const UMassProcessor& Initializer)
-	{
-		AddFragment(InFragment);
-		AddInitializer(Initializer);
-	}
-
 	template<typename T>
 	void AddTag()
 	{
@@ -145,54 +101,6 @@ struct FMassEntityTemplateBuildContext
 		return Template.HasSharedFragment<T>();
 	}
 
-	//----------------------------------------------------------------------//
-	// Initializers
-	//----------------------------------------------------------------------//
-	void AddInitializer(const UMassProcessor& Initializer) { Handlers.Initializers.AddUnique(&Initializer); }
-	
-	void AddDefaultInitializer(const UScriptStruct& FragmentType)
-	{
-		const UMassTranslatorRegistry& Registry = UMassTranslatorRegistry::Get();
-		if (const UMassFragmentInitializer* Initializer = Registry.GetFragmentInitializer(FragmentType))
-		{
-			Handlers.DefaultInitializers.AddUnique(Initializer);
-		}
-	}
-
-	template<typename T>
-	void AddDefaultInitializer()
-	{
-		const UMassTranslatorRegistry& Registry = UMassTranslatorRegistry::Get();
-		if (const UMassFragmentInitializer* Initializer = Registry.GetFragmentInitializer(*T::StaticStruct()))
-		{
-			Handlers.DefaultInitializers.AddUnique(Initializer);
-		}
-	}
-
-	//----------------------------------------------------------------------//
-	// Deinitializers
-	//----------------------------------------------------------------------//
-	void AddDeinitializer(const UMassProcessor& Deinitializer) { Handlers.Deinitializers.AddUnique(&Deinitializer); }
-
-	void AddDefaultDeinitializer(const UScriptStruct& FragmentType)
-	{
-		const UMassTranslatorRegistry& Registry = UMassTranslatorRegistry::Get();
-		if (const UMassFragmentDestructor* Deinitializer = Registry.GetFragmentDestructor(FragmentType))
-		{
-			Handlers.Deinitializers.AddUnique(Deinitializer);
-		}
-	}
-
-	template<typename T>
-	void AddDefaultDeinitializer()
-	{
-		const UMassTranslatorRegistry& Registry = UMassTranslatorRegistry::Get();
-		if (const UMassFragmentDestructor* Deinitializer = Registry.GetFragmentDestructor(*T::StaticStruct()))
-		{
-			Handlers.Deinitializers.AddUnique(Deinitializer);
-		}
-	}
-
 	//----------------------------------------------------------------------//
 	// Translators
 	//----------------------------------------------------------------------//
@@ -209,7 +117,6 @@ struct FMassEntityTemplateBuildContext
 	FMassEntityTemplateID GetTemplateID() const { return Template.GetTemplateID(); }
 	TArray<FMassEntityTemplate::FObjectFragmentInitializerFunction>& GetMutableObjectFragmentInitializers() { return Template.GetMutableObjectFragmentInitializers(); }
 
-	FMassObjectHandlers Handlers;
 protected:
 	FMassEntityTemplate& Template;
 };
diff --git a/Engine/Plugins/Runtime/MassGameplay/Source/MassSpawner/Public/MassTranslator.h b/Engine/Plugins/Runtime/MassGameplay/Source/MassSpawner/Public/MassTranslator.h
--- a/Engine/Plugins/Runtime/MassGameplay/Source/MassSpawner/Public/MassTranslator.h
+++ b/Engine/Plugins/Runtime/MassGameplay/Source/MassSpawner/Public/MassTranslator.h
@@ -5,6 +5,7 @@
 #include "MassProcessor.h"
 #include "MassCommonTypes.h"
 #include "MassEntityTypes.h"
+#include "MassObserverProcessor.h"
 #include "MassTranslator.generated.h"
 
 
@@ -50,46 +51,3 @@ protected:
 	 *  the system which translators they'd want their entity to be processed by. */
 	FMassTagBitSet RequiredTags;
 };
-
-UCLASS(abstract)
-class MASSSPAWNER_API UMassFragmentBuilder : public UMassProcessor
-{
-	GENERATED_BODY()
-
-public:
-	UMassFragmentBuilder();
-
-	/** Override to register this fragment builder to the registry */
-	virtual void Register() PURE_VIRTUAL(UMassFragmentBuilder::Register, );
-
-protected:
-	virtual void PostInitProperties() override;
-
-protected:
-	/** Set in class' constructor and determines for which Fragment type this given UMassFragmentBuilder will be used 
-	 *  as the default initializer/deinitializer (as returned by the UMassTranslatorRegistry). If set to null will
-	 *  require the user to manually add this UMassFragmentBuilder as initializer/deinitializer 
-	 *  (@see FMassEntityTemplateBuildContext.AddInitializer and FMassEntityTemplateBuildContext.AddDeinitializer) */
-	UPROPERTY()
-	UScriptStruct* FragmentType;
-};
-
-UCLASS(abstract)
-class MASSSPAWNER_API UMassFragmentInitializer : public UMassFragmentBuilder
-{
-	GENERATED_BODY()
-
-public:
-	/** adds this FragmentInitializer to the MassTranslatorRegistry to be used to initialize fragments of FragmentType */
-	virtual void Register();
-};
-
-UCLASS(abstract)
-class MASSSPAWNER_API UMassFragmentDestructor : public UMassFragmentBuilder
-{
-	GENERATED_BODY()
-
-public:
-	/** adds this FragmentDestructor to the MassTranslatorRegistry to be used to initialize fragments of FragmentType */
-	virtual void Register();
-};
\ No newline at end of file
diff --git a/Engine/Plugins/Runtime/MassGameplay/Source/MassSpawner/Public/MassTranslatorRegistry.h b/Engine/Plugins/Runtime/MassGameplay/Source/MassSpawner/Public/MassTranslatorRegistry.h
deleted file mode 100644
