Moved FMassTrafficCommandAddFragmentInstanceList over to MassEntity and renamed it to FMassCommandAddFragmentInstanceList

diff --git a/Engine/Plugins/Runtime/MassEntity/Source/MassEntity/Private/MassCommandBuffer.cpp b/Engine/Plugins/Runtime/MassEntity/Source/MassEntity/Private/MassCommandBuffer.cpp
--- a/Engine/Plugins/Runtime/MassEntity/Source/MassEntity/Private/MassCommandBuffer.cpp
+++ b/Engine/Plugins/Runtime/MassEntity/Source/MassEntity/Private/MassCommandBuffer.cpp
@@ -186,6 +186,11 @@ void FCommandAddFragmentInstance::Execute(UMassEntitySubsystem& System) const
 	System.AddFragmentInstanceListToEntity(TargetEntity, MakeArrayView(&Struct, 1));
 }
 
+void FMassCommandAddFragmentInstanceList::Execute(UMassEntitySubsystem& System) const
+{
+	System.AddFragmentInstanceListToEntity(TargetEntity, FragmentList);
+}
+
 void FCommandSwapTags::Execute(UMassEntitySubsystem& System) const
 {
 	if (System.IsEntityValid(TargetEntity) == false)
diff --git a/Engine/Plugins/Runtime/MassEntity/Source/MassEntity/Public/MassCommandBuffer.h b/Engine/Plugins/Runtime/MassEntity/Source/MassEntity/Public/MassCommandBuffer.h
--- a/Engine/Plugins/Runtime/MassEntity/Source/MassEntity/Public/MassCommandBuffer.h
+++ b/Engine/Plugins/Runtime/MassEntity/Source/MassEntity/Public/MassCommandBuffer.h
@@ -191,6 +191,40 @@ protected:
 	FInstancedStruct Struct;
 };
 
+/**
+ * Command that adds a list of instanced fragments to a given entity
+ */
+USTRUCT()
+struct MASSENTITY_API FMassCommandAddFragmentInstanceList : public FCommandBufferEntryBase
+{
+	GENERATED_BODY()
+
+		enum
+	{
+		Type = ECommandBufferOperationType::Add
+	};
+
+	FMassCommandAddFragmentInstanceList() = default;
+
+	FMassCommandAddFragmentInstanceList(const FMassEntityHandle InEntity, std::initializer_list<FInstancedStruct> InitList)
+		: FCommandBufferEntryBase(InEntity)
+		, FragmentList(InitList)
+	{}
+
+	void GetAffectedTypes(TMap<const UScriptStruct*, TArray<FMassEntityHandle>>& OutAdded, TMap<const UScriptStruct*, TArray<FMassEntityHandle>>& OutRemoved)
+	{
+		for (const FInstancedStruct& Struct : FragmentList)
+		{
+			OutAdded.FindOrAdd(Struct.GetScriptStruct()).AddUnique(TargetEntity);
+		}
+	}
+
+protected:
+	virtual void Execute(UMassEntitySubsystem& System) const override;
+
+	TArray<FInstancedStruct> FragmentList;
+};
+
 /**
  * Command dedicated to remove a fragment from an existing entity
  */
