Fixed FArchetypeChunkCollection to not break horribly when given collection of entities contains duplicates. This behavior is optional, but enabled by default.

diff --git a/Engine/Plugins/Runtime/MassEntity/Source/MassEntity/Private/MassArchetypeTypes.cpp b/Engine/Plugins/Runtime/MassEntity/Source/MassEntity/Private/MassArchetypeTypes.cpp
--- a/Engine/Plugins/Runtime/MassEntity/Source/MassEntity/Private/MassArchetypeTypes.cpp
+++ b/Engine/Plugins/Runtime/MassEntity/Source/MassEntity/Private/MassArchetypeTypes.cpp
@@ -21,10 +21,16 @@ uint32 GetTypeHash(const FArchetypeHandle& Instance)
 //////////////////////////////////////////////////////////////////////
 // FArchetypeChunkCollection
 
-FArchetypeChunkCollection::FArchetypeChunkCollection(const FArchetypeHandle& InArchetype, TConstArrayView<FMassEntityHandle> InEntities)
+FArchetypeChunkCollection::FArchetypeChunkCollection(const FArchetypeHandle& InArchetype, TConstArrayView<FMassEntityHandle> InEntities, EDuplicatesHandling DuplicatesHandling)
 	: Archetype(InArchetype)
 {
 	check(InArchetype.IsValid());
+
+	if (InEntities.Num() <= 0)
+	{
+		return;
+	}
+
 	const FMassArchetypeData& ArchetypeData = *InArchetype.DataPtr.Get();
 	const int32 NumEntitiesPerChunk = ArchetypeData.GetNumEntitiesPerChunk();
 
@@ -39,6 +45,44 @@ FArchetypeChunkCollection::FArchetypeChunkCollection(const FArchetypeHandle& InA
 	}
 	TrueIndices.Sort();
 
+#if DO_GUARD_SLOW
+	if (DuplicatesHandling == NoDuplicates)
+	{
+		// ensure there are no duplicates. 
+		int32 PrevIndex = TrueIndices[0];
+		for (int j = 1; j < TrueIndices.Num(); ++j)
+		{
+			checkf(TrueIndices[j] != PrevIndex, TEXT("InEntities contains duplicate while DuplicatesHandling is set to NoDuplicates"));
+			if (TrueIndices[j] == PrevIndex)
+			{
+				// fix it, for development's sake
+				DuplicatesHandling = FoldDuplicates;
+				break;
+			}
+			PrevIndex = TrueIndices[j];
+		}
+	}
+#endif // DO_GUARD_SLOW
+
+	if (DuplicatesHandling == FoldDuplicates)
+	{
+		int32 PrevIndex = TrueIndices[0];
+		for (int j = 1; j < TrueIndices.Num(); ++j)
+		{
+			if (TrueIndices[j] == PrevIndex)
+			{
+				const int32 Num = TrueIndices.Num();
+				int Skip = 0;
+				while ((j + ++Skip) < Num && TrueIndices[j + Skip] == PrevIndex);
+				
+				TrueIndices.RemoveAt(j, Skip, /*bAllowShrinking=*/false);
+				--j;
+				continue;
+			}
+			PrevIndex = TrueIndices[j];
+		}
+	}
+
 	// the following block of code is splitting up sorted AbsoluteIndices into 
 	// continuous chunks
 	int32 ChunkEnd = INDEX_NONE;
@@ -90,3 +134,20 @@ void FArchetypeChunkCollection::GatherChunksFromArchetype(TSharedPtr<FMassArchet
 		Chunks.Add(FChunkInfo(i));
 	}
 }
+
+bool FArchetypeChunkCollection::IsSame(const FArchetypeChunkCollection& Other) const
+{
+	if (Archetype != Other.Archetype || Chunks.Num() != Other.Chunks.Num())
+	{
+		return false;
+	}
+
+	for (int i = 0; i < Chunks.Num(); ++i)
+	{
+		if (Chunks[i] != Other.Chunks[i])
+		{
+			return false;
+		}
+	}
+	return true;
+}
diff --git a/Engine/Plugins/Runtime/MassEntity/Source/MassEntity/Private/MassEntitySubsystem.cpp b/Engine/Plugins/Runtime/MassEntity/Source/MassEntity/Private/MassEntitySubsystem.cpp
--- a/Engine/Plugins/Runtime/MassEntity/Source/MassEntity/Private/MassEntitySubsystem.cpp
+++ b/Engine/Plugins/Runtime/MassEntity/Source/MassEntity/Private/MassEntitySubsystem.cpp
@@ -336,7 +336,7 @@ TSharedRef<UMassEntitySubsystem::FEntityCreationContext> UMassEntitySubsystem::B
 	FEntityCreationContext* CreationContext = new FEntityCreationContext(Count);
 	// @todo this could probably be optimized since one would assume we're adding elements to OutEntities in order.
 	// Then again, if that's the case, the sorting will be almost instant
-	new (&CreationContext->ChunkCollection)FArchetypeChunkCollection(Archetype, MakeArrayView(&OutEntities[OutEntities.Num() - Count], Count));
+	new (&CreationContext->ChunkCollection)FArchetypeChunkCollection(Archetype, MakeArrayView(&OutEntities[OutEntities.Num() - Count], Count), FArchetypeChunkCollection::NoDuplicates);
 	if (ObserverManager.HasOnAddedObserversForFragments(ArchetypePtr->GetCompositionDescriptor().Fragments))
 	{
 		CreationContext->OnSpawningFinished = [this](FEntityCreationContext& Context){
diff --git a/Engine/Plugins/Runtime/MassEntity/Source/MassEntity/Private/MassEntityUtils.cpp b/Engine/Plugins/Runtime/MassEntity/Source/MassEntity/Private/MassEntityUtils.cpp
--- a/Engine/Plugins/Runtime/MassEntity/Source/MassEntity/Private/MassEntityUtils.cpp
+++ b/Engine/Plugins/Runtime/MassEntity/Source/MassEntity/Private/MassEntityUtils.cpp
@@ -53,7 +53,7 @@ void CreateSparseChunks(const UMassEntitySubsystem& EntitySystem, const TConstAr
 		// UMassSpawnerSubsystem::DestroyEntities' caller knowledge. Should be removed once that's addressed.
 		if (Pair.Key.IsValid())
 		{
-			OutChunkCollections.Add(FArchetypeChunkCollection(Pair.Key, Pair.Value));
+			OutChunkCollections.Add(FArchetypeChunkCollection(Pair.Key, Pair.Value, FArchetypeChunkCollection::NoDuplicates));
 		}
 	}
 }
diff --git a/Engine/Plugins/Runtime/MassEntity/Source/MassEntity/Private/MassExecutor.cpp b/Engine/Plugins/Runtime/MassEntity/Source/MassEntity/Private/MassExecutor.cpp
--- a/Engine/Plugins/Runtime/MassEntity/Source/MassEntity/Private/MassExecutor.cpp
+++ b/Engine/Plugins/Runtime/MassEntity/Source/MassEntity/Private/MassExecutor.cpp
@@ -33,7 +33,7 @@ void RunSparse(FMassRuntimePipeline& RuntimePipeline, FMassProcessingContext& Pr
 
 	TRACE_CPUPROFILER_EVENT_SCOPE_STR("MassExecutor RunSparseEntities");
 
-	const FArchetypeChunkCollection ChunkCollection(Archetype, Entities);
+	const FArchetypeChunkCollection ChunkCollection(Archetype, Entities, FArchetypeChunkCollection::NoDuplicates);
 	RunProcessorsView(RuntimePipeline.Processors, ProcessingContext, &ChunkCollection);
 }
 
diff --git a/Engine/Plugins/Runtime/MassEntity/Source/MassEntity/Private/MassObserverManager.cpp b/Engine/Plugins/Runtime/MassEntity/Source/MassEntity/Private/MassObserverManager.cpp
--- a/Engine/Plugins/Runtime/MassEntity/Source/MassEntity/Private/MassObserverManager.cpp
+++ b/Engine/Plugins/Runtime/MassEntity/Source/MassEntity/Private/MassObserverManager.cpp
@@ -138,7 +138,8 @@ bool FMassObserverManager::OnPostCompositionAdded(const FMassEntityHandle Entity
 	{
 		FMassProcessingContext ProcessingContext(EntitySubsystem, /*DeltaSeconds=*/0.f);
 		const FArchetypeHandle ArchetypeHandle = EntitySubsystem.GetArchetypeForEntity(Entity);
-		HandleFragmentsImpl(ProcessingContext, FArchetypeChunkCollection(ArchetypeHandle, MakeArrayView(&Entity, 1)), Overlap, OnFragmentAddedObservers);
+		HandleFragmentsImpl(ProcessingContext, FArchetypeChunkCollection(ArchetypeHandle, MakeArrayView(&Entity, 1)
+			, FArchetypeChunkCollection::NoDuplicates), Overlap, OnFragmentAddedObservers);
 		return true;
 	}
 
@@ -152,7 +153,8 @@ bool FMassObserverManager::OnPreCompositionRemoved(const FMassEntityHandle Entit
 	{
 		FMassProcessingContext ProcessingContext(EntitySubsystem, /*DeltaSeconds=*/0.f);
 		const FArchetypeHandle ArchetypeHandle = EntitySubsystem.GetArchetypeForEntity(Entity);
-		HandleFragmentsImpl(ProcessingContext, FArchetypeChunkCollection(ArchetypeHandle, MakeArrayView(&Entity, 1)), Overlap, OnFragmentRemovedObservers);
+		HandleFragmentsImpl(ProcessingContext, FArchetypeChunkCollection(ArchetypeHandle, MakeArrayView(&Entity, 1)
+			, FArchetypeChunkCollection::NoDuplicates), Overlap, OnFragmentRemovedObservers);
 		return true;
 	}
 
diff --git a/Engine/Plugins/Runtime/MassEntity/Source/MassEntity/Public/MassArchetypeTypes.h b/Engine/Plugins/Runtime/MassEntity/Source/MassEntity/Public/MassArchetypeTypes.h
--- a/Engine/Plugins/Runtime/MassEntity/Source/MassEntity/Public/MassArchetypeTypes.h
+++ b/Engine/Plugins/Runtime/MassEntity/Source/MassEntity/Public/MassArchetypeTypes.h
@@ -53,7 +53,7 @@ private:
 
 /** A struct that converts an arbitrary array of entities of given Archetype into a sequence of continuous
  *  entity chunks. The goal is to have the user create an instance of this struct once and run through a bunch of
- *  systems. The runtime code usually uses FMassArchetypeChunkIterator to iterate on the chunk collection
+ *  systems. The runtime code usually uses FMassArchetypeChunkIterator to iterate on the chunk collection.
  */
 struct MASSENTITY_API FArchetypeChunkCollection
 {
@@ -69,7 +69,22 @@ public:
 		explicit FChunkInfo(const int32 InChunkIndex, const int32 InSubchunkStart = 0, const int32 InLength = 0) : ChunkIndex(InChunkIndex), SubchunkStart(InSubchunkStart), Length(InLength) {}
 		/** Note that we consider invalid-length chunks valid as long as ChunkIndex and SubchunkStart are valid */
 		bool IsSet() const { return ChunkIndex != INDEX_NONE && SubchunkStart >= 0; }
+
+		bool operator==(const FChunkInfo& Other) const
+		{
+			return ChunkIndex == Other.ChunkIndex && SubchunkStart == Other.SubchunkStart && Length == Other.Length;
+		}
+		bool operator!=(const FChunkInfo& Other) const { return !(*this == Other); }
+	};
+
+	enum EDuplicatesHandling
+	{
+		NoDuplicates,	// indicates that the caller guarantees there are no duplicates in the input Entities collection
+						// note that in no-shipping builds a `check` will fail if duplicates are present.
+		FoldDuplicates,	// indicates that it's possible that Entities contains duplicates. The input Entities collection 
+						// will be processed and duplicates will be removed.
 	};
+
 private:
 	TArray<FChunkInfo> Chunks;
 	/** entity indices indicated by SubChunks are only valid with given Archetype */
@@ -77,7 +92,7 @@ private:
 
 public:
 	FArchetypeChunkCollection() = default;
-	FArchetypeChunkCollection(const FArchetypeHandle& InArchetype, TConstArrayView<FMassEntityHandle> InEntities);
+	FArchetypeChunkCollection(const FArchetypeHandle& InArchetype, TConstArrayView<FMassEntityHandle> InEntities, EDuplicatesHandling DuplicatesHandling);
 	explicit FArchetypeChunkCollection(FArchetypeHandle& InArchetypeHandle);
 	explicit FArchetypeChunkCollection(TSharedPtr<FMassArchetypeData>& InArchetype);
 
@@ -91,6 +106,9 @@ public:
 		Chunks.Reset();
 	}
 
+	/** The comparison function that checks if Other is identical to this. Intended for diagnostics/debugging. */
+	bool IsSame(const FArchetypeChunkCollection& Other) const;
+
 private:
 	void GatherChunksFromArchetype(TSharedPtr<FMassArchetypeData>& InArchetype);
 };
diff --git a/Engine/Plugins/Runtime/MassEntity/Source/MassEntity/Public/MassEntityUtils.h b/Engine/Plugins/Runtime/MassEntity/Source/MassEntity/Public/MassEntityUtils.h
--- a/Engine/Plugins/Runtime/MassEntity/Source/MassEntity/Public/MassEntityUtils.h
+++ b/Engine/Plugins/Runtime/MassEntity/Source/MassEntity/Public/MassEntityUtils.h
@@ -15,7 +15,10 @@ namespace UE::Mass::Utils
 /** returns the current execution mode for the processors calculated from the world network mode */
 MASSENTITY_API extern EProcessorExecutionFlags GetProcessorExecutionFlagsForWold(const UWorld& World);
 
-
+/** 
+ * Fills OutChunkCollections with per-archetype FArchetypeChunkCollection instances. 
+ * Note that it's caller responsibility to guarantee Entities does not contain duplicates.
+ */
 MASSENTITY_API extern void CreateSparseChunks(const UMassEntitySubsystem& EntitySystem, const TConstArrayView<FMassEntityHandle> Entities, TArray<FArchetypeChunkCollection>& OutChunkCollections);
 
 } // namespace UE::Mass::Utils
\ No newline at end of file
diff --git a/Engine/Plugins/Runtime/MassEntity/Source/MassEntityTestSuite/Private/MassChunkCollectionTest.cpp b/Engine/Plugins/Runtime/MassEntity/Source/MassEntityTestSuite/Private/MassChunkCollectionTest.cpp
new file mode 100644
--- /dev/null
+++ b/Engine/Plugins/Runtime/MassEntity/Source/MassEntityTestSuite/Private/MassChunkCollectionTest.cpp
@@ -0,0 +1,168 @@
+// Copyright Epic Games, Inc. All Rights Reserved.
+
+#include "MassEntitySubsystem.h"
+#include "MassEntityTestTypes.h"
+
+#define LOCTEXT_NAMESPACE "MassTest"
+
+PRAGMA_DISABLE_OPTIMIZATION
+
+//----------------------------------------------------------------------//
+// tests 
+//----------------------------------------------------------------------//
+namespace FMassChunkCollectionTest
+{
+
+void Shuffle(FRandomStream& Rand, TArray<FMassEntityHandle>& Entities)
+{
+	for (int i = 0; i < Entities.Num(); ++i)
+	{
+		const int32 NewIndex = Rand.RandRange(0, Entities.Num() - 1);
+		Entities.Swap(i, NewIndex);
+	}
+}
+
+struct FChunkCollectionTestBase : FEntityTestBase
+{
+	TArray<FMassEntityHandle> Entities;
+
+	virtual bool SetUp() override
+	{
+		FEntityTestBase::SetUp();
+
+		const int32 Count = 100;
+		EntitySubsystem->BatchCreateEntities(FloatsArchetype, Count, Entities);
+		return true;
+	}
+
+	virtual void TearDown() override
+	{
+		Entities.Reset();
+		FEntityTestBase::TearDown();
+	}
+};
+
+struct FChunkCollection_CreateBasic : FChunkCollectionTestBase
+{
+	virtual bool InstantTest() override
+	{
+		TArray<FMassEntityHandle> EntitiesSubSet;
+		
+		// Should end up as last chunk
+		EntitiesSubSet.Add(Entities[99]);
+		EntitiesSubSet.Add(Entities[97]);
+		EntitiesSubSet.Add(Entities[98]);
+
+		// Should end up as third chunk
+		EntitiesSubSet.Add(Entities[20]);
+		EntitiesSubSet.Add(Entities[22]);
+		EntitiesSubSet.Add(Entities[21]);
+
+		// Should end up as second chunk
+		EntitiesSubSet.Add(Entities[18]);
+
+		// Should end up as first chunk
+		EntitiesSubSet.Add(Entities[10]);
+		EntitiesSubSet.Add(Entities[13]);
+		EntitiesSubSet.Add(Entities[11]);
+		EntitiesSubSet.Add(Entities[12]);
+
+		FArchetypeChunkCollection ChunkCollection(FloatsArchetype, EntitiesSubSet, FArchetypeChunkCollection::NoDuplicates);
+		TArrayView<const FArchetypeChunkCollection::FChunkInfo> Chunks = ChunkCollection.GetChunks();
+		AITEST_EQUAL("The predicted sub-chunk count should match", Chunks.Num(), 4);
+		AITEST_EQUAL("The [10-13] chunk should be first and start at 10", Chunks[0].SubchunkStart, 10);
+		AITEST_EQUAL("The [10-13] chunk should be first and have a length of 4", Chunks[0].Length, 4);
+		AITEST_EQUAL("The [18] chunk should be second and start at 18", Chunks[1].SubchunkStart, 18); 
+		AITEST_EQUAL("The [18] chunk should be second and have a length of 1", Chunks[1].Length, 1);
+		AITEST_EQUAL("The [20-22] chunk should be third and start at 20", Chunks[2].SubchunkStart, 20);
+		AITEST_EQUAL("The [20-22] chunk should be third and have a length of 3", Chunks[2].Length, 3);
+		AITEST_EQUAL("The [97-99] chunk should be third and start at 97", Chunks[3].SubchunkStart, 97);
+		AITEST_EQUAL("The [97-99] chunk should be third and have a length of 3", Chunks[3].Length, 3);
+
+		return true;
+	}
+};
+IMPLEMENT_AI_INSTANT_TEST(FChunkCollection_CreateBasic, "System.Mass.ChunkCollection.Create.Basic");
+
+struct FChunkCollection_CreateOrderInvariant : FChunkCollectionTestBase
+{
+	virtual bool InstantTest() override
+	{
+		TArray<FMassEntityHandle> EntitiesSubSet(&Entities[10], 30);
+		EntitiesSubSet.RemoveAt(10, 1, false);
+
+		FArchetypeChunkCollection CollectionFromOrdered(FloatsArchetype, EntitiesSubSet, FArchetypeChunkCollection::NoDuplicates);
+
+		FRandomStream Rand(0);
+		Shuffle(Rand, EntitiesSubSet);
+
+		FArchetypeChunkCollection CollectionFromRandom(FloatsArchetype, EntitiesSubSet, FArchetypeChunkCollection::NoDuplicates);
+
+		AITEST_TRUE("The resulting chunk collection should be the same regardless of the order of input entities", CollectionFromOrdered.IsSame(CollectionFromRandom));
+		
+		// just to roughly make sure the result is what we expect
+		TArrayView<const FArchetypeChunkCollection::FChunkInfo> Chunks = CollectionFromOrdered.GetChunks();
+		AITEST_EQUAL("The result should contain two chunks", Chunks.Num(), 2);
+
+		return true;
+	}
+};
+IMPLEMENT_AI_INSTANT_TEST(FChunkCollection_CreateOrderInvariant, "System.Mass.ChunkCollection.Create.OrderInvariant");
+
+struct FChunkCollection_CreateCrossChunk : FEntityTestBase
+{
+	virtual bool InstantTest() override
+	{
+		TArray<FMassEntityHandle> Entities;
+		const int32 EntitiesPerChunk = EntitySubsystem->DebugGetArchetypeEntitiesCountPerChunk(FloatsArchetype);
+
+		const int32 SpillOver = 10;
+		const int32 Count = EntitiesPerChunk + SpillOver;
+		EntitySubsystem->BatchCreateEntities(FloatsArchetype, Count, Entities);
+		
+		TArray<FMassEntityHandle> EntitiesSubCollection;
+		EntitiesSubCollection.Add(Entities[EntitiesPerChunk]);
+		for (int i = 1; i < SpillOver; ++i)
+		{
+			EntitiesSubCollection.Add(Entities[EntitiesPerChunk + i]);
+			EntitiesSubCollection.Add(Entities[EntitiesPerChunk - i]);
+		}
+		
+		FArchetypeChunkCollection ChunkCollection(FloatsArchetype, EntitiesSubCollection, FArchetypeChunkCollection::NoDuplicates);
+		TArrayView<const FArchetypeChunkCollection::FChunkInfo> Chunks = ChunkCollection.GetChunks();
+		AITEST_EQUAL("The given continuous range should get split in two", Chunks.Num(), 2);
+		AITEST_EQUAL("The part in first archetype\'s chunk should contain 9 elements", Chunks[0].Length, 9);
+		AITEST_EQUAL("The part in second archetype\'s chunk should contain 10 elements", Chunks[1].Length, 10);
+
+		return true;
+	}
+};
+IMPLEMENT_AI_INSTANT_TEST(FChunkCollection_CreateCrossChunk, "System.Mass.ChunkCollection.Create.CrossChunk");
+
+struct FChunkCollection_CreateWithDuplicates : FChunkCollectionTestBase
+{
+	virtual bool InstantTest() override
+	{
+		TArray<FMassEntityHandle> EntitiesWithDuplicates;
+		EntitiesWithDuplicates.Add(Entities[0]);
+		EntitiesWithDuplicates.Add(Entities[0]);
+		EntitiesWithDuplicates.Add(Entities[0]);
+		EntitiesWithDuplicates.Add(Entities[1]);
+		EntitiesWithDuplicates.Add(Entities[2]);
+		EntitiesWithDuplicates.Add(Entities[2]);
+
+		FArchetypeChunkCollection ChunkCollection(FloatsArchetype, EntitiesWithDuplicates, FArchetypeChunkCollection::FoldDuplicates);
+		TArrayView<const FArchetypeChunkCollection::FChunkInfo> Chunks = ChunkCollection.GetChunks();
+		AITEST_EQUAL("The result should have a single subchunk", Chunks.Num(), 1);
+		AITEST_EQUAL("The resulting subchunk should be of length 3", Chunks[0].Length, 3);
+
+		return true;
+	}
+};
+IMPLEMENT_AI_INSTANT_TEST(FChunkCollection_CreateWithDuplicates, "System.Mass.ChunkCollection.Create.Duplicates");
+
+} // FMassChunkCollectionTest
+
+PRAGMA_ENABLE_OPTIMIZATION
+
+#undef LOCTEXT_NAMESPACE
diff --git a/Engine/Plugins/Runtime/MassEntity/Source/MassEntityTestSuite/Private/MassQueryTest.cpp b/Engine/Plugins/Runtime/MassEntity/Source/MassEntityTestSuite/Private/MassQueryTest.cpp
--- a/Engine/Plugins/Runtime/MassEntity/Source/MassEntityTestSuite/Private/MassQueryTest.cpp
+++ b/Engine/Plugins/Runtime/MassEntity/Source/MassEntityTestSuite/Private/MassQueryTest.cpp
@@ -215,7 +215,8 @@ struct FQueryTest_ExecuteSparse : FEntityTestBase
 
 		FMassExecutionContext ExecContext;
 		FMassEntityQuery().AddRequirement<FTestFragment_Float>(EMassFragmentAccess::ReadWrite)
-			.ForEachEntityChunk(FArchetypeChunkCollection(FloatsArchetype, EntitiesToProcess), *EntitySubsystem, ExecContext, [&TotalProcessed](FMassExecutionContext& Context)
+			.ForEachEntityChunk(FArchetypeChunkCollection(FloatsArchetype, EntitiesToProcess, FArchetypeChunkCollection::NoDuplicates)
+								, *EntitySubsystem, ExecContext, [&TotalProcessed](FMassExecutionContext& Context)
 			{
 				TotalProcessed += Context.GetNumEntities();
 				TArrayView<FTestFragment_Float> Floats = Context.GetMutableFragmentView<FTestFragment_Float>();
diff --git a/Engine/Plugins/Runtime/MassGameplay/Source/MassReplication/Public/MassClientBubbleHandler.h b/Engine/Plugins/Runtime/MassGameplay/Source/MassReplication/Public/MassClientBubbleHandler.h
--- a/Engine/Plugins/Runtime/MassGameplay/Source/MassReplication/Public/MassClientBubbleHandler.h
+++ b/Engine/Plugins/Runtime/MassGameplay/Source/MassReplication/Public/MassClientBubbleHandler.h
@@ -550,7 +550,8 @@ void TClientBubbleHandlerBase<AgentArrayItem>::PostReplicatedAddEntitiesHelper(c
 
 		int32 AgentsSpawnIdx = 0;
 
-		Query.ForEachEntityChunk(FArchetypeChunkCollection(ArchetypeHandle, Entities), *EntitySystem, ExecContext, [&AgentsSpawn, &AgentsSpawnIdx, this, ReplicationManager, &ExecContext, &CacheFragmentViewsForSpawnQuery, &SetSpawnedEntityData, &EntitySystem](FMassExecutionContext& Context)
+		Query.ForEachEntityChunk(FArchetypeChunkCollection(ArchetypeHandle, Entities, FArchetypeChunkCollection::NoDuplicates)
+								, *EntitySystem, ExecContext, [&AgentsSpawn, &AgentsSpawnIdx, this, ReplicationManager, &ExecContext, &CacheFragmentViewsForSpawnQuery, &SetSpawnedEntityData, &EntitySystem](FMassExecutionContext& Context)
 			{
 				CacheFragmentViewsForSpawnQuery(ExecContext);
 
diff --git a/Engine/Plugins/Runtime/MassGameplay/Source/MassSignals/Private/MassSignalProcessorBase.cpp b/Engine/Plugins/Runtime/MassGameplay/Source/MassSignals/Private/MassSignalProcessorBase.cpp
--- a/Engine/Plugins/Runtime/MassGameplay/Source/MassSignals/Private/MassSignalProcessorBase.cpp
+++ b/Engine/Plugins/Runtime/MassGameplay/Source/MassSignals/Private/MassSignalProcessorBase.cpp
@@ -57,12 +57,10 @@ void UMassSignalProcessorBase::Execute(UMassEntitySubsystem& EntitySubsystem, FM
 			void Reset()
 			{
 				Entities.Reset();
-				AddedEntities.Reset();
 			}
 
 			FArchetypeHandle Archetype;
 			TArray<FMassEntityHandle> Entities;
-			TSet<FMassEntityHandle> AddedEntities;
 		};
 		TArray<FEntitySet> EntitySets;
 
@@ -103,12 +101,8 @@ void UMassSignalProcessorBase::Execute(UMassEntitySubsystem& EntitySubsystem, FM
 					FEntitySet* Set = PrevSet->Archetype == Archetype ? PrevSet : EntitySets.FindByPredicate([&Archetype](const FEntitySet& Set) { return Archetype == Set.Archetype; });
 					if (Set != nullptr)
 					{
-						// Add only unique entities in the array.
-						if (!Set->AddedEntities.Contains(Entity))
-						{
-							Set->Entities.Add(Entity);
-							Set->AddedEntities.Add(Entity);
-						}
+						// We don't care about duplicates here, the FArchetypeChunkCollection creation below will handle it
+						Set->Entities.Add(Entity);
 						SignalNameLookup.AddSignalToEntity(Entity, SignalFlag);
 						PrevSet = Set;
 					}
@@ -123,7 +117,7 @@ void UMassSignalProcessorBase::Execute(UMassEntitySubsystem& EntitySubsystem, FM
 			{
 				if (Set.Entities.Num() > 0)
 				{
-					Context.SetChunkCollection(FArchetypeChunkCollection(Set.Archetype, Set.Entities));
+					Context.SetChunkCollection(FArchetypeChunkCollection(Set.Archetype, Set.Entities, FArchetypeChunkCollection::FoldDuplicates));
 					SignalEntities(EntitySubsystem, Context, SignalNameLookup);
 					Context.ClearChunkCollection();
 				}
