Got rid of AddUnique calls in mass commands

diff --git a/Engine/Plugins/Runtime/MassEntity/Source/MassEntity/Private/MassCommandBuffer.cpp b/Engine/Plugins/Runtime/MassEntity/Source/MassEntity/Private/MassCommandBuffer.cpp
--- a/Engine/Plugins/Runtime/MassEntity/Source/MassEntity/Private/MassCommandBuffer.cpp
+++ b/Engine/Plugins/Runtime/MassEntity/Source/MassEntity/Private/MassCommandBuffer.cpp
@@ -93,7 +93,7 @@ void FMassCommandBuffer::ReplayBufferAgainstSystem(UMassEntitySubsystem* EntityS
 			if (ObservedRemoveFragments.Contains(*It.Key))
 			{
 				TArray<FArchetypeChunkCollection> ChunkCollections;
-				UE::Mass::Utils::CreateSparseChunks(*EntitySystem, It.Value, ChunkCollections);
+				UE::Mass::Utils::CreateSparseChunks(*EntitySystem, It.Value, FArchetypeChunkCollection::FoldDuplicates, ChunkCollections);
 				for (FArchetypeChunkCollection& Collection : ChunkCollections)
 				{
 					check(It.Key);
@@ -105,7 +105,7 @@ void FMassCommandBuffer::ReplayBufferAgainstSystem(UMassEntitySubsystem* EntityS
 		TArray<FArchetypeChunkCollection> EntityChunksToDestroy;
 		if (EntitiesToDestroy.Num())
 		{
-			UE::Mass::Utils::CreateSparseChunks(*EntitySystem, EntitiesToDestroy, EntityChunksToDestroy);
+			UE::Mass::Utils::CreateSparseChunks(*EntitySystem, EntitiesToDestroy, FArchetypeChunkCollection::FoldDuplicates, EntityChunksToDestroy);
 			for (FArchetypeChunkCollection& Collection : EntityChunksToDestroy)
 			{
 				ObserverManager.OnPreEntitiesDestroyed(Collection);
@@ -152,7 +152,7 @@ void FMassCommandBuffer::ReplayBufferAgainstSystem(UMassEntitySubsystem* EntityS
 		if (ObservedAddFragments.Contains(*It.Key))
 		{
 			TArray<FArchetypeChunkCollection> ChunkCollections;
-			UE::Mass::Utils::CreateSparseChunks(*EntitySystem, It.Value, ChunkCollections);
+			UE::Mass::Utils::CreateSparseChunks(*EntitySystem, It.Value, FArchetypeChunkCollection::FoldDuplicates, ChunkCollections);
 			for (FArchetypeChunkCollection& Collection : ChunkCollections)
 			{
 				check(It.Key);
diff --git a/Engine/Plugins/Runtime/MassEntity/Source/MassEntity/Private/MassEntityUtils.cpp b/Engine/Plugins/Runtime/MassEntity/Source/MassEntity/Private/MassEntityUtils.cpp
--- a/Engine/Plugins/Runtime/MassEntity/Source/MassEntity/Private/MassEntityUtils.cpp
+++ b/Engine/Plugins/Runtime/MassEntity/Source/MassEntity/Private/MassEntityUtils.cpp
@@ -34,7 +34,8 @@ EProcessorExecutionFlags GetProcessorExecutionFlagsForWold(const UWorld& World)
 	return ExecutionFlags;
 }
 
-void CreateSparseChunks(const UMassEntitySubsystem& EntitySystem, const TConstArrayView<FMassEntityHandle> Entities, TArray<FArchetypeChunkCollection>& OutChunkCollections)
+void CreateSparseChunks(const UMassEntitySubsystem& EntitySystem, const TConstArrayView<FMassEntityHandle> Entities
+	, const FArchetypeChunkCollection::EDuplicatesHandling DuplicatesHandling, TArray<FArchetypeChunkCollection>& OutChunkCollections)
 {
 	TRACE_CPUPROFILER_EVENT_SCOPE_STR("Mass_CreateSparseChunks");
 
@@ -53,7 +54,7 @@ void CreateSparseChunks(const UMassEntitySubsystem& EntitySystem, const TConstAr
 		// UMassSpawnerSubsystem::DestroyEntities' caller knowledge. Should be removed once that's addressed.
 		if (Pair.Key.IsValid())
 		{
-			OutChunkCollections.Add(FArchetypeChunkCollection(Pair.Key, Pair.Value, FArchetypeChunkCollection::NoDuplicates));
+			OutChunkCollections.Add(FArchetypeChunkCollection(Pair.Key, Pair.Value, DuplicatesHandling));
 		}
 	}
 }
diff --git a/Engine/Plugins/Runtime/MassEntity/Source/MassEntity/Public/MassCommandBuffer.h b/Engine/Plugins/Runtime/MassEntity/Source/MassEntity/Public/MassCommandBuffer.h
--- a/Engine/Plugins/Runtime/MassEntity/Source/MassEntity/Public/MassCommandBuffer.h
+++ b/Engine/Plugins/Runtime/MassEntity/Source/MassEntity/Public/MassCommandBuffer.h
@@ -89,7 +89,7 @@ struct MASSENTITY_API FBuildEntityFromFragmentInstance : public FCommandBufferEn
 
 	void GetAffectedTypes(TMap<const UScriptStruct*, TArray<FMassEntityHandle>>& OutAdded, TMap<const UScriptStruct*, TArray<FMassEntityHandle>>& OutRemoved)
 	{
-		OutAdded.FindOrAdd(Struct.GetScriptStruct()).AddUnique(TargetEntity);
+		OutAdded.FindOrAdd(Struct.GetScriptStruct()).Add(TargetEntity);
 	}
 
 protected:
@@ -121,7 +121,7 @@ struct MASSENTITY_API FBuildEntityFromFragmentInstances : public FCommandBufferE
 	{
 		for (const FInstancedStruct& Struct : Instances)
 		{
-			OutAdded.FindOrAdd(Struct.GetScriptStruct()).AddUnique(TargetEntity);
+			OutAdded.FindOrAdd(Struct.GetScriptStruct()).Add(TargetEntity);
 		}
 	}
 
@@ -152,7 +152,7 @@ struct MASSENTITY_API FCommandAddFragment : public FCommandBufferEntryBase
 
 	void GetAffectedTypes(TMap<const UScriptStruct*, TArray<FMassEntityHandle>>& OutAdded, TMap<const UScriptStruct*, TArray<FMassEntityHandle>>& OutRemoved)
 	{
-		OutAdded.FindOrAdd(StructParam).AddUnique(TargetEntity);
+		OutAdded.FindOrAdd(StructParam).Add(TargetEntity);
 	}
 
 protected:
@@ -182,7 +182,7 @@ struct MASSENTITY_API FCommandAddFragmentInstance : public FCommandBufferEntryBa
 
 	void GetAffectedTypes(TMap<const UScriptStruct*, TArray<FMassEntityHandle>>& OutAdded, TMap<const UScriptStruct*, TArray<FMassEntityHandle>>& OutRemoved)
 	{
-		OutAdded.FindOrAdd(Struct.GetScriptStruct()).AddUnique(TargetEntity);
+		OutAdded.FindOrAdd(Struct.GetScriptStruct()).Add(TargetEntity);
 	}
 
 protected:
@@ -215,7 +215,7 @@ struct MASSENTITY_API FMassCommandAddFragmentInstanceList : public FCommandBuffe
 	{
 		for (const FInstancedStruct& Struct : FragmentList)
 		{
-			OutAdded.FindOrAdd(Struct.GetScriptStruct()).AddUnique(TargetEntity);
+			OutAdded.FindOrAdd(Struct.GetScriptStruct()).Add(TargetEntity);
 		}
 	}
 
@@ -246,7 +246,7 @@ struct MASSENTITY_API FCommandRemoveFragment : public FCommandBufferEntryBase
 
 	void GetAffectedTypes(TMap<const UScriptStruct*, TArray<FMassEntityHandle>>& OutAdded, TMap<const UScriptStruct*, TArray<FMassEntityHandle>>& OutRemoved)
 	{
-		OutRemoved.FindOrAdd(StructParam).AddUnique(TargetEntity);
+		OutRemoved.FindOrAdd(StructParam).Add(TargetEntity);
 	}
 
 protected:
@@ -278,7 +278,7 @@ struct MASSENTITY_API FCommandAddFragmentList : public FCommandBufferEntryBase
 	{
 		for (const UScriptStruct* StructParam : FragmentList)
 		{
-			OutAdded.FindOrAdd(StructParam).AddUnique(TargetEntity);
+			OutAdded.FindOrAdd(StructParam).Add(TargetEntity);
 		}
 	}
 
@@ -311,7 +311,7 @@ struct MASSENTITY_API FCommandRemoveFragmentList : public FCommandBufferEntryBas
 	{
 		for (const UScriptStruct* StructParam : FragmentList)
 		{
-			OutRemoved.FindOrAdd(StructParam).AddUnique(TargetEntity);
+			OutRemoved.FindOrAdd(StructParam).Add(TargetEntity);
 		}
 	}
 
@@ -431,7 +431,7 @@ struct MASSENTITY_API FCommandRemoveComposition : public FCommandBufferEntryBase
 			Descriptor.Fragments.ExportTypes(Fragments);
 			for (const UScriptStruct* StructParam : Fragments)
 			{
-				OutRemoved.FindOrAdd(StructParam).AddUnique(TargetEntity);
+				OutRemoved.FindOrAdd(StructParam).Add(TargetEntity);
 			}
 		}
 	}
diff --git a/Engine/Plugins/Runtime/MassEntity/Source/MassEntity/Public/MassEntityUtils.h b/Engine/Plugins/Runtime/MassEntity/Source/MassEntity/Public/MassEntityUtils.h
--- a/Engine/Plugins/Runtime/MassEntity/Source/MassEntity/Public/MassEntityUtils.h
+++ b/Engine/Plugins/Runtime/MassEntity/Source/MassEntity/Public/MassEntityUtils.h
@@ -3,11 +3,11 @@
 #pragma once
 
 #include "MassProcessingTypes.h"
+#include "MassArchetypeTypes.h"
 
 class UWorld;
 class UMassEntitySubsystem;
 struct FMassEntityHandle;
-struct FArchetypeChunkCollection;
 
 namespace UE::Mass::Utils
 {
@@ -17,8 +17,9 @@ MASSENTITY_API extern EProcessorExecutionFlags GetProcessorExecutionFlagsForWold
 
 /** 
  * Fills OutChunkCollections with per-archetype FArchetypeChunkCollection instances. 
- * Note that it's caller responsibility to guarantee Entities does not contain duplicates.
+ * @param DuplicatesHandling used to inform the function whether to expect duplicates.
  */
-MASSENTITY_API extern void CreateSparseChunks(const UMassEntitySubsystem& EntitySystem, const TConstArrayView<FMassEntityHandle> Entities, TArray<FArchetypeChunkCollection>& OutChunkCollections);
+MASSENTITY_API extern void CreateSparseChunks(const UMassEntitySubsystem& EntitySystem, const TConstArrayView<FMassEntityHandle> Entities
+	, const FArchetypeChunkCollection::EDuplicatesHandling DuplicatesHandling, TArray<FArchetypeChunkCollection>& OutChunkCollections);
 
 } // namespace UE::Mass::Utils
\ No newline at end of file
diff --git a/Engine/Plugins/Runtime/MassGameplay/Source/MassSpawner/Private/MassSpawnerSubsystem.cpp b/Engine/Plugins/Runtime/MassGameplay/Source/MassSpawner/Private/MassSpawnerSubsystem.cpp
--- a/Engine/Plugins/Runtime/MassGameplay/Source/MassSpawner/Private/MassSpawnerSubsystem.cpp
+++ b/Engine/Plugins/Runtime/MassGameplay/Source/MassSpawner/Private/MassSpawnerSubsystem.cpp
@@ -117,7 +117,7 @@ void UMassSpawnerSubsystem::DestroyEntities(const FMassEntityTemplateID Template
 
 
 	TArray<FArchetypeChunkCollection> ChunkCollections;
-	UE::Mass::Utils::CreateSparseChunks(*EntitySystem, Entities, ChunkCollections);
+	UE::Mass::Utils::CreateSparseChunks(*EntitySystem, Entities, FArchetypeChunkCollection::NoDuplicates, ChunkCollections);
 
 	for (const FArchetypeChunkCollection& Chunks : ChunkCollections)
 	{
