A quick implementation of extracting relevant fragments to add/remove information from Mass commands, based on cached "last used type" data.

diff --git a/Engine/Plugins/Runtime/MassEntity/Source/MassEntity/Private/MassCommandBuffer.cpp b/Engine/Plugins/Runtime/MassEntity/Source/MassEntity/Private/MassCommandBuffer.cpp
--- a/Engine/Plugins/Runtime/MassEntity/Source/MassEntity/Private/MassCommandBuffer.cpp
+++ b/Engine/Plugins/Runtime/MassEntity/Source/MassEntity/Private/MassCommandBuffer.cpp
@@ -8,9 +8,6 @@
 #include "VisualLogger/VisualLogger.h"
 
 
-//////////////////////////////////////////////////////////////////////
-// FMassCommandBuffer
-
 CSV_DEFINE_CATEGORY(MassEntities, true);
 CSV_DEFINE_CATEGORY(MassEntitiesCounters, true);
 
@@ -70,6 +67,19 @@ void GetCommandStatNames(FStructView Entry, FString& OutName, ANSIName*& OutANSI
 #endif
 } // UE::FLWCCommand
 
+//////////////////////////////////////////////////////////////////////
+// FMassCommandsObservedTypes
+
+void FMassCommandsObservedTypes::Reset()
+{
+	FragmentsToAdd.Reset();
+	FragmentsToRemove.Reset();
+}
+
+
+//////////////////////////////////////////////////////////////////////
+// FMassCommandBuffer
+
 void FMassCommandBuffer::ReplayBufferAgainstSystem(UMassEntitySubsystem* EntitySystem)
 {
 	check(EntitySystem);
@@ -87,7 +97,7 @@ void FMassCommandBuffer::ReplayBufferAgainstSystem(UMassEntitySubsystem* EntityS
 
 	if (ObservedRemoveFragments.IsEmpty() == false)
 	{
-		for (auto It : FragmentsToRemove)
+		for (auto It : ObservedTypes.GetFragmentsToRemove())
 		{
 			check(It.Key);
 			if (ObservedRemoveFragments.Contains(*It.Key))
@@ -108,7 +118,6 @@ void FMassCommandBuffer::ReplayBufferAgainstSystem(UMassEntitySubsystem* EntityS
 			UE::Mass::Utils::CreateSparseChunks(*EntitySystem, EntitiesToDestroy, FArchetypeChunkCollection::FoldDuplicates, EntityChunksToDestroy);
 			for (FArchetypeChunkCollection& Collection : EntityChunksToDestroy)
 			{
-				ObserverManager.OnPreEntitiesDestroyed(Collection);
 				EntitySystem->BatchDestroyEntityChunks(Collection);
 			}
 		}
@@ -146,7 +155,7 @@ void FMassCommandBuffer::ReplayBufferAgainstSystem(UMassEntitySubsystem* EntityS
 	// Using Clear() instead of Reset(), as otherwise the chunks moved into the PendingCommands in MoveAppend() can accumulate.
 	PendingCommands.Clear();
 
-	for (auto It : FragmentsToAdd)
+	for (auto It : ObservedTypes.GetFragmentsToAdd())
 	{
 		check(It.Key);
 		if (ObservedAddFragments.Contains(*It.Key))
@@ -161,8 +170,7 @@ void FMassCommandBuffer::ReplayBufferAgainstSystem(UMassEntitySubsystem* EntityS
 		}
 	}
 
-	FragmentsToAdd.Reset();
-	FragmentsToRemove.Reset();
+	ObservedTypes.Reset();
 }
 
 void FMassCommandBuffer::MoveAppend(FMassCommandBuffer& Other)
diff --git a/Engine/Plugins/Runtime/MassEntity/Source/MassEntity/Public/MassCommandBuffer.h b/Engine/Plugins/Runtime/MassEntity/Source/MassEntity/Public/MassCommandBuffer.h
--- a/Engine/Plugins/Runtime/MassEntity/Source/MassEntity/Public/MassCommandBuffer.h
+++ b/Engine/Plugins/Runtime/MassEntity/Source/MassEntity/Public/MassCommandBuffer.h
@@ -21,6 +21,67 @@ namespace ECommandBufferOperationType
 // (e.g., which system added the command, or even file/line number in development builds for the specific call via a macro)
 //@TODO: Support more commands
 
+struct FMassObservedTypeCollection
+{
+	void Add(const UScriptStruct& TypeType, FMassEntityHandle Entity)
+	{
+		if (&TypeType != LastAddedType)
+		{
+			LastAddedType = &TypeType;
+			LastUsedCollection = &Types.FindOrAdd(&TypeType);
+		}
+		CA_ASSUME(LastUsedCollection);
+		LastUsedCollection->Add(Entity);
+	}
+
+	void Reset()
+	{
+		LastAddedType = nullptr;
+		LastUsedCollection = nullptr;
+		Types.Reset();
+	}
+
+	const TMap<const UScriptStruct*, TArray<FMassEntityHandle>>& GetTypes() const 
+	{ 
+		return Types; 
+	}
+
+private:
+	const UScriptStruct* LastAddedType = nullptr;
+	TArray<FMassEntityHandle>* LastUsedCollection = nullptr;
+	TMap<const UScriptStruct*, TArray<FMassEntityHandle>> Types;
+};
+
+struct FMassCommandsObservedTypes
+{
+	void Reset();
+	void FragmentAdded(const UScriptStruct* TypeType, FMassEntityHandle Entity)
+	{
+		check(TypeType);
+		FragmentsToAdd.Add(*TypeType, Entity);
+	}
+	void FragmentRemoved(const UScriptStruct* TypeType, FMassEntityHandle Entity)
+	{
+		check(TypeType);
+		FragmentsToRemove.Add(*TypeType, Entity);
+	}
+
+	const TMap<const UScriptStruct*, TArray<FMassEntityHandle>>& GetFragmentsToAdd() const
+	{
+		return FragmentsToAdd.GetTypes();
+	}
+
+	const TMap<const UScriptStruct*, TArray<FMassEntityHandle>>& GetFragmentsToRemove() const
+	{
+		return FragmentsToRemove.GetTypes();
+	}
+
+protected:
+	FMassObservedTypeCollection FragmentsToAdd;
+	FMassObservedTypeCollection FragmentsToRemove;
+};
+
+
 USTRUCT()
 struct MASSENTITY_API FCommandBufferEntryBase
 {
@@ -35,7 +96,7 @@ struct MASSENTITY_API FCommandBufferEntryBase
 		: TargetEntity(InEntity)
 	{}
 
-	void GetAffectedTypes(TMap<const UScriptStruct*, TArray<FMassEntityHandle>>& OutAdded, TMap<const UScriptStruct*, TArray<FMassEntityHandle>>& OutRemoved) { ensure(false); }
+	void AppendAffectedEntitesPerType(FMassCommandsObservedTypes& ObservedTypes) { ensure(false); }
 	virtual void Execute(UMassEntitySubsystem& System) const PURE_VIRTUAL(FCommandBufferEntryBase::Execute, );
 };
 template<> struct TStructOpsTypeTraits<FCommandBufferEntryBase> : public TStructOpsTypeTraitsBase2<FCommandBufferEntryBase> { enum { WithPureVirtual = true, }; };
@@ -87,9 +148,9 @@ struct MASSENTITY_API FBuildEntityFromFragmentInstance : public FCommandBufferEn
 		, Struct(InStruct)
 	{}
 
-	void GetAffectedTypes(TMap<const UScriptStruct*, TArray<FMassEntityHandle>>& OutAdded, TMap<const UScriptStruct*, TArray<FMassEntityHandle>>& OutRemoved)
+	void AppendAffectedEntitesPerType(FMassCommandsObservedTypes& ObservedTypes)
 	{
-		OutAdded.FindOrAdd(Struct.GetScriptStruct()).Add(TargetEntity);
+		ObservedTypes.FragmentAdded(Struct.GetScriptStruct(), TargetEntity);
 	}
 
 protected:
@@ -117,11 +178,11 @@ struct MASSENTITY_API FBuildEntityFromFragmentInstances : public FCommandBufferE
 		, Instances(InInstances)
 	{}
 
-	void GetAffectedTypes(TMap<const UScriptStruct*, TArray<FMassEntityHandle>>& OutAdded, TMap<const UScriptStruct*, TArray<FMassEntityHandle>>& OutRemoved)
+	void AppendAffectedEntitesPerType(FMassCommandsObservedTypes& ObservedTypes)
 	{
 		for (const FInstancedStruct& Struct : Instances)
 		{
-			OutAdded.FindOrAdd(Struct.GetScriptStruct()).Add(TargetEntity);
+			ObservedTypes.FragmentAdded(Struct.GetScriptStruct(), TargetEntity);
 		}
 	}
 
@@ -150,9 +211,9 @@ struct MASSENTITY_API FCommandAddFragment : public FCommandBufferEntryBase
 		, StructParam(InStruct)
 	{}
 
-	void GetAffectedTypes(TMap<const UScriptStruct*, TArray<FMassEntityHandle>>& OutAdded, TMap<const UScriptStruct*, TArray<FMassEntityHandle>>& OutRemoved)
+	void AppendAffectedEntitesPerType(FMassCommandsObservedTypes& ObservedTypes)
 	{
-		OutAdded.FindOrAdd(StructParam).Add(TargetEntity);
+		ObservedTypes.FragmentAdded(StructParam, TargetEntity);
 	}
 
 protected:
@@ -180,9 +241,9 @@ struct MASSENTITY_API FCommandAddFragmentInstance : public FCommandBufferEntryBa
         , Struct(InStruct)
 	{}
 
-	void GetAffectedTypes(TMap<const UScriptStruct*, TArray<FMassEntityHandle>>& OutAdded, TMap<const UScriptStruct*, TArray<FMassEntityHandle>>& OutRemoved)
+	void AppendAffectedEntitesPerType(FMassCommandsObservedTypes& ObservedTypes)
 	{
-		OutAdded.FindOrAdd(Struct.GetScriptStruct()).Add(TargetEntity);
+		ObservedTypes.FragmentAdded(Struct.GetScriptStruct(), TargetEntity);
 	}
 
 protected:
@@ -199,7 +260,7 @@ struct MASSENTITY_API FMassCommandAddFragmentInstanceList : public FCommandBuffe
 {
 	GENERATED_BODY()
 
-		enum
+	enum
 	{
 		Type = ECommandBufferOperationType::Add
 	};
@@ -211,11 +272,11 @@ struct MASSENTITY_API FMassCommandAddFragmentInstanceList : public FCommandBuffe
 		, FragmentList(InitList)
 	{}
 
-	void GetAffectedTypes(TMap<const UScriptStruct*, TArray<FMassEntityHandle>>& OutAdded, TMap<const UScriptStruct*, TArray<FMassEntityHandle>>& OutRemoved)
+	void AppendAffectedEntitesPerType(FMassCommandsObservedTypes& ObservedTypes)
 	{
 		for (const FInstancedStruct& Struct : FragmentList)
 		{
-			OutAdded.FindOrAdd(Struct.GetScriptStruct()).Add(TargetEntity);
+			ObservedTypes.FragmentAdded(Struct.GetScriptStruct(), TargetEntity);
 		}
 	}
 
@@ -244,9 +305,9 @@ struct MASSENTITY_API FCommandRemoveFragment : public FCommandBufferEntryBase
 		, StructParam(InStruct)
 	{}
 
-	void GetAffectedTypes(TMap<const UScriptStruct*, TArray<FMassEntityHandle>>& OutAdded, TMap<const UScriptStruct*, TArray<FMassEntityHandle>>& OutRemoved)
+	void AppendAffectedEntitesPerType(FMassCommandsObservedTypes& ObservedTypes)
 	{
-		OutRemoved.FindOrAdd(StructParam).Add(TargetEntity);
+		ObservedTypes.FragmentRemoved(StructParam, TargetEntity);
 	}
 
 protected:
@@ -274,11 +335,11 @@ struct MASSENTITY_API FCommandAddFragmentList : public FCommandBufferEntryBase
 		, FragmentList(InFragmentList)
 	{}
 
-	void GetAffectedTypes(TMap<const UScriptStruct*, TArray<FMassEntityHandle>>& OutAdded, TMap<const UScriptStruct*, TArray<FMassEntityHandle>>& OutRemoved)
+	void AppendAffectedEntitesPerType(FMassCommandsObservedTypes& ObservedTypes)
 	{
 		for (const UScriptStruct* StructParam : FragmentList)
 		{
-			OutAdded.FindOrAdd(StructParam).Add(TargetEntity);
+			ObservedTypes.FragmentAdded(StructParam, TargetEntity);
 		}
 	}
 
@@ -307,11 +368,11 @@ struct MASSENTITY_API FCommandRemoveFragmentList : public FCommandBufferEntryBas
 		, FragmentList(InFragmentList)
 	{}
 
-	void GetAffectedTypes(TMap<const UScriptStruct*, TArray<FMassEntityHandle>>& OutAdded, TMap<const UScriptStruct*, TArray<FMassEntityHandle>>& OutRemoved)
+	void AppendAffectedEntitesPerType(FMassCommandsObservedTypes& ObservedTypes)
 	{
 		for (const UScriptStruct* StructParam : FragmentList)
 		{
-			OutRemoved.FindOrAdd(StructParam).Add(TargetEntity);
+			ObservedTypes.FragmentRemoved(StructParam, TargetEntity);
 		}
 	}
 
@@ -421,7 +482,7 @@ struct MASSENTITY_API FCommandRemoveComposition : public FCommandBufferEntryBase
 		, Descriptor(InDescriptor)
 	{}
 
-	void GetAffectedTypes(TMap<const UScriptStruct*, TArray<FMassEntityHandle>>& OutAdded, TMap<const UScriptStruct*, TArray<FMassEntityHandle>>& OutRemoved)
+	void AppendAffectedEntitesPerType(FMassCommandsObservedTypes& ObservedTypes)
 	{
 		if (Descriptor.Fragments.IsEmpty() == false)
 		{
@@ -431,7 +492,7 @@ struct MASSENTITY_API FCommandRemoveComposition : public FCommandBufferEntryBase
 			Descriptor.Fragments.ExportTypes(Fragments);
 			for (const UScriptStruct* StructParam : Fragments)
 			{
-				OutRemoved.FindOrAdd(StructParam).Add(TargetEntity);
+				ObservedTypes.FragmentRemoved(StructParam, TargetEntity);
 			}
 		}
 	}
@@ -469,7 +530,7 @@ public:
 		T& Command = PendingCommands.Emplace_GetRef<T>(Forward<TArgs>(InArgs)...); 
 		if (constexpr bool bIsModifyingComposition = ((T::Type & (ECommandBufferOperationType::Add | ECommandBufferOperationType::Remove)) != 0))
 		{	
-			Command.GetAffectedTypes(FragmentsToAdd, FragmentsToRemove);
+			Command.AppendAffectedEntitesPerType(ObservedTypes);
 		}
 		return Command;
 	}
@@ -533,6 +594,5 @@ private:
 
 	TArray<FMassEntityHandle> EntitiesToDestroy;
 
-	TMap<const UScriptStruct*, TArray<FMassEntityHandle>> FragmentsToAdd;
-	TMap<const UScriptStruct*, TArray<FMassEntityHandle>> FragmentsToRemove;
+	FMassCommandsObservedTypes ObservedTypes;
 };
