Fixed FMassArchetypeSharedFragmentValues hash calculation

diff --git a/Engine/Plugins/Runtime/MassEntity/Source/MassEntity/Private/MassEntityTypes.cpp b/Engine/Plugins/Runtime/MassEntity/Source/MassEntity/Private/MassEntityTypes.cpp
--- a/Engine/Plugins/Runtime/MassEntity/Source/MassEntity/Private/MassEntityTypes.cpp
+++ b/Engine/Plugins/Runtime/MassEntity/Source/MassEntity/Private/MassEntityTypes.cpp
@@ -11,12 +11,12 @@ uint32 FMassArchetypeSharedFragmentValues::CalculateHash() const
 	uint32 Hash = 0;
 	for (const FConstSharedStruct& Fragment : ConstSharedFragments)
 	{
-		PointerHash(Fragment.GetMemory(), Hash);
+		Hash = PointerHash(Fragment.GetMemory(), Hash);
 	}
 
 	for (const FSharedStruct& Fragment : SharedFragments)
 	{
-		PointerHash(Fragment.GetMemory(), Hash);
+		Hash = PointerHash(Fragment.GetMemory(), Hash);
 	}
 
 	return Hash;
