[MassEntity] allow access to the shared fragments from FMassEntityView
BuildEntity commands can now setup SharedFragments

diff --git a/Engine/Plugins/Runtime/MassEntity/Source/MassEntity/Private/MassCommandBuffer.cpp b/Engine/Plugins/Runtime/MassEntity/Source/MassEntity/Private/MassCommandBuffer.cpp
--- a/Engine/Plugins/Runtime/MassEntity/Source/MassEntity/Private/MassCommandBuffer.cpp
+++ b/Engine/Plugins/Runtime/MassEntity/Source/MassEntity/Private/MassCommandBuffer.cpp
@@ -222,12 +222,12 @@ void FCommandSwapTags::Execute(UMassEntitySubsystem& System) const
 
 void FBuildEntityFromFragmentInstance::Execute(UMassEntitySubsystem& System) const
 {
-	System.BuildEntity(TargetEntity, MakeArrayView(&Struct, 1));
+	System.BuildEntity(TargetEntity, MakeArrayView(&Struct, 1), SharedFragmentValues);
 }
 
 void FBuildEntityFromFragmentInstances::Execute(UMassEntitySubsystem& System) const
 {
-	System.BuildEntity(TargetEntity, Instances);
+	System.BuildEntity(TargetEntity, Instances, SharedFragmentValues);
 }
 
 void FCommandAddFragment::Execute(UMassEntitySubsystem& System) const
diff --git a/Engine/Plugins/Runtime/MassEntity/Source/MassEntity/Private/MassEntitySubsystem.cpp b/Engine/Plugins/Runtime/MassEntity/Source/MassEntity/Private/MassEntitySubsystem.cpp
--- a/Engine/Plugins/Runtime/MassEntity/Source/MassEntity/Private/MassEntitySubsystem.cpp
+++ b/Engine/Plugins/Runtime/MassEntity/Source/MassEntity/Private/MassEntitySubsystem.cpp
@@ -85,14 +85,14 @@ void UMassEntitySubsystem::Deinitialize()
 	//@TODO: Should actually do this...
 }
 
-FArchetypeHandle UMassEntitySubsystem::CreateArchetype(TConstArrayView<const UScriptStruct*> FragmentsAngTagsList)
+FArchetypeHandle UMassEntitySubsystem::CreateArchetype(TConstArrayView<const UScriptStruct*> FragmentsAndTagsList)
 {
 	FMassChunkFragmentBitSet ChunkFragments;
 	FMassTagBitSet Tags;
 	TArray<const UScriptStruct*, TInlineAllocator<16>> FragmentList;
-	FragmentList.Reserve(FragmentsAngTagsList.Num());
+	FragmentList.Reserve(FragmentsAndTagsList.Num());
 
-	for (const UScriptStruct* Type : FragmentsAngTagsList)
+	for (const UScriptStruct* Type : FragmentsAndTagsList)
 	{
 		if (Type->IsChildOf(FMassFragment::StaticStruct()))
 		{
@@ -133,7 +133,7 @@ FArchetypeHandle UMassEntitySubsystem::CreateArchetype(const FMassArchetypeCompo
 	TArray<TSharedPtr<FMassArchetypeData>>& HashRow = FragmentHashToArchetypeMap.FindOrAdd(TypeHash);
 
 	FArchetypeHandle Result;
-	for (TSharedPtr<FMassArchetypeData>& Ptr : HashRow)
+	for (const TSharedPtr<FMassArchetypeData>& Ptr : HashRow)
 	{
 		if (Ptr->IsEquivalent(Composition, SharedFragmentValues))
 		{
@@ -145,7 +145,7 @@ FArchetypeHandle UMassEntitySubsystem::CreateArchetype(const FMassArchetypeCompo
 	if (!Result.DataPtr.IsValid())
 	{
 		// Create a new archetype
-		TSharedPtr<FMassArchetypeData> NewArchetype = MakeShareable(new FMassArchetypeData);
+		const TSharedPtr<FMassArchetypeData> NewArchetype = MakeShareable(new FMassArchetypeData);
 		NewArchetype->Initialize(Composition, SharedFragmentValues);
 		HashRow.Add(NewArchetype);
 
@@ -167,13 +167,13 @@ FArchetypeHandle UMassEntitySubsystem::InternalCreateSiblingArchetype(const TSha
 {
 	checkSlow(SourceArchetype.IsValid());
 	const FMassArchetypeData& SourceArchetypeRef = *SourceArchetype.Get();
-	FMassArchetypeCompositionDescriptor NewComposition(SourceArchetypeRef.GetFragmentBitSet(), OverrideTags, SourceArchetypeRef.GetChunkFragmentBitSet(), SourceArchetypeRef.GetSharedFragmentBitSet());
+	const FMassArchetypeCompositionDescriptor NewComposition(SourceArchetypeRef.GetFragmentBitSet(), OverrideTags, SourceArchetypeRef.GetChunkFragmentBitSet(), SourceArchetypeRef.GetSharedFragmentBitSet());
 	const uint32 TypeHash = NewComposition.CalculateHash();
 
 	TArray<TSharedPtr<FMassArchetypeData>>& HashRow = FragmentHashToArchetypeMap.FindOrAdd(TypeHash);
 
 	FArchetypeHandle Result;
-	for (TSharedPtr<FMassArchetypeData>& Ptr : HashRow)
+	for (const TSharedPtr<FMassArchetypeData>& Ptr : HashRow)
 	{
 		if (Ptr->IsEquivalent(NewComposition, SourceArchetypeRef.GetSharedFragmentValues()))
 		{
@@ -185,7 +185,7 @@ FArchetypeHandle UMassEntitySubsystem::InternalCreateSiblingArchetype(const TSha
 	if (!Result.DataPtr.IsValid())
 	{
 		// Create a new archetype
-		TSharedPtr<FMassArchetypeData> NewArchetype = MakeShareable(new FMassArchetypeData);
+		const TSharedPtr<FMassArchetypeData> NewArchetype = MakeShareable(new FMassArchetypeData);
 		NewArchetype->InitializeWithSibling(SourceArchetypeRef, OverrideTags);
 		HashRow.Add(NewArchetype);
 
@@ -248,7 +248,7 @@ FMassEntityHandle UMassEntitySubsystem::CreateEntity(const FArchetypeHandle Arch
 {
 	check(Archetype.DataPtr.IsValid());
 
-	FMassEntityHandle Entity = ReserveEntity();
+	const FMassEntityHandle Entity = ReserveEntity();
 	InternalBuildEntity(Entity, Archetype);
 	return Entity;
 }
@@ -257,13 +257,13 @@ FMassEntityHandle UMassEntitySubsystem::CreateEntity(TConstArrayView<FInstancedS
 {
 	check(FragmentInstanceList.Num() > 0);
 
-	FArchetypeHandle Archetype = CreateArchetype(FMassArchetypeCompositionDescriptor(FragmentInstanceList, FMassTagBitSet(), FMassChunkFragmentBitSet(), FMassSharedFragmentBitSet()), FMassArchetypeSharedFragmentValues());
+	const FArchetypeHandle Archetype = CreateArchetype(FMassArchetypeCompositionDescriptor(FragmentInstanceList, FMassTagBitSet(), FMassChunkFragmentBitSet(), FMassSharedFragmentBitSet()), FMassArchetypeSharedFragmentValues());
 	check(Archetype.DataPtr.IsValid());
 
-	FMassEntityHandle Entity = ReserveEntity();
+	const FMassEntityHandle Entity = ReserveEntity();
 	InternalBuildEntity(Entity, Archetype);
 
-	FEntityData& EntityData = Entities[Entity.Index];
+	const FEntityData& EntityData = Entities[Entity.Index];
 	EntityData.CurrentArchetype->SetFragmentsData(Entity, FragmentInstanceList);
 
 	return Entity;
@@ -295,18 +295,29 @@ void UMassEntitySubsystem::BuildEntity(FMassEntityHandle Entity, FArchetypeHandl
 	InternalBuildEntity(Entity, Archetype);
 }
 
-void UMassEntitySubsystem::BuildEntity(FMassEntityHandle Entity, TConstArrayView<FInstancedStruct> FragmentInstanceList)
+void UMassEntitySubsystem::BuildEntity(FMassEntityHandle Entity, TConstArrayView<FInstancedStruct> FragmentInstanceList, FMassArchetypeSharedFragmentValues SharedFragmentValues)
 {
 	check(FragmentInstanceList.Num() > 0);
 	checkf(!IsEntityBuilt(Entity), TEXT("Expecting an entity that is not already built"));
 
-	const FMassArchetypeCompositionDescriptor Composition(FragmentInstanceList, FMassTagBitSet(), FMassChunkFragmentBitSet(), FMassSharedFragmentBitSet());
-	FArchetypeHandle Archetype = CreateArchetype(Composition, FMassArchetypeSharedFragmentValues());
+	SharedFragmentValues.Sort();
+	SharedFragmentValues.CacheHash();
+	FMassArchetypeCompositionDescriptor Composition(FragmentInstanceList, FMassTagBitSet(), FMassChunkFragmentBitSet(), FMassSharedFragmentBitSet());
+	for (const FConstSharedStruct& SharedFragment : SharedFragmentValues.GetConstSharedFragments())
+	{
+		Composition.SharedFragments.Add(*SharedFragment.GetScriptStruct());
+	}
+	for (const FSharedStruct& SharedFragment : SharedFragmentValues.GetSharedFragments())
+	{
+		Composition.SharedFragments.Add(*SharedFragment.GetScriptStruct());
+	}
+
+	const FArchetypeHandle Archetype = CreateArchetype(Composition, SharedFragmentValues);
 	check(Archetype.DataPtr.IsValid());
 
 	InternalBuildEntity(Entity, Archetype);
 
-	FEntityData& EntityData = Entities[Entity.Index];
+	const FEntityData& EntityData = Entities[Entity.Index];
 	EntityData.CurrentArchetype->SetFragmentsData(Entity, FragmentInstanceList);
 }
 
@@ -353,7 +364,7 @@ void UMassEntitySubsystem::DestroyEntity(FMassEntityHandle Entity)
 	
 	CheckIfEntityIsActive(Entity);
 
-	FEntityData& EntityData = Entities[Entity.Index];
+	const FEntityData& EntityData = Entities[Entity.Index];
 	FMassArchetypeData* Archetype = EntityData.CurrentArchetype.Get();
 	check(Archetype);
 	Archetype->RemoveEntity(Entity);
@@ -537,7 +548,7 @@ void UMassEntitySubsystem::InternalReleaseEntity(FMassEntityHandle Entity)
 
 void UMassEntitySubsystem::InternalAddFragmentListToEntityChecked(FMassEntityHandle Entity, const FMassFragmentBitSet& InFragments)
 {
-	FEntityData& EntityData = Entities[Entity.Index];
+	const FEntityData& EntityData = Entities[Entity.Index];
 	FMassArchetypeData* OldArchetype = EntityData.CurrentArchetype.Get();
 	check(OldArchetype);
 
@@ -581,7 +592,7 @@ void UMassEntitySubsystem::AddFragmentInstanceListToEntity(FMassEntityHandle Ent
 
 	InternalAddFragmentListToEntityChecked(Entity, FMassFragmentBitSet(FragmentInstanceList));
 
-	FEntityData& EntityData = Entities[Entity.Index];
+	const FEntityData& EntityData = Entities[Entity.Index];
 	EntityData.CurrentArchetype->SetFragmentsData(Entity, FragmentInstanceList);
 }
 
@@ -635,7 +646,7 @@ void UMassEntitySubsystem::SwapTagsForEntity(FMassEntityHandle Entity, const USc
 	
 	if (NewTagBitSet != CurrentArchetype->GetTagBitSet())
 	{
-		FArchetypeHandle NewArchetypeHandle = InternalCreateSiblingArchetype(EntityData.CurrentArchetype, NewTagBitSet);
+		const FArchetypeHandle NewArchetypeHandle = InternalCreateSiblingArchetype(EntityData.CurrentArchetype, NewTagBitSet);
 		checkSlow(NewArchetypeHandle.IsValid());
 
 		// Move the entity over
@@ -659,7 +670,7 @@ void UMassEntitySubsystem::AddTagToEntity(FMassEntityHandle Entity, const UScrip
 		//FMassTagBitSet NewTags = CurrentArchetype->GetTagBitSet() - *TagType;
 		FMassTagBitSet NewTags = CurrentArchetype->GetTagBitSet();
 		NewTags.Add(*TagType);
-		FArchetypeHandle NewArchetypeHandle = InternalCreateSiblingArchetype(EntityData.CurrentArchetype, NewTags);
+		const FArchetypeHandle NewArchetypeHandle = InternalCreateSiblingArchetype(EntityData.CurrentArchetype, NewTags);
 		checkSlow(NewArchetypeHandle.IsValid());
 
 		// Move the entity over
@@ -683,7 +694,7 @@ void UMassEntitySubsystem::RemoveTagFromEntity(FMassEntityHandle Entity, const U
 		// CurrentArchetype->GetTagBitSet() -  *TagType
 		FMassTagBitSet NewTags = CurrentArchetype->GetTagBitSet();
 		NewTags.Remove(*TagType);
-		FArchetypeHandle NewArchetypeHandle = InternalCreateSiblingArchetype(EntityData.CurrentArchetype, NewTags);
+		const FArchetypeHandle NewArchetypeHandle = InternalCreateSiblingArchetype(EntityData.CurrentArchetype, NewTags);
 		checkSlow(NewArchetypeHandle.IsValid());
 
 		// Move the entity over
@@ -710,7 +721,7 @@ void UMassEntitySubsystem::SetEntityFragmentsValues(FMassEntityHandle Entity, TA
 {
 	CheckIfEntityIsActive(Entity);
 
-	FEntityData& EntityData = Entities[Entity.Index];
+	const FEntityData& EntityData = Entities[Entity.Index];
 	EntityData.CurrentArchetype->SetFragmentsData(Entity, FragmentInstanceList);
 }
 
@@ -1066,7 +1077,7 @@ void UMassEntitySubsystem::DebugGetArchetypeStrings(const FArchetypeHandle& Arch
 		return;
 	}
 
-	FMassArchetypeData& ArchetypeRef = *Archetype.DataPtr.Get();
+	const FMassArchetypeData& ArchetypeRef = *Archetype.DataPtr.Get();
 	
 	OutFragmentNames.Reserve(ArchetypeRef.GetFragmentConfigs().Num());
 	for (const FMassArchetypeFragmentConfig& FragmentConfig : ArchetypeRef.GetFragmentConfigs())
@@ -1087,7 +1098,7 @@ FAutoConsoleCommandWithWorldArgsAndOutputDevice GPrintArchetypesCmd(
 	FConsoleCommandWithWorldArgsAndOutputDeviceDelegate::CreateStatic(
 		[](const TArray<FString>& Params, UWorld* World, FOutputDevice& Ar)
 {
-	if (UMassEntitySubsystem* EntitySystem = World ? World->GetSubsystem<UMassEntitySubsystem>() : nullptr)
+	if (const UMassEntitySubsystem* EntitySystem = World ? World->GetSubsystem<UMassEntitySubsystem>() : nullptr)
 	{
 		EntitySystem->DebugPrintArchetypes(Ar);
 	}
diff --git a/Engine/Plugins/Runtime/MassEntity/Source/MassEntity/Private/MassEntityView.cpp b/Engine/Plugins/Runtime/MassEntity/Source/MassEntity/Private/MassEntityView.cpp
--- a/Engine/Plugins/Runtime/MassEntity/Source/MassEntity/Private/MassEntityView.cpp
+++ b/Engine/Plugins/Runtime/MassEntity/Source/MassEntity/Private/MassEntityView.cpp
@@ -42,8 +42,34 @@ void* FMassEntityView::GetFragmentPtrChecked(const UScriptStruct& FragmentType)
 	return Archetype->GetFragmentData(FragmentIndex, EntityHandle);
 }
 
+const void* FMassEntityView::GetConstSharedFragmentPtr(const UScriptStruct& FragmentType) const
+{
+	const FConstSharedStruct* SharedFragment = Archetype->GetSharedFragmentValues().GetConstSharedFragments().FindByPredicate(FStructTypeEqualOperator(&FragmentType));
+	return (SharedFragment != nullptr) ? SharedFragment->GetMemory() : nullptr;
+}
+
+const void* FMassEntityView::GetConstSharedFragmentPtrChecked(const UScriptStruct& FragmentType) const
+{
+	const FConstSharedStruct* SharedFragment = Archetype->GetSharedFragmentValues().GetConstSharedFragments().FindByPredicate(FStructTypeEqualOperator(&FragmentType));
+	check(SharedFragment != nullptr);
+	return SharedFragment->GetMemory();
+}
+
+void* FMassEntityView::GetSharedFragmentPtr(const UScriptStruct& FragmentType) const
+{
+	const FSharedStruct* SharedFragment = Archetype->GetSharedFragmentValues().GetSharedFragments().FindByPredicate(FStructTypeEqualOperator(&FragmentType));
+	return (SharedFragment != nullptr) ? SharedFragment->GetMutableMemory() : nullptr;
+}
+
+void* FMassEntityView::GetSharedFragmentPtrChecked(const UScriptStruct& FragmentType) const
+{
+	const FSharedStruct* SharedFragment = Archetype->GetSharedFragmentValues().GetSharedFragments().FindByPredicate(FStructTypeEqualOperator(&FragmentType));
+	check(SharedFragment != nullptr);
+	return SharedFragment->GetMutableMemory();
+}
+
 bool FMassEntityView::HasTag(const UScriptStruct& TagType) const
 {
 	checkSlow(Archetype && EntityHandle.IsValid());
 	return Archetype->HasTagType(&TagType);
-}
+}
\ No newline at end of file
diff --git a/Engine/Plugins/Runtime/MassEntity/Source/MassEntity/Public/MassCommandBuffer.h b/Engine/Plugins/Runtime/MassEntity/Source/MassEntity/Public/MassCommandBuffer.h
--- a/Engine/Plugins/Runtime/MassEntity/Source/MassEntity/Public/MassCommandBuffer.h
+++ b/Engine/Plugins/Runtime/MassEntity/Source/MassEntity/Public/MassCommandBuffer.h
@@ -143,9 +143,10 @@ struct MASSENTITY_API FBuildEntityFromFragmentInstance : public FCommandBufferEn
 	};
 
 	FBuildEntityFromFragmentInstance() = default;
-	FBuildEntityFromFragmentInstance(const FMassEntityHandle Entity, FStructView InStruct)
+	FBuildEntityFromFragmentInstance(const FMassEntityHandle Entity, FStructView InStruct, FMassArchetypeSharedFragmentValues InSharedFragmentValues)
 		: FCommandBufferEntryBase(Entity)
 		, Struct(InStruct)
+		, SharedFragmentValues(InSharedFragmentValues)
 	{}
 
 	void AppendAffectedEntitesPerType(FMassCommandsObservedTypes& ObservedTypes)
@@ -157,6 +158,7 @@ protected:
 	virtual void Execute(UMassEntitySubsystem& System) const override;
 
 	FInstancedStruct Struct;
+	FMassArchetypeSharedFragmentValues SharedFragmentValues;
 };
 
 /**
@@ -173,9 +175,10 @@ struct MASSENTITY_API FBuildEntityFromFragmentInstances : public FCommandBufferE
 	};
 
 	FBuildEntityFromFragmentInstances() = default;
-	FBuildEntityFromFragmentInstances(const FMassEntityHandle Entity, TConstArrayView<FInstancedStruct> InInstances)
+	FBuildEntityFromFragmentInstances(const FMassEntityHandle Entity, TConstArrayView<FInstancedStruct> InInstances, FMassArchetypeSharedFragmentValues InSharedFragmentValues = {})
 		: FCommandBufferEntryBase(Entity)
 		, Instances(InInstances)
+		, SharedFragmentValues(InSharedFragmentValues)
 	{}
 
 	void AppendAffectedEntitesPerType(FMassCommandsObservedTypes& ObservedTypes)
@@ -190,6 +193,7 @@ protected:
 	virtual void Execute(UMassEntitySubsystem& System) const override;
 
 	TArray<FInstancedStruct> Instances;
+	FMassArchetypeSharedFragmentValues SharedFragmentValues;
 };
 
 /**
diff --git a/Engine/Plugins/Runtime/MassEntity/Source/MassEntity/Public/MassEntitySubsystem.h b/Engine/Plugins/Runtime/MassEntity/Source/MassEntity/Public/MassEntitySubsystem.h
--- a/Engine/Plugins/Runtime/MassEntity/Source/MassEntity/Public/MassEntitySubsystem.h
+++ b/Engine/Plugins/Runtime/MassEntity/Source/MassEntity/Public/MassEntitySubsystem.h
@@ -83,7 +83,7 @@ public:
 	 * A special, relaxed but slower version of CreateArchetype functions that allows FragmentAngTagsList to contain 
 	 * both fragments and tags. 
 	 */
-	FArchetypeHandle CreateArchetype(TConstArrayView<const UScriptStruct*> FragmentsAngTagsList);
+	FArchetypeHandle CreateArchetype(TConstArrayView<const UScriptStruct*> FragmentsAndTagsList);
 
 	/**
 	 * CreateArchetype from a composition descriptor and initial values
@@ -173,7 +173,7 @@ public:
 	 * Builds an entity for it to be ready to be used by the subsystem
 	 * @param Entity to build which was retrieved with ReserveEntity() method
 	 * @param FragmentInstanceList is the fragments to create the entity from and initialize values*/
-	void BuildEntity(FMassEntityHandle Entity, TConstArrayView<FInstancedStruct> FragmentInstanceList);
+	void BuildEntity(FMassEntityHandle Entity, TConstArrayView<FInstancedStruct> FragmentInstanceList, FMassArchetypeSharedFragmentValues SharedFragmentValues = {});
 
 	/*
 	 * Releases a previously reserved entity that was not yet built, otherwise call DestroyEntity
diff --git a/Engine/Plugins/Runtime/MassEntity/Source/MassEntity/Public/MassEntityView.h b/Engine/Plugins/Runtime/MassEntity/Source/MassEntity/Public/MassEntityView.h
--- a/Engine/Plugins/Runtime/MassEntity/Source/MassEntity/Public/MassEntityView.h
+++ b/Engine/Plugins/Runtime/MassEntity/Source/MassEntity/Public/MassEntityView.h
@@ -69,6 +69,46 @@ struct MASSENTITY_API FMassEntityView
 		return FStructView(FragmentType, static_cast<uint8*>(GetFragmentPtr(*FragmentType)));
 	}
 
+	/** if the viewed entity doesn't have the given const shared fragment the function will return null */
+	template<typename T>
+	const T* GetConstSharedFragmentDataPtr() const
+	{
+		static_assert(TIsDerivedFrom<T, FMassSharedFragment>::IsDerived,
+			"Given struct doesn't represent a valid shared fragment type. Make sure to inherit from FMassSharedFragment or one of its child-types.");
+
+		return (const T*)GetConstSharedFragmentPtr(*T::StaticStruct());
+	}
+
+	/** will fail a check if the viewed entity doesn't have the given const shared fragment */
+	template<typename T>
+	const T& GetConstSharedFragmentData() const
+	{
+		static_assert(TIsDerivedFrom<T, FMassSharedFragment>::IsDerived,
+			"Given struct doesn't represent a valid const shared fragment type. Make sure to inherit from FMassSharedFragment or one of its child-types.");
+
+		return *((const T*)GetConstSharedFragmentPtrChecked(*T::StaticStruct()));
+	}
+
+	/** will fail a check if the viewed entity doesn't have the given shared fragment */
+	template<typename T>
+	T& GetSharedFragmentData() const
+	{
+		static_assert(TIsDerivedFrom<T, FMassSharedFragment>::IsDerived,
+			"Given struct doesn't represent a valid shared fragment type. Make sure to inherit from FMassSharedFragment or one of its child-types.");
+
+		return *((T*)GetSharedFragmentPtrChecked(*T::StaticStruct()));
+	}
+
+	/** if the viewed entity doesn't have the given shared fragment the function will return null */
+	template<typename T>
+	T* GetSharedFragmentDataPtr() const
+	{
+		static_assert(TIsDerivedFrom<T, FMassSharedFragment>::IsDerived,
+			"Given struct doesn't represent a valid shared fragment type. Make sure to inherit from FMassSharedFragment or one of its child-types.");
+
+		return (T*)GetSharedFragmentPtr(*T::StaticStruct());
+	}
+
 	template<typename T>
 	bool HasTag() const
 	{
@@ -82,6 +122,10 @@ struct MASSENTITY_API FMassEntityView
 protected:
 	void* GetFragmentPtr(const UScriptStruct& FragmentType) const;
 	void* GetFragmentPtrChecked(const UScriptStruct& FragmentType) const;
+	const void* GetConstSharedFragmentPtr(const UScriptStruct& FragmentType) const;
+	const void* GetConstSharedFragmentPtrChecked(const UScriptStruct& FragmentType) const;
+	void* GetSharedFragmentPtr(const UScriptStruct& FragmentType) const;
+	void* GetSharedFragmentPtrChecked(const UScriptStruct& FragmentType) const;
 	bool HasTag(const UScriptStruct& TagType) const;
 
 private:
