FCommandAddFragmentInstance can now be constructed from a ConstStructView
- fixed some typos

diff --git a/Engine/Plugins/Runtime/MassEntity/Source/MassEntity/Public/MassCommandBuffer.h b/Engine/Plugins/Runtime/MassEntity/Source/MassEntity/Public/MassCommandBuffer.h
--- a/Engine/Plugins/Runtime/MassEntity/Source/MassEntity/Public/MassCommandBuffer.h
+++ b/Engine/Plugins/Runtime/MassEntity/Source/MassEntity/Public/MassCommandBuffer.h
@@ -96,7 +96,7 @@ struct MASSENTITY_API FCommandBufferEntryBase
 		: TargetEntity(InEntity)
 	{}
 
-	void AppendAffectedEntitesPerType(FMassCommandsObservedTypes& ObservedTypes) { ensure(false); }
+	void AppendAffectedEntitiesPerType(FMassCommandsObservedTypes& ObservedTypes) { ensure(false); }
 	virtual void Execute(UMassEntitySubsystem& System) const PURE_VIRTUAL(FCommandBufferEntryBase::Execute, );
 };
 template<> struct TStructOpsTypeTraits<FCommandBufferEntryBase> : public TStructOpsTypeTraitsBase2<FCommandBufferEntryBase> { enum { WithPureVirtual = true, }; };
@@ -149,7 +149,7 @@ struct MASSENTITY_API FBuildEntityFromFragmentInstance : public FCommandBufferEn
 		, SharedFragmentValues(InSharedFragmentValues)
 	{}
 
-	void AppendAffectedEntitesPerType(FMassCommandsObservedTypes& ObservedTypes)
+	void AppendAffectedEntitiesPerType(FMassCommandsObservedTypes& ObservedTypes)
 	{
 		ObservedTypes.FragmentAdded(Struct.GetScriptStruct(), TargetEntity);
 	}
@@ -181,7 +181,7 @@ struct MASSENTITY_API FBuildEntityFromFragmentInstances : public FCommandBufferE
 		, SharedFragmentValues(InSharedFragmentValues)
 	{}
 
-	void AppendAffectedEntitesPerType(FMassCommandsObservedTypes& ObservedTypes)
+	void AppendAffectedEntitiesPerType(FMassCommandsObservedTypes& ObservedTypes)
 	{
 		for (const FInstancedStruct& Struct : Instances)
 		{
@@ -215,7 +215,7 @@ struct MASSENTITY_API FCommandAddFragment : public FCommandBufferEntryBase
 		, StructParam(InStruct)
 	{}
 
-	void AppendAffectedEntitesPerType(FMassCommandsObservedTypes& ObservedTypes)
+	void AppendAffectedEntitiesPerType(FMassCommandsObservedTypes& ObservedTypes)
 	{
 		ObservedTypes.FragmentAdded(StructParam, TargetEntity);
 	}
@@ -240,12 +240,12 @@ struct MASSENTITY_API FCommandAddFragmentInstance : public FCommandBufferEntryBa
 	};
 
 	FCommandAddFragmentInstance() = default;
-	FCommandAddFragmentInstance(const FMassEntityHandle InEntity, FStructView InStruct)
+	FCommandAddFragmentInstance(const FMassEntityHandle InEntity, FConstStructView InStruct)
         : FCommandBufferEntryBase(InEntity)
         , Struct(InStruct)
 	{}
 
-	void AppendAffectedEntitesPerType(FMassCommandsObservedTypes& ObservedTypes)
+	void AppendAffectedEntitiesPerType(FMassCommandsObservedTypes& ObservedTypes)
 	{
 		ObservedTypes.FragmentAdded(Struct.GetScriptStruct(), TargetEntity);
 	}
@@ -276,7 +276,7 @@ struct MASSENTITY_API FMassCommandAddFragmentInstanceList : public FCommandBuffe
 		, FragmentList(InitList)
 	{}
 
-	void AppendAffectedEntitesPerType(FMassCommandsObservedTypes& ObservedTypes)
+	void AppendAffectedEntitiesPerType(FMassCommandsObservedTypes& ObservedTypes)
 	{
 		for (const FInstancedStruct& Struct : FragmentList)
 		{
@@ -534,7 +534,7 @@ public:
 		T& Command = PendingCommands.Emplace_GetRef<T>(Forward<TArgs>(InArgs)...); 
 		if (constexpr bool bIsModifyingComposition = ((T::Type & (ECommandBufferOperationType::Add | ECommandBufferOperationType::Remove)) != 0))
 		{	
-			Command.AppendAffectedEntitesPerType(ObservedTypes);
+			Command.AppendAffectedEntitiesPerType(ObservedTypes);
 		}
 		return Command;
 	}
