Fixed some MassCommandBuffer typos

diff --git a/Engine/Plugins/Runtime/MassEntity/Source/MassEntity/Public/MassCommandBuffer.h b/Engine/Plugins/Runtime/MassEntity/Source/MassEntity/Public/MassCommandBuffer.h
--- a/Engine/Plugins/Runtime/MassEntity/Source/MassEntity/Public/MassCommandBuffer.h
+++ b/Engine/Plugins/Runtime/MassEntity/Source/MassEntity/Public/MassCommandBuffer.h
@@ -309,7 +309,7 @@ struct MASSENTITY_API FCommandRemoveFragment : public FCommandBufferEntryBase
 		, StructParam(InStruct)
 	{}
 
-	void AppendAffectedEntitesPerType(FMassCommandsObservedTypes& ObservedTypes)
+	void AppendAffectedEntitiesPerType(FMassCommandsObservedTypes& ObservedTypes)
 	{
 		ObservedTypes.FragmentRemoved(StructParam, TargetEntity);
 	}
@@ -339,7 +339,7 @@ struct MASSENTITY_API FCommandAddFragmentList : public FCommandBufferEntryBase
 		, FragmentList(InFragmentList)
 	{}
 
-	void AppendAffectedEntitesPerType(FMassCommandsObservedTypes& ObservedTypes)
+	void AppendAffectedEntitiesPerType(FMassCommandsObservedTypes& ObservedTypes)
 	{
 		for (const UScriptStruct* StructParam : FragmentList)
 		{
@@ -372,7 +372,7 @@ struct MASSENTITY_API FCommandRemoveFragmentList : public FCommandBufferEntryBas
 		, FragmentList(InFragmentList)
 	{}
 
-	void AppendAffectedEntitesPerType(FMassCommandsObservedTypes& ObservedTypes)
+	void AppendAffectedEntitiesPerType(FMassCommandsObservedTypes& ObservedTypes)
 	{
 		for (const UScriptStruct* StructParam : FragmentList)
 		{
@@ -486,7 +486,7 @@ struct MASSENTITY_API FCommandRemoveComposition : public FCommandBufferEntryBase
 		, Descriptor(InDescriptor)
 	{}
 
-	void AppendAffectedEntitesPerType(FMassCommandsObservedTypes& ObservedTypes)
+	void AppendAffectedEntitiesPerType(FMassCommandsObservedTypes& ObservedTypes)
 	{
 		if (Descriptor.Fragments.IsEmpty() == false)
 		{
