Mass: Allow to get shared fragment by UScriptStruct

diff --git a/Engine/Plugins/Runtime/MassEntity/Source/MassEntity/Public/MassEntityView.h b/Engine/Plugins/Runtime/MassEntity/Source/MassEntity/Public/MassEntityView.h
--- a/Engine/Plugins/Runtime/MassEntity/Source/MassEntity/Public/MassEntityView.h
+++ b/Engine/Plugins/Runtime/MassEntity/Source/MassEntity/Public/MassEntityView.h
@@ -89,6 +89,12 @@ struct MASSENTITY_API FMassEntityView
 		return *((const T*)GetConstSharedFragmentPtrChecked(*T::StaticStruct()));
 	}
 
+	FConstStructView GetConstSharedFragmentDataStruct(const UScriptStruct* FragmentType) const
+	{
+		check(FragmentType);
+		return FConstStructView(FragmentType, static_cast<const uint8*>(GetConstSharedFragmentPtr(*FragmentType)));
+	}
+
 	/** will fail a check if the viewed entity doesn't have the given shared fragment */
 	template<typename T>
 	T& GetSharedFragmentData() const
