Renamed a bunch of Mass types to follow our own type naming guidelines.

FArchetypeHandle -> FMassArchetypeHandle
FInternalEntityHandle -> FMassRawEntityInChunkData
FArchetypeChunkCollection -> FMassArchetypeSubChunks
FArchetypeChunkCollection::FSubChunkInfo -> FMassArchetypeSubChunks::FSubChunkInfo
FArchetypeChunkCollection::Chunks -> FMassArchetypeSubChunks::SubChunks
UE::MassEntity::Internal -> UE::MassEntity::Private

added:
using FMassArchetypeSubChunks::FSubChunkArray = TArray<FMassArchetypeSubChunks::FSubChunkInfo>
using FMassArchetypeSubChunks::FConstSubChunkArrayView = TConstArrayView<FMassArchetypeSubChunks::FSubChunkInfo>;

diff --git a/Engine/Plugins/AI/MassAI/Source/MassAIDebug/Private/GameplayDebuggerCategory_Mass.cpp b/Engine/Plugins/AI/MassAI/Source/MassAIDebug/Private/GameplayDebuggerCategory_Mass.cpp
--- a/Engine/Plugins/AI/MassAI/Source/MassAIDebug/Private/GameplayDebuggerCategory_Mass.cpp
+++ b/Engine/Plugins/AI/MassAI/Source/MassAIDebug/Private/GameplayDebuggerCategory_Mass.cpp
@@ -290,7 +290,7 @@ void FGameplayDebuggerCategory_Mass::CollectData(APlayerController* OwnerPC, AAc
 				}
 				else
 				{
-					const FArchetypeHandle Archetype = EntitySystem->GetArchetypeForEntity(CachedEntity);
+					const FMassArchetypeHandle Archetype = EntitySystem->GetArchetypeForEntity(CachedEntity);
 					TArray<FName> ComponentNames;
 					TArray<FName> TagNames;
 					EntitySystem->DebugGetArchetypeStrings(Archetype, ComponentNames, TagNames);
diff --git a/Engine/Plugins/Runtime/MassEntity/Source/MassEntity/Private/MassArchetypeData.cpp b/Engine/Plugins/Runtime/MassEntity/Source/MassEntity/Private/MassArchetypeData.cpp
--- a/Engine/Plugins/Runtime/MassEntity/Source/MassEntity/Private/MassArchetypeData.cpp
+++ b/Engine/Plugins/Runtime/MassEntity/Source/MassEntity/Private/MassArchetypeData.cpp
@@ -265,19 +265,19 @@ void FMassArchetypeData::RemoveEntityInternal(const int32 AbsoluteIndex, const b
 	}
 }
 
-void FMassArchetypeData::BatchDestroyEntityChunks(const FArchetypeChunkCollection& ChunkCollection, TArray<FMassEntityHandle>& OutEntitiesRemoved)
+void FMassArchetypeData::BatchDestroyEntityChunks(FMassArchetypeSubChunks::FConstSubChunkArrayView SubChunkContainer, TArray<FMassEntityHandle>& OutEntitiesRemoved)
 {
 	const int32 InitialOutEntitiesCount = OutEntitiesRemoved.Num();
 
 	// Sorting the subchunks info so that subchunks of a given chunk are processed "from the back". Otherwise removing 
 	// a subchunk from the front of the chunk would inevitably invalidate following subchunks' information.
-	TArray<FArchetypeChunkCollection::FChunkInfo> Subchunks(ChunkCollection.GetChunks());
-	Subchunks.Sort([](const FArchetypeChunkCollection::FChunkInfo& A, const FArchetypeChunkCollection::FChunkInfo& B) 
+	FMassArchetypeSubChunks::FSubChunkArray Subchunks(SubChunkContainer);
+	Subchunks.Sort([](const FMassArchetypeSubChunks::FSubChunkInfo& A, const FMassArchetypeSubChunks::FSubChunkInfo& B) 
 		{ 
 			return A.ChunkIndex < B.ChunkIndex || (A.ChunkIndex == B.ChunkIndex && A.SubchunkStart > B.SubchunkStart);
 		});
 
-	for (const FArchetypeChunkCollection::FChunkInfo SubchunkInfo : Subchunks)
+	for (const FMassArchetypeSubChunks::FSubChunkInfo SubchunkInfo : Subchunks)
 	{ 
 		FMassArchetypeChunk& Chunk = Chunks[SubchunkInfo.ChunkIndex];
 
@@ -367,7 +367,7 @@ bool FMassArchetypeData::HasFragmentDataForEntity(const UScriptStruct* FragmentT
 
 void* FMassArchetypeData::GetFragmentDataForEntityChecked(const UScriptStruct* FragmentType, int32 EntityIndex) const
 {
-	const FInternalEntityHandle InternalIndex = MakeEntityHandle(EntityIndex);
+	const FMassRawEntityInChunkData InternalIndex = MakeEntityHandle(EntityIndex);
 	
 	// failing the below Find means given entity's archetype is missing given FragmentType
 	const int32 FragmentIndex = FragmentIndexMap.FindChecked(FragmentType);
@@ -378,7 +378,7 @@ void* FMassArchetypeData::GetFragmentDataForEntity(const UScriptStruct* Fragment
 {
 	if (const int32* FragmentIndex = FragmentIndexMap.Find(FragmentType))
 	{
-		FInternalEntityHandle InternalIndex = MakeEntityHandle(EntityIndex);
+		FMassRawEntityInChunkData InternalIndex = MakeEntityHandle(EntityIndex);
 		// failing the below Find means given entity's archetype is missing given FragmentType
 		return GetFragmentData(*FragmentIndex, InternalIndex);
 	}
@@ -387,7 +387,7 @@ void* FMassArchetypeData::GetFragmentDataForEntity(const UScriptStruct* Fragment
 
 void FMassArchetypeData::SetFragmentsData(const FMassEntityHandle Entity, TArrayView<const FInstancedStruct> FragmentInstances)
 {
-	FInternalEntityHandle InternalIndex = MakeEntityHandle(Entity);
+	FMassRawEntityInChunkData InternalIndex = MakeEntityHandle(Entity);
 
 	for (const FInstancedStruct& Instance : FragmentInstances)
 	{
@@ -400,7 +400,7 @@ void FMassArchetypeData::SetFragmentsData(const FMassEntityHandle Entity, TArray
 	}
 }
 
-void FMassArchetypeData::SetFragmentData(const FArchetypeChunkCollection& ChunkCollection, const FInstancedStruct& FragmentSource)
+void FMassArchetypeData::SetFragmentData(FMassArchetypeSubChunks::FConstSubChunkArrayView SubChunkContainer, const FInstancedStruct& FragmentSource)
 {
 	check(FragmentSource.IsValid());
 	const UScriptStruct* FragmentType = FragmentSource.GetScriptStruct();
@@ -410,7 +410,7 @@ void FMassArchetypeData::SetFragmentData(const FArchetypeChunkCollection& ChunkC
 	const uint8* FragmentSourceMemory = FragmentSource.GetMemory();
 	check(FragmentSourceMemory);
 	
-	for (FMassArchetypeChunkIterator ChunkIterator(ChunkCollection); ChunkIterator; ++ChunkIterator)
+	for (FMassArchetypeChunkIterator ChunkIterator(SubChunkContainer); ChunkIterator; ++ChunkIterator)
 	{
 		uint8* FragmentMemory = (uint8*)FragmentConfigs[FragmentIndex].GetFragmentData(Chunks[ChunkIterator->ChunkIndex].GetRawMemory(), ChunkIterator->SubchunkStart);
 		for (int i = ChunkIterator->Length; i; --i, FragmentMemory += FragmentTypeSize)
@@ -470,17 +470,15 @@ void FMassArchetypeData::MoveEntityToAnotherArchetype(const FMassEntityHandle En
 	RemoveEntityInternal(AbsoluteIndex, bDestroyFragments);
 }
 
-void FMassArchetypeData::ExecuteFunction(FMassExecutionContext& RunContext, const FMassExecuteFunction& Function, const FMassQueryRequirementIndicesMapping& RequirementMapping, const FArchetypeChunkCollection& ChunkCollection)
+void FMassArchetypeData::ExecuteFunction(FMassExecutionContext& RunContext, const FMassExecuteFunction& Function, const FMassQueryRequirementIndicesMapping& RequirementMapping, FMassArchetypeSubChunks::FConstSubChunkArrayView SubChunkContainer)
 {
-	check(ChunkCollection.GetArchetype() == this);
-
 	// mz@todo to be removed
-	RunContext.SetCurrentArchetypeData(*this);
+	RunContext.SetCurrentArchetypesTagBitSet(GetTagBitSet());
 
 	BindConstSharedFragmentRequirements(RunContext, RequirementMapping.ChunkFragments);
 	BindSharedFragmentRequirements(RunContext, RequirementMapping.ChunkFragments);
 
-	for (FMassArchetypeChunkIterator ChunkIterator(ChunkCollection); ChunkIterator; ++ChunkIterator)
+	for (FMassArchetypeChunkIterator ChunkIterator(SubChunkContainer); ChunkIterator; ++ChunkIterator)
 	{
 		FMassArchetypeChunk& Chunk = Chunks[ChunkIterator->ChunkIndex];
 		
@@ -501,7 +499,7 @@ void FMassArchetypeData::ExecuteFunction(FMassExecutionContext& RunContext, cons
 void FMassArchetypeData::ExecuteFunction(FMassExecutionContext& RunContext, const FMassExecuteFunction& Function, const FMassQueryRequirementIndicesMapping& RequirementMapping, const FMassArchetypeConditionFunction& ArchetypeCondition, const FMassChunkConditionFunction& ChunkCondition)
 {
 	// mz@todo to be removed
-	RunContext.SetCurrentArchetypeData(*this);
+	RunContext.SetCurrentArchetypesTagBitSet(GetTagBitSet());
 
 	BindConstSharedFragmentRequirements(RunContext, RequirementMapping.ConstSharedFragments);
 	BindSharedFragmentRequirements(RunContext, RequirementMapping.SharedFragments);
@@ -525,7 +523,7 @@ void FMassArchetypeData::ExecuteFunction(FMassExecutionContext& RunContext, cons
 	}
 }
 
-void FMassArchetypeData::ExecutionFunctionForChunk(FMassExecutionContext RunContext, const FMassExecuteFunction& Function, const FMassQueryRequirementIndicesMapping& RequirementMapping, const FArchetypeChunkCollection::FChunkInfo& ChunkInfo, const FMassChunkConditionFunction& ChunkCondition)
+void FMassArchetypeData::ExecutionFunctionForChunk(FMassExecutionContext RunContext, const FMassExecuteFunction& Function, const FMassQueryRequirementIndicesMapping& RequirementMapping, const FMassArchetypeSubChunks::FSubChunkInfo& ChunkInfo, const FMassChunkConditionFunction& ChunkCondition)
 {
 	FMassArchetypeChunk& Chunk = Chunks[ChunkInfo.ChunkIndex];
 	const int32 ChunkLength = ChunkInfo.Length > 0 ? ChunkInfo.Length : (Chunk.GetNumInstances() - ChunkInfo.SubchunkStart);
@@ -535,7 +533,7 @@ void FMassArchetypeData::ExecutionFunctionForChunk(FMassExecutionContext RunCont
 		BindConstSharedFragmentRequirements(RunContext, RequirementMapping.ChunkFragments);
 		BindSharedFragmentRequirements(RunContext, RequirementMapping.ChunkFragments);
 
-		RunContext.SetCurrentArchetypeData(*this);
+		RunContext.SetCurrentArchetypesTagBitSet(GetTagBitSet());
 		RunContext.SetCurrentChunkSerialModificationNumber(Chunk.GetSerialModificationNumber());
 		BindChunkFragmentRequirements(RunContext, RequirementMapping.ChunkFragments, Chunk);
 
diff --git a/Engine/Plugins/Runtime/MassEntity/Source/MassEntity/Private/MassArchetypeData.h b/Engine/Plugins/Runtime/MassEntity/Source/MassEntity/Private/MassArchetypeData.h
--- a/Engine/Plugins/Runtime/MassEntity/Source/MassEntity/Private/MassArchetypeData.h
+++ b/Engine/Plugins/Runtime/MassEntity/Source/MassEntity/Private/MassArchetypeData.h
@@ -8,7 +8,7 @@
 struct FMassEntityQuery;
 struct FMassExecutionContext;
 class FOutputDevice;
-struct FArchetypeChunkCollection;
+struct FMassArchetypeSubChunks;
 
 // This is one chunk within an archetype
 struct FMassArchetypeChunk
@@ -158,7 +158,7 @@ private:
 	int32 EntityListOffsetWithinChunk;
 
 	friend FMassEntityQuery;
-	friend FArchetypeChunkCollection;
+	friend FMassArchetypeSubChunks;
 
 public:
 	TConstArrayView<FMassArchetypeFragmentConfig> GetFragmentConfigs() const { return FragmentConfigs; }
@@ -190,7 +190,7 @@ public:
 
 	void AddEntity(FMassEntityHandle Entity);
 	void RemoveEntity(FMassEntityHandle Entity);
-	void BatchDestroyEntityChunks(const FArchetypeChunkCollection& ChunkCollection, TArray<FMassEntityHandle>& OutEntitiesRemoved);
+	void BatchDestroyEntityChunks(FMassArchetypeSubChunks::FConstSubChunkArrayView SubChunkContainer, TArray<FMassEntityHandle>& OutEntitiesRemoved);
 
 	bool HasFragmentDataForEntity(const UScriptStruct* FragmentType, int32 EntityIndex) const;
 	void* GetFragmentDataForEntityChecked(const UScriptStruct* FragmentType, int32 EntityIndex) const;
@@ -205,10 +205,10 @@ public:
 
 	int32 GetChunkCount() const { return Chunks.Num(); }
 
-	void ExecuteFunction(FMassExecutionContext& RunContext, const FMassExecuteFunction& Function, const FMassQueryRequirementIndicesMapping& RequirementMapping, const FArchetypeChunkCollection& ChunkCollection);
+	void ExecuteFunction(FMassExecutionContext& RunContext, const FMassExecuteFunction& Function, const FMassQueryRequirementIndicesMapping& RequirementMapping, FMassArchetypeSubChunks::FConstSubChunkArrayView SubChunkContainer);
 	void ExecuteFunction(FMassExecutionContext& RunContext, const FMassExecuteFunction& Function, const FMassQueryRequirementIndicesMapping& RequirementMapping, const FMassArchetypeConditionFunction& ArchetypeCondition, const FMassChunkConditionFunction& ChunkCondition);
 
-	void ExecutionFunctionForChunk(FMassExecutionContext RunContext, const FMassExecuteFunction& Function, const FMassQueryRequirementIndicesMapping& RequirementMapping, const FArchetypeChunkCollection::FChunkInfo& ChunkInfo, const FMassChunkConditionFunction& ChunkCondition = FMassChunkConditionFunction());
+	void ExecutionFunctionForChunk(FMassExecutionContext RunContext, const FMassExecuteFunction& Function, const FMassQueryRequirementIndicesMapping& RequirementMapping, const FMassArchetypeSubChunks::FSubChunkInfo& ChunkInfo, const FMassChunkConditionFunction& ChunkCondition = FMassChunkConditionFunction());
 
 	/**
 	 * Compacts entities to fill up chunks as much as possible
@@ -231,7 +231,7 @@ public:
 
 	/** For all entities indicated by ChunkCollection the function sets the value of fragment of type
 	 *  FragmentSource.GetScriptStruct to the value represented by FragmentSource.GetMemory */
-	void SetFragmentData(const FArchetypeChunkCollection& ChunkCollection, const FInstancedStruct& FragmentSource);
+	void SetFragmentData(FMassArchetypeSubChunks::FConstSubChunkArrayView SubChunkContainer, const FInstancedStruct& FragmentSource);
 
 	/** Returns conversion from given Requirements to archetype's fragment indices */
 	void GetRequirementsFragmentMapping(TConstArrayView<FMassFragmentRequirement> Requirements, FMassFragmentIndicesMapping& OutFragmentIndices);
@@ -272,20 +272,20 @@ public:
 	FORCEINLINE const int32* GetFragmentIndex(const UScriptStruct* FragmentType) const { return FragmentIndexMap.Find(FragmentType); }
 	FORCEINLINE int32 GetFragmentIndexChecked(const UScriptStruct* FragmentType) const { return FragmentIndexMap.FindChecked(FragmentType); }
 
-	FORCEINLINE void* GetFragmentData(const int32 FragmentIndex, const FInternalEntityHandle EntityIndex) const
+	FORCEINLINE void* GetFragmentData(const int32 FragmentIndex, const FMassRawEntityInChunkData EntityIndex) const
 	{
 		return FragmentConfigs[FragmentIndex].GetFragmentData(EntityIndex.ChunkRawMemory, EntityIndex.IndexWithinChunk);
 	}
 
-	FORCEINLINE FInternalEntityHandle MakeEntityHandle(int32 EntityIndex) const
+	FORCEINLINE FMassRawEntityInChunkData MakeEntityHandle(int32 EntityIndex) const
 	{
 		const int32 AbsoluteIndex = EntityMap.FindChecked(EntityIndex);
 		const int32 ChunkIndex = AbsoluteIndex / NumEntitiesPerChunk;
 	
-		return FInternalEntityHandle(Chunks[ChunkIndex].GetRawMemory(), AbsoluteIndex % NumEntitiesPerChunk); 
+		return FMassRawEntityInChunkData(Chunks[ChunkIndex].GetRawMemory(), AbsoluteIndex % NumEntitiesPerChunk); 
 	}
 
-	FORCEINLINE FInternalEntityHandle MakeEntityHandle(FMassEntityHandle Entity) const
+	FORCEINLINE FMassRawEntityInChunkData MakeEntityHandle(FMassEntityHandle Entity) const
 	{
 		return MakeEntityHandle(Entity.Index); 
 	}
diff --git a/Engine/Plugins/Runtime/MassEntity/Source/MassEntity/Private/MassArchetypeTypes.cpp b/Engine/Plugins/Runtime/MassEntity/Source/MassEntity/Private/MassArchetypeTypes.cpp
--- a/Engine/Plugins/Runtime/MassEntity/Source/MassEntity/Private/MassArchetypeTypes.cpp
+++ b/Engine/Plugins/Runtime/MassEntity/Source/MassEntity/Private/MassArchetypeTypes.cpp
@@ -7,21 +7,17 @@
 
 
 //////////////////////////////////////////////////////////////////////
-// FArchetypeHandle
-bool FArchetypeHandle::operator==(const FMassArchetypeData* Other) const
-{
-	return DataPtr.Get() == Other;
-}
+// FMassArchetypeHandle
 
-uint32 GetTypeHash(const FArchetypeHandle& Instance)
+uint32 GetTypeHash(const FMassArchetypeHandle& Instance)
 {
 	return GetTypeHash(Instance.DataPtr.Get());
 }
 
 //////////////////////////////////////////////////////////////////////
-// FArchetypeChunkCollection
+// FMassArchetypeSubChunks 
 
-FArchetypeChunkCollection::FArchetypeChunkCollection(const FArchetypeHandle& InArchetype, TConstArrayView<FMassEntityHandle> InEntities, EDuplicatesHandling DuplicatesHandling)
+FMassArchetypeSubChunks::FMassArchetypeSubChunks(const FMassArchetypeHandle& InArchetype, TConstArrayView<FMassEntityHandle> InEntities, EDuplicatesHandling DuplicatesHandling)
 	: Archetype(InArchetype)
 {
 	check(InArchetype.IsValid());
@@ -86,8 +82,8 @@ FArchetypeChunkCollection::FArchetypeChunkCollection(const FArchetypeHandle& InA
 	// the following block of code is splitting up sorted AbsoluteIndices into 
 	// continuous chunks
 	int32 ChunkEnd = INDEX_NONE;
-	FChunkInfo DummyChunk;
-	FChunkInfo* SubChunkPtr = &DummyChunk;
+	FSubChunkInfo DummyChunk;
+	FSubChunkInfo* SubChunkPtr = &DummyChunk;
 	int32 SubchunkLen = 0;
 	int32 PrevAbsoluteIndex = INDEX_NONE;
 	for (const int32 Index : TrueIndices)
@@ -102,7 +98,7 @@ FArchetypeChunkCollection::FArchetypeChunkCollection(const FArchetypeHandle& InA
 			SubchunkLen = 0;
 			// new subchunk
 			const int32 SubchunkStart = Index % NumEntitiesPerChunk;
-			SubChunkPtr = &Chunks.Add_GetRef(FChunkInfo(ChunkIndex, SubchunkStart));
+			SubChunkPtr = &SubChunks.Add_GetRef(FSubChunkInfo(ChunkIndex, SubchunkStart));
 		}
 		++SubchunkLen;
 		PrevAbsoluteIndex = Index;
@@ -111,40 +107,40 @@ FArchetypeChunkCollection::FArchetypeChunkCollection(const FArchetypeHandle& InA
 	SubChunkPtr->Length = SubchunkLen;
 }
 
-FArchetypeChunkCollection::FArchetypeChunkCollection(FArchetypeHandle& InArchetypeHandle)
+FMassArchetypeSubChunks::FMassArchetypeSubChunks(FMassArchetypeHandle& InArchetypeHandle)
 {
 	check(InArchetypeHandle.DataPtr.IsValid());
 	GatherChunksFromArchetype(InArchetypeHandle.DataPtr);
 }
 
-FArchetypeChunkCollection::FArchetypeChunkCollection(TSharedPtr<FMassArchetypeData>& InArchetype)
+FMassArchetypeSubChunks::FMassArchetypeSubChunks(TSharedPtr<FMassArchetypeData>& InArchetype)
 {
 	GatherChunksFromArchetype(InArchetype);
 }
 
-void FArchetypeChunkCollection::GatherChunksFromArchetype(TSharedPtr<FMassArchetypeData>& InArchetype)
+void FMassArchetypeSubChunks::GatherChunksFromArchetype(TSharedPtr<FMassArchetypeData>& InArchetype)
 {
 	check(InArchetype.IsValid());
 	Archetype.DataPtr = InArchetype;
 
 	const int32 ChunkCount = InArchetype->GetChunkCount();
-	Chunks.Reset(ChunkCount);
+	SubChunks.Reset(ChunkCount);
 	for (int32 i = 0; i < ChunkCount; ++i)
 	{
-		Chunks.Add(FChunkInfo(i));
+		SubChunks.Add(FSubChunkInfo(i));
 	}
 }
 
-bool FArchetypeChunkCollection::IsSame(const FArchetypeChunkCollection& Other) const
+bool FMassArchetypeSubChunks::IsSame(const FMassArchetypeSubChunks& Other) const
 {
-	if (Archetype != Other.Archetype || Chunks.Num() != Other.Chunks.Num())
+	if (Archetype != Other.Archetype || SubChunks.Num() != Other.SubChunks.Num())
 	{
 		return false;
 	}
 
-	for (int i = 0; i < Chunks.Num(); ++i)
+	for (int i = 0; i < SubChunks.Num(); ++i)
 	{
-		if (Chunks[i] != Other.Chunks[i])
+		if (SubChunks[i] != Other.SubChunks[i])
 		{
 			return false;
 		}
diff --git a/Engine/Plugins/Runtime/MassEntity/Source/MassEntity/Private/MassCommandBuffer.cpp b/Engine/Plugins/Runtime/MassEntity/Source/MassEntity/Private/MassCommandBuffer.cpp
--- a/Engine/Plugins/Runtime/MassEntity/Source/MassEntity/Private/MassCommandBuffer.cpp
+++ b/Engine/Plugins/Runtime/MassEntity/Source/MassEntity/Private/MassCommandBuffer.cpp
@@ -102,9 +102,9 @@ void FMassCommandBuffer::ReplayBufferAgainstSystem(UMassEntitySubsystem* EntityS
 			check(It.Key);
 			if (ObservedRemoveFragments.Contains(*It.Key))
 			{
-				TArray<FArchetypeChunkCollection> ChunkCollections;
-				UE::Mass::Utils::CreateSparseChunks(*EntitySystem, It.Value, FArchetypeChunkCollection::FoldDuplicates, ChunkCollections);
-				for (FArchetypeChunkCollection& Collection : ChunkCollections)
+				TArray<FMassArchetypeSubChunks > ChunkCollections;
+				UE::Mass::Utils::CreateSparseChunks(*EntitySystem, It.Value, FMassArchetypeSubChunks::FoldDuplicates, ChunkCollections);
+				for (FMassArchetypeSubChunks& Collection : ChunkCollections)
 				{
 					check(It.Key);
 					ObserverManager.OnPreFragmentRemoved(*It.Key, Collection);
@@ -112,11 +112,11 @@ void FMassCommandBuffer::ReplayBufferAgainstSystem(UMassEntitySubsystem* EntityS
 			}
 		}
 	
-		TArray<FArchetypeChunkCollection> EntityChunksToDestroy;
+		TArray<FMassArchetypeSubChunks > EntityChunksToDestroy;
 		if (EntitiesToDestroy.Num())
 		{
-			UE::Mass::Utils::CreateSparseChunks(*EntitySystem, EntitiesToDestroy, FArchetypeChunkCollection::FoldDuplicates, EntityChunksToDestroy);
-			for (FArchetypeChunkCollection& Collection : EntityChunksToDestroy)
+			UE::Mass::Utils::CreateSparseChunks(*EntitySystem, EntitiesToDestroy, FMassArchetypeSubChunks::FoldDuplicates, EntityChunksToDestroy);
+			for (FMassArchetypeSubChunks& Collection : EntityChunksToDestroy)
 			{
 				EntitySystem->BatchDestroyEntityChunks(Collection);
 			}
@@ -160,9 +160,9 @@ void FMassCommandBuffer::ReplayBufferAgainstSystem(UMassEntitySubsystem* EntityS
 		check(It.Key);
 		if (ObservedAddFragments.Contains(*It.Key))
 		{
-			TArray<FArchetypeChunkCollection> ChunkCollections;
-			UE::Mass::Utils::CreateSparseChunks(*EntitySystem, It.Value, FArchetypeChunkCollection::FoldDuplicates, ChunkCollections);
-			for (FArchetypeChunkCollection& Collection : ChunkCollections)
+			TArray<FMassArchetypeSubChunks > ChunkCollections;
+			UE::Mass::Utils::CreateSparseChunks(*EntitySystem, It.Value, FMassArchetypeSubChunks::FoldDuplicates, ChunkCollections);
+			for (FMassArchetypeSubChunks& Collection : ChunkCollections)
 			{
 				check(It.Key);
 				ObserverManager.OnPostFragmentAdded(*It.Key, Collection);
diff --git a/Engine/Plugins/Runtime/MassEntity/Source/MassEntity/Private/MassEntityQuery.cpp b/Engine/Plugins/Runtime/MassEntity/Source/MassEntity/Private/MassEntityQuery.cpp
--- a/Engine/Plugins/Runtime/MassEntity/Source/MassEntity/Private/MassEntityQuery.cpp
+++ b/Engine/Plugins/Runtime/MassEntity/Source/MassEntity/Private/MassEntityQuery.cpp
@@ -104,7 +104,7 @@ bool FMassEntityQuery::CheckValidity() const
 	return RequiredAllFragments.IsEmpty() == false || RequiredAnyFragments.IsEmpty() == false || RequiredOptionalFragments.IsEmpty() == false;
 }
 
-bool FMassEntityQuery::DoesArchetypeMatchRequirements(const FArchetypeHandle& ArchetypeHandle) const
+bool FMassEntityQuery::DoesArchetypeMatchRequirements(const FMassArchetypeHandle& ArchetypeHandle) const
 {
 	check(ArchetypeHandle.IsValid());
 	const FMassArchetypeData* Archetype = ArchetypeHandle.DataPtr.Get();
@@ -122,7 +122,7 @@ bool FMassEntityQuery::DoesArchetypeMatchRequirements(const FArchetypeHandle& Ar
 		&& ArchetypeComposition.ChunkFragments.HasNone(RequiredNoneChunkFragments);
 }
 
-void FMassEntityQuery::ForEachEntityChunk(const FArchetypeChunkCollection& Chunks, UMassEntitySubsystem& EntitySubsystem, FMassExecutionContext& ExecutionContext, const FMassExecuteFunction& ExecuteFunction)
+void FMassEntityQuery::ForEachEntityChunk(const FMassArchetypeSubChunks& Chunks, UMassEntitySubsystem& EntitySubsystem, FMassExecutionContext& ExecutionContext, const FMassExecuteFunction& ExecuteFunction)
 {
 	// mz@todo I don't like that we're copying data here.
 	ExecutionContext.SetChunkCollection(Chunks);
@@ -148,7 +148,7 @@ void FMassEntityQuery::ForEachEntityChunk(UMassEntitySubsystem& EntitySubsystem,
 			return;
 		}
 		ExecutionContext.SetRequirements(Requirements, ChunkRequirements, ConstSharedRequirements, SharedRequirements);
-		ExecutionContext.GetChunkCollection().GetArchetype().DataPtr->ExecuteFunction(ExecutionContext, ExecuteFunction, {}, ExecutionContext.GetChunkCollection());
+		ExecutionContext.GetChunkCollection().GetArchetype().DataPtr->ExecuteFunction(ExecutionContext, ExecuteFunction, {}, ExecutionContext.GetChunkCollection().GetChunks());
 #if WITH_MASSENTITY_DEBUG
 		NumEntitiesToProcess = ExecutionContext.GetNumEntities();
 #endif
@@ -161,7 +161,7 @@ void FMassEntityQuery::ForEachEntityChunk(UMassEntitySubsystem& EntitySubsystem,
 
 		for (int i = 0; i < ValidArchetypes.Num(); ++i)
 		{
-			FArchetypeHandle& Archetype = ValidArchetypes[i];
+			FMassArchetypeHandle& Archetype = ValidArchetypes[i];
 			check(Archetype.IsValid());
 			Archetype.DataPtr->ExecuteFunction(ExecutionContext, ExecuteFunction, ArchetypeFragmentMapping[i], ArchetypeCondition, ChunkCondition);
 			ExecutionContext.ClearFragmentViews();
@@ -193,7 +193,7 @@ void FMassEntityQuery::ParallelForEachEntityChunk(UMassEntitySubsystem& EntitySu
 	{
 		FMassArchetypeData& Archetype;
 		const int32 ArchetypeIndex;
-		const FArchetypeChunkCollection::FChunkInfo ChunkInfo;
+		const FMassArchetypeSubChunks::FSubChunkInfo ChunkInfo;
 	};
 	TArray<FChunkJob> Jobs;
 
@@ -201,7 +201,7 @@ void FMassEntityQuery::ParallelForEachEntityChunk(UMassEntitySubsystem& EntitySu
 	if (ExecutionContext.GetChunkCollection().IsSet())
 	{
 		// verify the archetype matches requirements
-		FArchetypeHandle ArchetypeHandle = ExecutionContext.GetChunkCollection().GetArchetype();
+		FMassArchetypeHandle ArchetypeHandle = ExecutionContext.GetChunkCollection().GetArchetype();
 		if (DoesArchetypeMatchRequirements(ArchetypeHandle) == false)
 		{
 			// mz@todo add a unit test for this message
@@ -213,8 +213,8 @@ void FMassEntityQuery::ParallelForEachEntityChunk(UMassEntitySubsystem& EntitySu
 		ExecutionContext.SetRequirements(Requirements, ChunkRequirements, ConstSharedRequirements, SharedRequirements);
 		check(ArchetypeHandle.IsValid());
 		FMassArchetypeData& ArchetypeRef = *ArchetypeHandle.DataPtr.Get();
-		const FArchetypeChunkCollection AsChunkCollection(ArchetypeHandle.DataPtr);
-		for (const FArchetypeChunkCollection::FChunkInfo& ChunkInfo : AsChunkCollection.GetChunks())
+		const FMassArchetypeSubChunks AsChunkCollection(ArchetypeHandle.DataPtr);
+		for (const FMassArchetypeSubChunks::FSubChunkInfo& ChunkInfo : AsChunkCollection.GetChunks())
 		{
 			Jobs.Add({ ArchetypeRef, INDEX_NONE, ChunkInfo });
 		}
@@ -225,11 +225,11 @@ void FMassEntityQuery::ParallelForEachEntityChunk(UMassEntitySubsystem& EntitySu
 		ExecutionContext.SetRequirements(Requirements, ChunkRequirements, ConstSharedRequirements, SharedRequirements);
 		for (int ArchetypeIndex = 0; ArchetypeIndex < ValidArchetypes.Num(); ++ArchetypeIndex)
 		{
-			FArchetypeHandle& Archetype = ValidArchetypes[ArchetypeIndex];
+			FMassArchetypeHandle& Archetype = ValidArchetypes[ArchetypeIndex];
 			check(Archetype.IsValid());
 			FMassArchetypeData& ArchetypeRef = *Archetype.DataPtr.Get();
-			const FArchetypeChunkCollection AsChunkCollection(Archetype.DataPtr);
-			for (const FArchetypeChunkCollection::FChunkInfo& ChunkInfo : AsChunkCollection.GetChunks())
+			const FMassArchetypeSubChunks AsChunkCollection(Archetype.DataPtr);
+			for (const FMassArchetypeSubChunks::FSubChunkInfo& ChunkInfo : AsChunkCollection.GetChunks())
 			{
 				Jobs.Add({ArchetypeRef, ArchetypeIndex, ChunkInfo});
 			}
@@ -252,7 +252,7 @@ int32 FMassEntityQuery::GetNumMatchingEntities(UMassEntitySubsystem& InEntitySub
 {
 	CacheArchetypes(InEntitySubsystem);
 	int32 TotalEntities = 0;
-	for (FArchetypeHandle& ArchetypeHandle : ValidArchetypes)
+	for (FMassArchetypeHandle& ArchetypeHandle : ValidArchetypes)
 	{
 		if (const FMassArchetypeData* Archetype = ArchetypeHandle.DataPtr.Get())
 		{
@@ -266,7 +266,7 @@ bool FMassEntityQuery::HasMatchingEntities(UMassEntitySubsystem& InEntitySubsyst
 {
 	CacheArchetypes(InEntitySubsystem);
 
-	for (FArchetypeHandle& ArchetypeHandle : ValidArchetypes)
+	for (FMassArchetypeHandle& ArchetypeHandle : ValidArchetypes)
 	{
 		const FMassArchetypeData* Archetype = ArchetypeHandle.DataPtr.Get();
 		if (Archetype && Archetype->GetNumEntities() > 0)
@@ -301,7 +301,7 @@ FString FMassEntityQuery::DebugGetDescription() const
 #endif
 }
 
-FString FMassEntityQuery::DebugGetArchetypeCompatibilityDescription(const FArchetypeHandle& ArchetypeHandle) const
+FString FMassEntityQuery::DebugGetArchetypeCompatibilityDescription(const FMassArchetypeHandle& ArchetypeHandle) const
 {
 	if (ArchetypeHandle.IsValid() == false)
 	{
diff --git a/Engine/Plugins/Runtime/MassEntity/Source/MassEntity/Private/MassEntitySubsystem.cpp b/Engine/Plugins/Runtime/MassEntity/Source/MassEntity/Private/MassEntitySubsystem.cpp
--- a/Engine/Plugins/Runtime/MassEntity/Source/MassEntity/Private/MassEntitySubsystem.cpp
+++ b/Engine/Plugins/Runtime/MassEntity/Source/MassEntity/Private/MassEntitySubsystem.cpp
@@ -85,7 +85,7 @@ void UMassEntitySubsystem::Deinitialize()
 	//@TODO: Should actually do this...
 }
 
-FArchetypeHandle UMassEntitySubsystem::CreateArchetype(TConstArrayView<const UScriptStruct*> FragmentsAndTagsList)
+FMassArchetypeHandle UMassEntitySubsystem::CreateArchetype(TConstArrayView<const UScriptStruct*> FragmentsAndTagsList)
 {
 	FMassChunkFragmentBitSet ChunkFragments;
 	FMassTagBitSet Tags;
@@ -117,7 +117,7 @@ FArchetypeHandle UMassEntitySubsystem::CreateArchetype(TConstArrayView<const USc
 	return CreateArchetype(Composition, FMassArchetypeSharedFragmentValues());
 }
 
-FArchetypeHandle UMassEntitySubsystem::CreateArchetype(const TSharedPtr<FMassArchetypeData>& SourceArchetype, const FMassFragmentBitSet& NewFragments)
+FMassArchetypeHandle UMassEntitySubsystem::CreateArchetype(const TSharedPtr<FMassArchetypeData>& SourceArchetype, const FMassFragmentBitSet& NewFragments)
 {
 	check(SourceArchetype.IsValid());
 	checkf(NewFragments.IsEmpty() == false, TEXT("%s Adding an empty fragment list to an archetype is not supported."), ANSI_TO_TCHAR(__FUNCTION__));
@@ -126,13 +126,13 @@ FArchetypeHandle UMassEntitySubsystem::CreateArchetype(const TSharedPtr<FMassArc
 	return CreateArchetype(Composition, SourceArchetype->GetSharedFragmentValues());
 }
 
-FArchetypeHandle UMassEntitySubsystem::CreateArchetype(const FMassArchetypeCompositionDescriptor& Composition, const FMassArchetypeSharedFragmentValues& SharedFragmentValues)
+FMassArchetypeHandle UMassEntitySubsystem::CreateArchetype(const FMassArchetypeCompositionDescriptor& Composition, const FMassArchetypeSharedFragmentValues& SharedFragmentValues)
 {
 	const uint32 TypeHash = Composition.CalculateHash();
 
 	TArray<TSharedPtr<FMassArchetypeData>>& HashRow = FragmentHashToArchetypeMap.FindOrAdd(TypeHash);
 
-	FArchetypeHandle Result;
+	FMassArchetypeHandle Result;
 	for (const TSharedPtr<FMassArchetypeData>& Ptr : HashRow)
 	{
 		if (Ptr->IsEquivalent(Composition, SharedFragmentValues))
@@ -149,7 +149,7 @@ FArchetypeHandle UMassEntitySubsystem::CreateArchetype(const FMassArchetypeCompo
 		NewArchetype->Initialize(Composition, SharedFragmentValues);
 		HashRow.Add(NewArchetype);
 
-		for (const FMassArchetypeFragmentConfig& FragmentConfig : NewArchetype->GetFragmentConfigs())
+		for (const FMassArchetypeFragmentConfig & FragmentConfig : NewArchetype->GetFragmentConfigs())
 		{
 			checkSlow(FragmentConfig.FragmentType)
 			FragmentTypeToArchetypeMap.FindOrAdd(FragmentConfig.FragmentType).Add(NewArchetype);
@@ -163,7 +163,7 @@ FArchetypeHandle UMassEntitySubsystem::CreateArchetype(const FMassArchetypeCompo
 	return MoveTemp(Result);
 }
 
-FArchetypeHandle UMassEntitySubsystem::InternalCreateSiblingArchetype(const TSharedPtr<FMassArchetypeData>& SourceArchetype, const FMassTagBitSet& OverrideTags)
+FMassArchetypeHandle UMassEntitySubsystem::InternalCreateSiblingArchetype(const TSharedPtr<FMassArchetypeData>& SourceArchetype, const FMassTagBitSet& OverrideTags)
 {
 	checkSlow(SourceArchetype.IsValid());
 	const FMassArchetypeData& SourceArchetypeRef = *SourceArchetype.Get();
@@ -172,7 +172,7 @@ FArchetypeHandle UMassEntitySubsystem::InternalCreateSiblingArchetype(const TSha
 
 	TArray<TSharedPtr<FMassArchetypeData>>& HashRow = FragmentHashToArchetypeMap.FindOrAdd(TypeHash);
 
-	FArchetypeHandle Result;
+	FMassArchetypeHandle Result;
 	for (const TSharedPtr<FMassArchetypeData>& Ptr : HashRow)
 	{
 		if (Ptr->IsEquivalent(NewComposition, SourceArchetypeRef.GetSharedFragmentValues()))
@@ -203,9 +203,9 @@ FArchetypeHandle UMassEntitySubsystem::InternalCreateSiblingArchetype(const TSha
 	return MoveTemp(Result);
 }
 
-FArchetypeHandle UMassEntitySubsystem::GetArchetypeForEntity(FMassEntityHandle Entity) const
+FMassArchetypeHandle UMassEntitySubsystem::GetArchetypeForEntity(FMassEntityHandle Entity) const
 {
-	FArchetypeHandle Result;
+	FMassArchetypeHandle Result;
 	if (IsEntityValid(Entity))
 	{
 		Result.DataPtr = Entities[Entity.Index].CurrentArchetype;
@@ -213,7 +213,7 @@ FArchetypeHandle UMassEntitySubsystem::GetArchetypeForEntity(FMassEntityHandle E
 	return Result;
 }
 
-void UMassEntitySubsystem::ForEachArchetypeFragmentType(const FArchetypeHandle Archetype, TFunction< void(const UScriptStruct* /*FragmentType*/)> Function)
+void UMassEntitySubsystem::ForEachArchetypeFragmentType(const FMassArchetypeHandle Archetype, TFunction< void(const UScriptStruct* /*FragmentType*/)> Function)
 {
 	check(Archetype.DataPtr.IsValid());
 	const FMassArchetypeData& ArchetypeData = *Archetype.DataPtr.Get();
@@ -244,7 +244,7 @@ void UMassEntitySubsystem::DoEntityCompaction(const double TimeAllowed)
 	}
 }
 
-FMassEntityHandle UMassEntitySubsystem::CreateEntity(const FArchetypeHandle Archetype)
+FMassEntityHandle UMassEntitySubsystem::CreateEntity(const FMassArchetypeHandle Archetype)
 {
 	check(Archetype.DataPtr.IsValid());
 
@@ -257,7 +257,7 @@ FMassEntityHandle UMassEntitySubsystem::CreateEntity(TConstArrayView<FInstancedS
 {
 	check(FragmentInstanceList.Num() > 0);
 
-	const FArchetypeHandle Archetype = CreateArchetype(FMassArchetypeCompositionDescriptor(FragmentInstanceList, FMassTagBitSet(), FMassChunkFragmentBitSet(), FMassSharedFragmentBitSet()), FMassArchetypeSharedFragmentValues());
+	const FMassArchetypeHandle Archetype = CreateArchetype(FMassArchetypeCompositionDescriptor(FragmentInstanceList, FMassTagBitSet(), FMassChunkFragmentBitSet(), FMassSharedFragmentBitSet()), FMassArchetypeSharedFragmentValues());
 	check(Archetype.DataPtr.IsValid());
 
 	const FMassEntityHandle Entity = ReserveEntity();
@@ -287,7 +287,7 @@ void UMassEntitySubsystem::ReleaseReservedEntity(FMassEntityHandle Entity)
 	InternalReleaseEntity(Entity);
 }
 
-void UMassEntitySubsystem::BuildEntity(FMassEntityHandle Entity, FArchetypeHandle Archetype)
+void UMassEntitySubsystem::BuildEntity(FMassEntityHandle Entity, FMassArchetypeHandle Archetype)
 {
 	checkf(!IsEntityBuilt(Entity), TEXT("Expecting an entity that is not already built"));
 	check(Archetype.DataPtr.IsValid());
@@ -312,7 +312,7 @@ void UMassEntitySubsystem::BuildEntity(FMassEntityHandle Entity, TConstArrayView
 		Composition.SharedFragments.Add(*SharedFragment.GetScriptStruct());
 	}
 
-	const FArchetypeHandle Archetype = CreateArchetype(Composition, SharedFragmentValues);
+	const FMassArchetypeHandle Archetype = CreateArchetype(Composition, SharedFragmentValues);
 	check(Archetype.DataPtr.IsValid());
 
 	InternalBuildEntity(Entity, Archetype);
@@ -321,7 +321,7 @@ void UMassEntitySubsystem::BuildEntity(FMassEntityHandle Entity, TConstArrayView
 	EntityData.CurrentArchetype->SetFragmentsData(Entity, FragmentInstanceList);
 }
 
-TSharedRef<UMassEntitySubsystem::FEntityCreationContext> UMassEntitySubsystem::BatchCreateEntities(const FArchetypeHandle Archetype, const int32 Count, TArray<FMassEntityHandle>& OutEntities)
+TSharedRef<UMassEntitySubsystem::FEntityCreationContext> UMassEntitySubsystem::BatchCreateEntities(const FMassArchetypeHandle Archetype, const int32 Count, TArray<FMassEntityHandle>& OutEntities)
 {
 	FMassArchetypeData* ArchetypePtr = Archetype.DataPtr.Get();
 	check(ArchetypePtr);
@@ -347,7 +347,7 @@ TSharedRef<UMassEntitySubsystem::FEntityCreationContext> UMassEntitySubsystem::B
 	FEntityCreationContext* CreationContext = new FEntityCreationContext(Count);
 	// @todo this could probably be optimized since one would assume we're adding elements to OutEntities in order.
 	// Then again, if that's the case, the sorting will be almost instant
-	new (&CreationContext->ChunkCollection)FArchetypeChunkCollection(Archetype, MakeArrayView(&OutEntities[OutEntities.Num() - Count], Count), FArchetypeChunkCollection::NoDuplicates);
+	new (&CreationContext->ChunkCollection)FMassArchetypeSubChunks(Archetype, MakeArrayView(&OutEntities[OutEntities.Num() - Count], Count), FMassArchetypeSubChunks::NoDuplicates);
 	if (ObserverManager.HasOnAddedObserversForFragments(ArchetypePtr->GetCompositionDescriptor().Fragments))
 	{
 		CreationContext->OnSpawningFinished = [this](FEntityCreationContext& Context){
@@ -408,7 +408,7 @@ void UMassEntitySubsystem::BatchDestroyEntities(TConstArrayView<FMassEntityHandl
 	}
 }
 
-void UMassEntitySubsystem::BatchDestroyEntityChunks(const FArchetypeChunkCollection& Chunks)
+void UMassEntitySubsystem::BatchDestroyEntityChunks(const FMassArchetypeSubChunks& Chunks)
 {
 	TRACE_CPUPROFILER_EVENT_SCOPE_STR("BatchDestroyEntityChunks");
 
@@ -421,7 +421,7 @@ void UMassEntitySubsystem::BatchDestroyEntityChunks(const FArchetypeChunkCollect
 
 	check(Chunks.GetArchetype().IsValid());
 	TArray<FMassEntityHandle> EntitiesRemoved;
-	Chunks.GetArchetype().DataPtr->BatchDestroyEntityChunks(Chunks, EntitiesRemoved);
+	Chunks.GetArchetype().DataPtr->BatchDestroyEntityChunks(Chunks.GetChunks(), EntitiesRemoved);
 	
 	EntityFreeIndexList.Reserve(EntityFreeIndexList.Num() + EntitiesRemoved.Num());
 
@@ -476,7 +476,7 @@ void UMassEntitySubsystem::AddCompositionToEntity_GetDelta(FMassEntityHandle Ent
 		NewDescriptor.Fragments += InDescriptor.Fragments;
 		NewDescriptor.Tags+= InDescriptor.Tags;
 
-		const FArchetypeHandle NewArchetypeHandle = CreateArchetype(NewDescriptor, OldArchetype->GetSharedFragmentValues());
+		const FMassArchetypeHandle NewArchetypeHandle = CreateArchetype(NewDescriptor, OldArchetype->GetSharedFragmentValues());
 
 		if (ensure(NewArchetypeHandle.DataPtr != EntityData.CurrentArchetype))
 		{
@@ -513,7 +513,7 @@ void UMassEntitySubsystem::RemoveCompositionFromEntity(FMassEntityHandle Entity,
 			ensureMsgf(OldArchetype->GetCompositionDescriptor().HasAll(InDescriptor), TEXT("Some of the elements being removed are already missing from entity\'s composition."));
 			ObserverManager.OnPreCompositionRemoved(Entity, InDescriptor);
 
-			const FArchetypeHandle NewArchetypeHandle = CreateArchetype(NewDescriptor, OldArchetype->GetSharedFragmentValues());
+			const FMassArchetypeHandle NewArchetypeHandle = CreateArchetype(NewDescriptor, OldArchetype->GetSharedFragmentValues());
 
 			if (ensure(NewArchetypeHandle.DataPtr != EntityData.CurrentArchetype))
 			{
@@ -527,12 +527,12 @@ void UMassEntitySubsystem::RemoveCompositionFromEntity(FMassEntityHandle Entity,
 	}
 }
 
-const FMassArchetypeCompositionDescriptor& UMassEntitySubsystem::GetArchetypeComposition(const FArchetypeHandle& ArchetypeHandle) const
+const FMassArchetypeCompositionDescriptor& UMassEntitySubsystem::GetArchetypeComposition(const FMassArchetypeHandle& ArchetypeHandle) const
 {
 	return ArchetypeHandle.DataPtr->GetCompositionDescriptor();
 }
 
-void UMassEntitySubsystem::InternalBuildEntity(FMassEntityHandle Entity, const FArchetypeHandle Archetype)
+void UMassEntitySubsystem::InternalBuildEntity(FMassEntityHandle Entity, const FMassArchetypeHandle Archetype)
 {
 	FEntityData& EntityData = Entities[Entity.Index];
 	EntityData.CurrentArchetype = Archetype.DataPtr;
@@ -571,7 +571,7 @@ void UMassEntitySubsystem::InternalAddFragmentListToEntity(FMassEntityHandle Ent
 	check(EntityData.CurrentArchetype.IsValid());
 
 	// fetch or create the new archetype
-	const FArchetypeHandle NewArchetypeHandle = CreateArchetype(EntityData.CurrentArchetype, NewFragments);
+	const FMassArchetypeHandle NewArchetypeHandle = CreateArchetype(EntityData.CurrentArchetype, NewFragments);
 
 	if (NewArchetypeHandle.DataPtr != EntityData.CurrentArchetype)
 	{
@@ -617,7 +617,7 @@ void UMassEntitySubsystem::RemoveFragmentListFromEntity(FMassEntityHandle Entity
 	{
 		// If all the fragments got removed this will result in fetching of the empty archetype
 		const FMassArchetypeCompositionDescriptor NewComposition(OldArchetype->GetFragmentBitSet() - FragmentsToRemove, OldArchetype->GetTagBitSet(), OldArchetype->GetChunkFragmentBitSet(), OldArchetype->GetSharedFragmentBitSet());
-		const FArchetypeHandle NewArchetypeHandle = CreateArchetype(NewComposition, OldArchetype->GetSharedFragmentValues());
+		const FMassArchetypeHandle NewArchetypeHandle = CreateArchetype(NewComposition, OldArchetype->GetSharedFragmentValues());
 
 		// Move the entity over
 		FMassArchetypeData* NewArchetype = NewArchetypeHandle.DataPtr.Get();
@@ -646,7 +646,7 @@ void UMassEntitySubsystem::SwapTagsForEntity(FMassEntityHandle Entity, const USc
 	
 	if (NewTagBitSet != CurrentArchetype->GetTagBitSet())
 	{
-		const FArchetypeHandle NewArchetypeHandle = InternalCreateSiblingArchetype(EntityData.CurrentArchetype, NewTagBitSet);
+		const FMassArchetypeHandle NewArchetypeHandle = InternalCreateSiblingArchetype(EntityData.CurrentArchetype, NewTagBitSet);
 		checkSlow(NewArchetypeHandle.IsValid());
 
 		// Move the entity over
@@ -670,7 +670,7 @@ void UMassEntitySubsystem::AddTagToEntity(FMassEntityHandle Entity, const UScrip
 		//FMassTagBitSet NewTags = CurrentArchetype->GetTagBitSet() - *TagType;
 		FMassTagBitSet NewTags = CurrentArchetype->GetTagBitSet();
 		NewTags.Add(*TagType);
-		const FArchetypeHandle NewArchetypeHandle = InternalCreateSiblingArchetype(EntityData.CurrentArchetype, NewTags);
+		const FMassArchetypeHandle NewArchetypeHandle = InternalCreateSiblingArchetype(EntityData.CurrentArchetype, NewTags);
 		checkSlow(NewArchetypeHandle.IsValid());
 
 		// Move the entity over
@@ -694,7 +694,7 @@ void UMassEntitySubsystem::RemoveTagFromEntity(FMassEntityHandle Entity, const U
 		// CurrentArchetype->GetTagBitSet() -  *TagType
 		FMassTagBitSet NewTags = CurrentArchetype->GetTagBitSet();
 		NewTags.Remove(*TagType);
-		const FArchetypeHandle NewArchetypeHandle = InternalCreateSiblingArchetype(EntityData.CurrentArchetype, NewTags);
+		const FMassArchetypeHandle NewArchetypeHandle = InternalCreateSiblingArchetype(EntityData.CurrentArchetype, NewTags);
 		checkSlow(NewArchetypeHandle.IsValid());
 
 		// Move the entity over
@@ -704,7 +704,7 @@ void UMassEntitySubsystem::RemoveTagFromEntity(FMassEntityHandle Entity, const U
 }
 
 
-void UMassEntitySubsystem::MoveEntityToAnotherArchetype(FMassEntityHandle Entity, FArchetypeHandle NewArchetypeHandle)
+void UMassEntitySubsystem::MoveEntityToAnotherArchetype(FMassEntityHandle Entity, FMassArchetypeHandle NewArchetypeHandle)
 {
 	CheckIfEntityIsActive(Entity);
 
@@ -725,14 +725,16 @@ void UMassEntitySubsystem::SetEntityFragmentsValues(FMassEntityHandle Entity, TA
 	EntityData.CurrentArchetype->SetFragmentsData(Entity, FragmentInstanceList);
 }
 
-void UMassEntitySubsystem::BatchSetEntityFragmentsValues(const FArchetypeChunkCollection& SparseEntities, TArrayView<const FInstancedStruct> FragmentInstanceList)
+void UMassEntitySubsystem::BatchSetEntityFragmentsValues(const FMassArchetypeSubChunks& SparseEntities, TArrayView<const FInstancedStruct> FragmentInstanceList)
 {
 	FMassArchetypeData* Archetype = SparseEntities.GetArchetype().DataPtr.Get();
 	check(Archetype);
 
+	const FMassArchetypeSubChunks::FConstSubChunkArrayView SparseChunks = SparseEntities.GetChunks();
+
 	for (const FInstancedStruct& FragmentTemplate : FragmentInstanceList)
 	{
-		Archetype->SetFragmentData(SparseEntities, FragmentTemplate);
+		Archetype->SetFragmentData(SparseChunks, FragmentTemplate);
 	}
 }
 
@@ -775,7 +777,7 @@ void UMassEntitySubsystem::CheckIfEntityIsActive(FMassEntityHandle Entity) const
 	checkf(IsEntityBuilt(Entity), TEXT("Entity not yet created(ID: %d, SN:%d)"));
 }
 
-void UMassEntitySubsystem::GetValidArchetypes(const FMassEntityQuery& Query, TArray<FArchetypeHandle>& OutValidArchetypes)
+void UMassEntitySubsystem::GetValidArchetypes(const FMassEntityQuery& Query, TArray<FMassArchetypeHandle>& OutValidArchetypes)
 {
 	//@TODO: Not optimized yet, but we call this rarely now, so not a big deal.
 
@@ -1007,7 +1009,7 @@ void UMassEntitySubsystem::DebugPrintEntity(FMassEntityHandle Entity, FOutputDev
 	}
 }
 
-void UMassEntitySubsystem::DebugGetStringDesc(const FArchetypeHandle& Archetype, FOutputDevice& Ar)
+void UMassEntitySubsystem::DebugGetStringDesc(const FMassArchetypeHandle& Archetype, FOutputDevice& Ar)
 {
 	Ar.Logf(TEXT("%s"), Archetype.IsValid() ? *Archetype.DataPtr->DebugGetDescription() : TEXT("INVALID"));
 }
@@ -1031,7 +1033,7 @@ void UMassEntitySubsystem::DebugGetArchetypesStringDetails(FOutputDevice& Ar, co
 #endif // WITH_MASSENTITY_DEBUG
 }
 
-void UMassEntitySubsystem::DebugGetArchetypeFragmentTypes(const FArchetypeHandle& Archetype, TArray<const UScriptStruct*>& InOutFragmentList) const
+void UMassEntitySubsystem::DebugGetArchetypeFragmentTypes(const FMassArchetypeHandle& Archetype, TArray<const UScriptStruct*>& InOutFragmentList) const
 {
 	if (Archetype.DataPtr.IsValid())
 	{
@@ -1039,12 +1041,12 @@ void UMassEntitySubsystem::DebugGetArchetypeFragmentTypes(const FArchetypeHandle
 	}
 }
 
-int32 UMassEntitySubsystem::DebugGetArchetypeEntitiesCount(const FArchetypeHandle& Archetype) const
+int32 UMassEntitySubsystem::DebugGetArchetypeEntitiesCount(const FMassArchetypeHandle& Archetype) const
 {
 	return Archetype.DataPtr.IsValid() ? Archetype.DataPtr->GetNumEntities() : 0;
 }
 
-int32 UMassEntitySubsystem::DebugGetArchetypeEntitiesCountPerChunk(const FArchetypeHandle& Archetype) const
+int32 UMassEntitySubsystem::DebugGetArchetypeEntitiesCountPerChunk(const FMassArchetypeHandle& Archetype) const
 {
 	return Archetype.DataPtr.IsValid() ? Archetype.DataPtr->GetNumEntitiesPerChunk() : 0;
 }
@@ -1070,7 +1072,7 @@ void UMassEntitySubsystem::DebugRemoveAllEntities()
 	}
 }
 
-void UMassEntitySubsystem::DebugGetArchetypeStrings(const FArchetypeHandle& Archetype, TArray<FName>& OutFragmentNames, TArray<FName>& OutTagNames)
+void UMassEntitySubsystem::DebugGetArchetypeStrings(const FMassArchetypeHandle& Archetype, TArray<FName>& OutFragmentNames, TArray<FName>& OutTagNames)
 {
 	if (Archetype.IsValid() == false)
 	{
@@ -1127,18 +1129,13 @@ void FMassExecutionContext::ClearExecutionData()
 	ChunkSerialModificationNumber = -1;
 }
 
-void FMassExecutionContext::SetCurrentArchetypeData(FMassArchetypeData& ArchetypeData)
-{
-	CurrentArchetypesTagBitSet = ArchetypeData.GetTagBitSet();
-}
-
-void FMassExecutionContext::SetChunkCollection(const FArchetypeChunkCollection& InChunkCollection)
+void FMassExecutionContext::SetChunkCollection(const FMassArchetypeSubChunks& InChunkCollection)
 {
 	check(ChunkCollection.IsEmpty());
 	ChunkCollection = InChunkCollection;
 }
 
-void FMassExecutionContext::SetChunkCollection(FArchetypeChunkCollection&& InChunkCollection)
+void FMassExecutionContext::SetChunkCollection(FMassArchetypeSubChunks&& InChunkCollection)
 {
 	check(ChunkCollection.IsEmpty());
 	ChunkCollection = MoveTemp(InChunkCollection);
diff --git a/Engine/Plugins/Runtime/MassEntity/Source/MassEntity/Private/MassEntityUtils.cpp b/Engine/Plugins/Runtime/MassEntity/Source/MassEntity/Private/MassEntityUtils.cpp
--- a/Engine/Plugins/Runtime/MassEntity/Source/MassEntity/Private/MassEntityUtils.cpp
+++ b/Engine/Plugins/Runtime/MassEntity/Source/MassEntity/Private/MassEntityUtils.cpp
@@ -35,15 +35,15 @@ EProcessorExecutionFlags GetProcessorExecutionFlagsForWold(const UWorld& World)
 }
 
 void CreateSparseChunks(const UMassEntitySubsystem& EntitySystem, const TConstArrayView<FMassEntityHandle> Entities
-	, const FArchetypeChunkCollection::EDuplicatesHandling DuplicatesHandling, TArray<FArchetypeChunkCollection>& OutChunkCollections)
+	, const FMassArchetypeSubChunks::EDuplicatesHandling DuplicatesHandling, TArray<FMassArchetypeSubChunks >& OutChunkCollections)
 {
 	TRACE_CPUPROFILER_EVENT_SCOPE_STR("Mass_CreateSparseChunks");
 
-	TMap<const FArchetypeHandle, TArray<FMassEntityHandle>> ArchetypeToEntities;
+	TMap<const FMassArchetypeHandle, TArray<FMassEntityHandle>> ArchetypeToEntities;
 
 	for (const FMassEntityHandle& Entity : Entities)
 	{
-		FArchetypeHandle Archetype = EntitySystem.GetArchetypeForEntity(Entity);
+		FMassArchetypeHandle Archetype = EntitySystem.GetArchetypeForEntity(Entity);
 		TArray<FMassEntityHandle>& PerArchetypeEntities = ArchetypeToEntities.FindOrAdd(Archetype);
 		PerArchetypeEntities.Add(Entity);
 	}
@@ -54,7 +54,7 @@ void CreateSparseChunks(const UMassEntitySubsystem& EntitySystem, const TConstAr
 		// UMassSpawnerSubsystem::DestroyEntities' caller knowledge. Should be removed once that's addressed.
 		if (Pair.Key.IsValid())
 		{
-			OutChunkCollections.Add(FArchetypeChunkCollection(Pair.Key, Pair.Value, DuplicatesHandling));
+			OutChunkCollections.Add(FMassArchetypeSubChunks(Pair.Key, Pair.Value, DuplicatesHandling));
 		}
 	}
 }
diff --git a/Engine/Plugins/Runtime/MassEntity/Source/MassEntity/Private/MassEntityView.cpp b/Engine/Plugins/Runtime/MassEntity/Source/MassEntity/Private/MassEntityView.cpp
--- a/Engine/Plugins/Runtime/MassEntity/Source/MassEntity/Private/MassEntityView.cpp
+++ b/Engine/Plugins/Runtime/MassEntity/Source/MassEntity/Private/MassEntityView.cpp
@@ -7,7 +7,7 @@
 
 //////////////////////////////////////////////////////////////////////
 // FMassEntityView
-FMassEntityView::FMassEntityView(const FArchetypeHandle& ArchetypeHandle, FMassEntityHandle InEntity)
+FMassEntityView::FMassEntityView(const FMassArchetypeHandle& ArchetypeHandle, FMassEntityHandle InEntity)
 {
 	Entity = InEntity;
 	check(ArchetypeHandle.IsValid());
@@ -18,7 +18,7 @@ FMassEntityView::FMassEntityView(const FArchetypeHandle& ArchetypeHandle, FMassE
 FMassEntityView::FMassEntityView(const UMassEntitySubsystem& EntitySubsystem, FMassEntityHandle InEntity)
 {
 	Entity = InEntity;
-	const FArchetypeHandle ArchetypeHandle = EntitySubsystem.GetArchetypeForEntity(Entity);
+	const FMassArchetypeHandle ArchetypeHandle = EntitySubsystem.GetArchetypeForEntity(Entity);
 	check(ArchetypeHandle.IsValid());
 	Archetype = ArchetypeHandle.DataPtr.Get();
 	EntityHandle = Archetype->MakeEntityHandle(Entity);
diff --git a/Engine/Plugins/Runtime/MassEntity/Source/MassEntity/Private/MassExecutor.cpp b/Engine/Plugins/Runtime/MassEntity/Source/MassEntity/Private/MassExecutor.cpp
--- a/Engine/Plugins/Runtime/MassEntity/Source/MassEntity/Private/MassExecutor.cpp
+++ b/Engine/Plugins/Runtime/MassEntity/Source/MassEntity/Private/MassExecutor.cpp
@@ -21,7 +21,7 @@ void Run(FMassRuntimePipeline& RuntimePipeline, FMassProcessingContext& Processi
 	RunProcessorsView(RuntimePipeline.Processors, ProcessingContext);
 }
 
-void RunSparse(FMassRuntimePipeline& RuntimePipeline, FMassProcessingContext& ProcessingContext, FArchetypeHandle Archetype, TConstArrayView<FMassEntityHandle> Entities)
+void RunSparse(FMassRuntimePipeline& RuntimePipeline, FMassProcessingContext& ProcessingContext, FMassArchetypeHandle Archetype, TConstArrayView<FMassEntityHandle> Entities)
 {
 	if (!ensure(ProcessingContext.EntitySubsystem) ||
 		!ensure(RuntimePipeline.Processors.Find(nullptr) == INDEX_NONE) ||
@@ -33,11 +33,11 @@ void RunSparse(FMassRuntimePipeline& RuntimePipeline, FMassProcessingContext& Pr
 
 	TRACE_CPUPROFILER_EVENT_SCOPE_STR("MassExecutor RunSparseEntities");
 
-	const FArchetypeChunkCollection ChunkCollection(Archetype, Entities, FArchetypeChunkCollection::NoDuplicates);
+	const FMassArchetypeSubChunks ChunkCollection(Archetype, Entities, FMassArchetypeSubChunks::NoDuplicates);
 	RunProcessorsView(RuntimePipeline.Processors, ProcessingContext, &ChunkCollection);
 }
 
-void RunSparse(FMassRuntimePipeline& RuntimePipeline, FMassProcessingContext& ProcessingContext, const FArchetypeChunkCollection& ChunkCollection)
+void RunSparse(FMassRuntimePipeline& RuntimePipeline, FMassProcessingContext& ProcessingContext, const FMassArchetypeSubChunks& ChunkCollection)
 {
 	if (!ensure(ProcessingContext.EntitySubsystem) ||
 		!ensure(RuntimePipeline.Processors.Find(nullptr) == INDEX_NONE) ||
@@ -65,7 +65,7 @@ void Run(UMassProcessor& Processor, FMassProcessingContext& ProcessingContext)
 	RunProcessorsView(MakeArrayView(&ProcPtr, 1), ProcessingContext);
 }
 
-void RunProcessorsView(TArrayView<UMassProcessor*> Processors, FMassProcessingContext& ProcessingContext, const FArchetypeChunkCollection* ChunkCollection)
+void RunProcessorsView(TArrayView<UMassProcessor*> Processors, FMassProcessingContext& ProcessingContext, const FMassArchetypeSubChunks* ChunkCollection)
 {
 	TRACE_CPUPROFILER_EVENT_SCOPE(RunProcessorsView);
 
diff --git a/Engine/Plugins/Runtime/MassEntity/Source/MassEntity/Private/MassObserverManager.cpp b/Engine/Plugins/Runtime/MassEntity/Source/MassEntity/Private/MassObserverManager.cpp
--- a/Engine/Plugins/Runtime/MassEntity/Source/MassEntity/Private/MassObserverManager.cpp
+++ b/Engine/Plugins/Runtime/MassEntity/Source/MassEntity/Private/MassObserverManager.cpp
@@ -66,7 +66,7 @@ void FMassObserverManager::Initialize()
 	}
 }
 
-bool FMassObserverManager::OnPostEntitiesCreated(const FArchetypeChunkCollection& ChunkCollection)
+bool FMassObserverManager::OnPostEntitiesCreated(const FMassArchetypeSubChunks& ChunkCollection)
 {
 	FMassProcessingContext ProcessingContext(EntitySubsystem, /*DeltaSeconds=*/0.f);
 	// requesting not to flush commands since handling creation of new entities can result in multiple collections of
@@ -83,7 +83,7 @@ bool FMassObserverManager::OnPostEntitiesCreated(const FArchetypeChunkCollection
 	return false;
 }
 
-bool FMassObserverManager::OnPostEntitiesCreated(FMassProcessingContext& ProcessingContext, const FArchetypeChunkCollection& ChunkCollection)
+bool FMassObserverManager::OnPostEntitiesCreated(FMassProcessingContext& ProcessingContext, const FMassArchetypeSubChunks& ChunkCollection)
 {
 	TRACE_CPUPROFILER_EVENT_SCOPE_STR("OnPostEntitiesCreated")
 
@@ -100,7 +100,7 @@ bool FMassObserverManager::OnPostEntitiesCreated(FMassProcessingContext& Process
 	return false;
 }
 
-bool FMassObserverManager::OnPreEntitiesDestroyed(const FArchetypeChunkCollection& ChunkCollection)
+bool FMassObserverManager::OnPreEntitiesDestroyed(const FMassArchetypeSubChunks& ChunkCollection)
 {
 	FMassProcessingContext ProcessingContext(EntitySubsystem, /*DeltaSeconds=*/0.f);
 	ProcessingContext.bFlushCommandBuffer = false;
@@ -114,7 +114,7 @@ bool FMassObserverManager::OnPreEntitiesDestroyed(const FArchetypeChunkCollectio
 	return false;
 }
 
-bool FMassObserverManager::OnPreEntitiesDestroyed(FMassProcessingContext& ProcessingContext, const FArchetypeChunkCollection& ChunkCollection)
+bool FMassObserverManager::OnPreEntitiesDestroyed(FMassProcessingContext& ProcessingContext, const FMassArchetypeSubChunks& ChunkCollection)
 {
 	TRACE_CPUPROFILER_EVENT_SCOPE_STR("OnPreEntitiesDestroyed")
 
@@ -137,9 +137,9 @@ bool FMassObserverManager::OnPostCompositionAdded(const FMassEntityHandle Entity
 	if (Overlap.IsEmpty() == false)
 	{
 		FMassProcessingContext ProcessingContext(EntitySubsystem, /*DeltaSeconds=*/0.f);
-		const FArchetypeHandle ArchetypeHandle = EntitySubsystem.GetArchetypeForEntity(Entity);
-		HandleFragmentsImpl(ProcessingContext, FArchetypeChunkCollection(ArchetypeHandle, MakeArrayView(&Entity, 1)
-			, FArchetypeChunkCollection::NoDuplicates), Overlap, OnFragmentAddedObservers);
+		const FMassArchetypeHandle ArchetypeHandle = EntitySubsystem.GetArchetypeForEntity(Entity);
+		HandleFragmentsImpl(ProcessingContext, FMassArchetypeSubChunks(ArchetypeHandle, MakeArrayView(&Entity, 1)
+			, FMassArchetypeSubChunks::NoDuplicates), Overlap, OnFragmentAddedObservers);
 		return true;
 	}
 
@@ -152,16 +152,16 @@ bool FMassObserverManager::OnPreCompositionRemoved(const FMassEntityHandle Entit
 	if (Overlap.IsEmpty() == false)
 	{
 		FMassProcessingContext ProcessingContext(EntitySubsystem, /*DeltaSeconds=*/0.f);
-		const FArchetypeHandle ArchetypeHandle = EntitySubsystem.GetArchetypeForEntity(Entity);
-		HandleFragmentsImpl(ProcessingContext, FArchetypeChunkCollection(ArchetypeHandle, MakeArrayView(&Entity, 1)
-			, FArchetypeChunkCollection::NoDuplicates), Overlap, OnFragmentRemovedObservers);
+		const FMassArchetypeHandle ArchetypeHandle = EntitySubsystem.GetArchetypeForEntity(Entity);
+		HandleFragmentsImpl(ProcessingContext, FMassArchetypeSubChunks(ArchetypeHandle, MakeArrayView(&Entity, 1)
+			, FMassArchetypeSubChunks::NoDuplicates), Overlap, OnFragmentRemovedObservers);
 		return true;
 	}
 
 	return false;
 }
 
-void FMassObserverManager::HandleFragmentsImpl(FMassProcessingContext& ProcessingContext, const FArchetypeChunkCollection& ChunkCollection, const FMassFragmentBitSet& FragmentsBitSet, TMap<const UScriptStruct*, FMassRuntimePipeline>& HandlersContainer)
+void FMassObserverManager::HandleFragmentsImpl(FMassProcessingContext& ProcessingContext, const FMassArchetypeSubChunks& ChunkCollection, const FMassFragmentBitSet& FragmentsBitSet, TMap<const UScriptStruct*, FMassRuntimePipeline>& HandlersContainer)
 {
 	TArray<const UScriptStruct*> Fragments;
 	FragmentsBitSet.ExportTypes(Fragments);
@@ -175,7 +175,7 @@ void FMassObserverManager::HandleFragmentsImpl(FMassProcessingContext& Processin
 	}
 }
 
-void FMassObserverManager::HandleSingleFragmentImpl(const UScriptStruct& FragmentType, const FArchetypeChunkCollection& ChunkCollection, const FMassFragmentBitSet& FragmentFilterBitSet, TMap<const UScriptStruct*, FMassRuntimePipeline>& HandlersContainer)
+void FMassObserverManager::HandleSingleFragmentImpl(const UScriptStruct& FragmentType, const FMassArchetypeSubChunks& ChunkCollection, const FMassFragmentBitSet& FragmentFilterBitSet, TMap<const UScriptStruct*, FMassRuntimePipeline>& HandlersContainer)
 {
 	if (FragmentFilterBitSet.Contains(FragmentType))
 	{
diff --git a/Engine/Plugins/Runtime/MassEntity/Source/MassEntity/Private/MassProcessingPhase.cpp b/Engine/Plugins/Runtime/MassEntity/Source/MassEntity/Private/MassProcessingPhase.cpp
--- a/Engine/Plugins/Runtime/MassEntity/Source/MassEntity/Private/MassProcessingPhase.cpp
+++ b/Engine/Plugins/Runtime/MassEntity/Source/MassEntity/Private/MassProcessingPhase.cpp
@@ -24,7 +24,7 @@ namespace FMassTweakables
 	};
 }
 
-namespace UE::MassEntity::Internal  
+namespace UE::MassEntity::Private
 {
 	ETickingGroup PhaseToTickingGroup[int(EMassProcessingPhase::MAX)]
 	{
@@ -35,7 +35,7 @@ namespace UE::MassEntity::Internal
 		ETickingGroup::TG_PostPhysics,	// EMassProcessingPhase::PostPhysics
 		ETickingGroup::TG_LastDemotable, // EMassProcessingPhase::FrameEnd
 	};
-} // UE::MassEntity::Internal
+} // UE::MassEntity::Private
 
 //----------------------------------------------------------------------//
 //  FMassProcessingPhase
@@ -293,7 +293,7 @@ void UMassProcessingPhaseManager::CreatePhases()
 	for (int i = 0; i < int(EMassProcessingPhase::MAX); ++i)
 	{
 		ProcessingPhases[i].Phase = EMassProcessingPhase(i);
-		ProcessingPhases[i].TickGroup = UE::MassEntity::Internal::PhaseToTickingGroup[i];
+		ProcessingPhases[i].TickGroup = UE::MassEntity::Private::PhaseToTickingGroup[i];
 		UMassCompositeProcessor* PhaseProcessor = NewObject<UMassCompositeProcessor>(this, UMassCompositeProcessor::StaticClass()
 			, *FString::Printf(TEXT("ProcessingPhase_%s"), *EnumToString(EMassProcessingPhase(i))));
 		SetPhaseProcessor(EMassProcessingPhase(i), PhaseProcessor);
diff --git a/Engine/Plugins/Runtime/MassEntity/Source/MassEntity/Public/MassArchetypeTypes.h b/Engine/Plugins/Runtime/MassEntity/Source/MassEntity/Public/MassArchetypeTypes.h
--- a/Engine/Plugins/Runtime/MassEntity/Source/MassEntity/Public/MassArchetypeTypes.h
+++ b/Engine/Plugins/Runtime/MassEntity/Source/MassEntity/Public/MassArchetypeTypes.h
@@ -13,7 +13,7 @@ struct FMassExecutionContext;
 struct FMassFragment;
 struct FMassArchetypeChunkIterator;
 struct FMassEntityQuery;
-struct FArchetypeChunkCollection;
+struct FMassArchetypeSubChunks;
 struct FMassEntityView;
 
 typedef TFunction< void(FMassExecutionContext& /*ExecutionContext*/) > FMassExecuteFunction;
@@ -22,59 +22,58 @@ typedef TFunction< bool(const FMassExecutionContext& /*ExecutionContext*/) > FMa
 
 
 //////////////////////////////////////////////////////////////////////
-//
+// FMassArchetypeHandle
 
 // An opaque handle to an archetype
-struct FArchetypeHandle final
+struct FMassArchetypeHandle final
 {
-	FArchetypeHandle() = default;
+	FMassArchetypeHandle() = default;
 	bool IsValid() const { return DataPtr.IsValid(); }
 
-	MASSENTITY_API bool operator==(const FMassArchetypeData* Other) const;
-	bool operator==(const FArchetypeHandle& Other) const { return DataPtr == Other.DataPtr; }
-	bool operator!=(const FArchetypeHandle& Other) const { return DataPtr != Other.DataPtr; }
+	bool operator==(const FMassArchetypeHandle& Other) const { return DataPtr == Other.DataPtr; }
+	bool operator!=(const FMassArchetypeHandle& Other) const { return DataPtr != Other.DataPtr; }
 
-	MASSENTITY_API friend uint32 GetTypeHash(const FArchetypeHandle& Instance);
+	MASSENTITY_API friend uint32 GetTypeHash(const FMassArchetypeHandle& Instance);
 private:
-	FArchetypeHandle(const TSharedPtr<FMassArchetypeData>& InDataPtr)
+	FMassArchetypeHandle(const TSharedPtr<FMassArchetypeData>& InDataPtr)
 	: DataPtr(InDataPtr)
 	{}
 	TSharedPtr<FMassArchetypeData> DataPtr;
 
 	friend UMassEntitySubsystem;
-	friend FArchetypeChunkCollection;
+	friend FMassArchetypeSubChunks;
 	friend FMassEntityQuery;
 	friend FMassEntityView;
 };
 
 
 //////////////////////////////////////////////////////////////////////
-//
+// FMassArchetypeSubChunks 
 
 /** A struct that converts an arbitrary array of entities of given Archetype into a sequence of continuous
  *  entity chunks. The goal is to have the user create an instance of this struct once and run through a bunch of
  *  systems. The runtime code usually uses FMassArchetypeChunkIterator to iterate on the chunk collection.
  */
-struct MASSENTITY_API FArchetypeChunkCollection
+struct MASSENTITY_API FMassArchetypeSubChunks 
 {
 public:
-	struct FChunkInfo
+	struct FSubChunkInfo
 	{
 		int32 ChunkIndex = INDEX_NONE;
 		int32 SubchunkStart = 0;
 		/** negative or 0-length means "all available entities within chunk" */
 		int32 Length = 0;
 
-		FChunkInfo() = default;
-		explicit FChunkInfo(const int32 InChunkIndex, const int32 InSubchunkStart = 0, const int32 InLength = 0) : ChunkIndex(InChunkIndex), SubchunkStart(InSubchunkStart), Length(InLength) {}
+		FSubChunkInfo() = default;
+		explicit FSubChunkInfo(const int32 InChunkIndex, const int32 InSubchunkStart = 0, const int32 InLength = 0) : ChunkIndex(InChunkIndex), SubchunkStart(InSubchunkStart), Length(InLength) {}
 		/** Note that we consider invalid-length chunks valid as long as ChunkIndex and SubchunkStart are valid */
 		bool IsSet() const { return ChunkIndex != INDEX_NONE && SubchunkStart >= 0; }
 
-		bool operator==(const FChunkInfo& Other) const
+		bool operator==(const FSubChunkInfo& Other) const
 		{
 			return ChunkIndex == Other.ChunkIndex && SubchunkStart == Other.SubchunkStart && Length == Other.Length;
 		}
-		bool operator!=(const FChunkInfo& Other) const { return !(*this == Other); }
+		bool operator!=(const FSubChunkInfo& Other) const { return !(*this == Other); }
 	};
 
 	enum EDuplicatesHandling
@@ -85,36 +84,39 @@ public:
 						// will be processed and duplicates will be removed.
 	};
 
+	using FSubChunkArray = TArray<FSubChunkInfo>;
+	using FConstSubChunkArrayView = TConstArrayView<FSubChunkInfo>;
+
 private:
-	TArray<FChunkInfo> Chunks;
+	FSubChunkArray SubChunks;
 	/** entity indices indicated by SubChunks are only valid with given Archetype */
-	FArchetypeHandle Archetype;
+	FMassArchetypeHandle Archetype;
 
 public:
-	FArchetypeChunkCollection() = default;
-	FArchetypeChunkCollection(const FArchetypeHandle& InArchetype, TConstArrayView<FMassEntityHandle> InEntities, EDuplicatesHandling DuplicatesHandling);
-	explicit FArchetypeChunkCollection(FArchetypeHandle& InArchetypeHandle);
-	explicit FArchetypeChunkCollection(TSharedPtr<FMassArchetypeData>& InArchetype);
-
-	TArrayView<const FChunkInfo> GetChunks() const { return Chunks; }
-	const FArchetypeHandle& GetArchetype() const { return Archetype; }
-	bool IsEmpty() const { return Chunks.Num() == 0 && Archetype.IsValid() == false; }
+	FMassArchetypeSubChunks() = default;
+	FMassArchetypeSubChunks(const FMassArchetypeHandle& InArchetype, TConstArrayView<FMassEntityHandle> InEntities, EDuplicatesHandling DuplicatesHandling);
+	explicit FMassArchetypeSubChunks(FMassArchetypeHandle& InArchetypeHandle);
+	explicit FMassArchetypeSubChunks(TSharedPtr<FMassArchetypeData>& InArchetype);
+
+	FConstSubChunkArrayView GetChunks() const { return SubChunks; }
+	const FMassArchetypeHandle& GetArchetype() const { return Archetype; }
+	bool IsEmpty() const { return SubChunks.Num() == 0 && Archetype.IsValid() == false; }
 	bool IsSet() const { return Archetype.IsValid(); }
 	void Reset() 
 	{ 
-		Archetype = FArchetypeHandle();
-		Chunks.Reset();
+		Archetype = FMassArchetypeHandle();
+		SubChunks.Reset();
 	}
 
 	/** The comparison function that checks if Other is identical to this. Intended for diagnostics/debugging. */
-	bool IsSame(const FArchetypeChunkCollection& Other) const;
+	bool IsSame(const FMassArchetypeSubChunks& Other) const;
 
 private:
 	void GatherChunksFromArchetype(TSharedPtr<FMassArchetypeData>& InArchetype);
 };
 
 //////////////////////////////////////////////////////////////////////
-//
+// FMassArchetypeChunkIterator
 
 /**
  *  The type used to iterate over given archetype's chunks, be it full, continuous chunks or sparse subchunks. It hides
@@ -123,36 +125,40 @@ private:
 struct MASSENTITY_API FMassArchetypeChunkIterator
 {
 private:
-	const FArchetypeChunkCollection& ChunkData;
+	FMassArchetypeSubChunks::FConstSubChunkArrayView SubChunks;
 	int32 CurrentChunkIndex = 0;
 
 public:
-	explicit FMassArchetypeChunkIterator(const FArchetypeChunkCollection& InChunkData) : ChunkData(InChunkData), CurrentChunkIndex(0) {}
+	explicit FMassArchetypeChunkIterator(const FMassArchetypeSubChunks::FConstSubChunkArrayView& InSubChunks) : SubChunks(InSubChunks), CurrentChunkIndex(0) {}
 
-	operator bool() const { return ChunkData.GetChunks().IsValidIndex(CurrentChunkIndex) && ChunkData.GetChunks()[CurrentChunkIndex].IsSet(); }
+	operator bool() const { return SubChunks.IsValidIndex(CurrentChunkIndex) && SubChunks[CurrentChunkIndex].IsSet(); }
 	FMassArchetypeChunkIterator& operator++() { ++CurrentChunkIndex; return *this; }
 
-	const FArchetypeChunkCollection::FChunkInfo* operator->() const { check(bool(*this)); return &ChunkData.GetChunks()[CurrentChunkIndex]; }
-	const FArchetypeChunkCollection::FChunkInfo& operator*() const { check(bool(*this)); return ChunkData.GetChunks()[CurrentChunkIndex]; }
+	const FMassArchetypeSubChunks::FSubChunkInfo* operator->() const { check(bool(*this)); return &SubChunks[CurrentChunkIndex]; }
+	const FMassArchetypeSubChunks::FSubChunkInfo& operator*() const { check(bool(*this)); return SubChunks[CurrentChunkIndex]; }
 };
 
 //////////////////////////////////////////////////////////////////////
-//
-struct FInternalEntityHandle 
+// FMassRawEntityInChunkData
+
+struct FMassRawEntityInChunkData 
 {
-	FInternalEntityHandle() = default;
-	FInternalEntityHandle(uint8* InChunkRawMemory, const int32 InIndexWithinChunk)
+	FMassRawEntityInChunkData() = default;
+	FMassRawEntityInChunkData(uint8* InChunkRawMemory, const int32 InIndexWithinChunk)
         : ChunkRawMemory(InChunkRawMemory), IndexWithinChunk(InIndexWithinChunk)
 	{}
 	bool IsValid() const { return ChunkRawMemory != nullptr && IndexWithinChunk != INDEX_NONE; }
-	bool operator==(const FInternalEntityHandle & Other) const { return ChunkRawMemory == Other.ChunkRawMemory && IndexWithinChunk == Other.IndexWithinChunk; }
+	bool operator==(const FMassRawEntityInChunkData & Other) const { return ChunkRawMemory == Other.ChunkRawMemory && IndexWithinChunk == Other.IndexWithinChunk; }
 
 	uint8* ChunkRawMemory = nullptr;
 	int32 IndexWithinChunk = INDEX_NONE;
 };
 
-typedef TArray<int32, TInlineAllocator<16>> FMassFragmentIndicesMapping;
-typedef TConstArrayView<int32> FMassFragmentIndicesMappingView;
+//////////////////////////////////////////////////////////////////////
+// FMassQueryRequirementIndicesMapping
+
+using FMassFragmentIndicesMapping = TArray<int32, TInlineAllocator<16>>;
+
 struct FMassQueryRequirementIndicesMapping
 {
 	FMassQueryRequirementIndicesMapping() = default;
diff --git a/Engine/Plugins/Runtime/MassEntity/Source/MassEntity/Public/MassEntityQuery.h b/Engine/Plugins/Runtime/MassEntity/Source/MassEntity/Public/MassEntityQuery.h
--- a/Engine/Plugins/Runtime/MassEntity/Source/MassEntity/Public/MassEntityQuery.h
+++ b/Engine/Plugins/Runtime/MassEntity/Source/MassEntity/Public/MassEntityQuery.h
@@ -14,7 +14,7 @@ class UMassEntitySubsystem;
 struct FMassArchetypeData;
 struct FMassExecutionContext;
 struct FMassFragment;
-struct FArchetypeHandle;
+struct FMassArchetypeHandle;
 
 enum class EMassFragmentAccess : uint8
 {
@@ -106,7 +106,7 @@ public:
 	void ForEachEntityChunk(UMassEntitySubsystem& EntitySubsystem, FMassExecutionContext& ExecutionContext, const FMassExecuteFunction& ExecuteFunction);
 	
 	/** Will first verify that the archetype given with Chunks matches the query's requirements, and if so will run the other, more generic ForEachEntityChunk implementation */
-	void ForEachEntityChunk(const FArchetypeChunkCollection& Chunks, UMassEntitySubsystem& EntitySubsystem, FMassExecutionContext& ExecutionContext, const FMassExecuteFunction& ExecuteFunction);
+	void ForEachEntityChunk(const FMassArchetypeSubChunks& Chunks, UMassEntitySubsystem& EntitySubsystem, FMassExecutionContext& ExecutionContext, const FMassExecuteFunction& ExecuteFunction);
 
 	/**
 	 * Attempts to process every chunk of every affected archetype in parallel.
@@ -319,7 +319,7 @@ public:
 		ArchetypeDataVersion = 0;
 	}
 	
-	bool DoesArchetypeMatchRequirements(const FArchetypeHandle& ArchetypeHandle) const;
+	bool DoesArchetypeMatchRequirements(const FMassArchetypeHandle& ArchetypeHandle) const;
 
 	/** The function validates requirements we make for queries. See the FMassEntityQuery struct description for details.
 	 *  Note that this function is non-trivial and end users are not expected to need to use it. 
@@ -327,7 +327,7 @@ public:
 	bool CheckValidity() const;
 
 	FString DebugGetDescription() const;
-	FString DebugGetArchetypeCompatibilityDescription(const FArchetypeHandle& ArchetypeHandle) const;
+	FString DebugGetArchetypeCompatibilityDescription(const FMassArchetypeHandle& ArchetypeHandle) const;
 
 	TConstArrayView<FMassFragmentRequirement> GetRequirements() const { return Requirements; }
 	const FMassFragmentBitSet& GetRequiredAllFragments() const { return RequiredAllFragments; }
@@ -344,7 +344,7 @@ public:
 	const FMassSharedFragmentBitSet& GetRequiredOptionalSharedFragments() const { return RequiredOptionalSharedFragments; }
 	const FMassSharedFragmentBitSet& GetRequiredNoneSharedFragments() const { return RequiredNoneSharedFragments; }
 
-	const TArray<FArchetypeHandle>& GetArchetypes() const
+	const TArray<FMassArchetypeHandle>& GetArchetypes() const
 	{ 
 		return ValidArchetypes; 
 	}
@@ -365,7 +365,7 @@ public:
 
 	/** 
 	 * Sets a chunk filter condition that will applied to each chunk of all valid archetypes. Note 
-	 * that this condition won't be applied when a specific collection of chunks is used (via FArchetypeChunkCollection)
+	 * that this condition won't be applied when a specific collection of chunks is used (via FMassArchetypeSubChunks )
 	 * The value returned by InFunction controls whether to allow execution (true) or block it (false).
 	 */
 	void SetChunkFilter(const FMassChunkConditionFunction& InFunction) { ChunkCondition = InFunction; }
@@ -425,7 +425,7 @@ private:
 	uint32 EntitySubsystemHash = 0;
 	uint32 ArchetypeDataVersion = 0;
 
-	TArray<FArchetypeHandle> ValidArchetypes;
+	TArray<FMassArchetypeHandle> ValidArchetypes;
 	TArray<FMassQueryRequirementIndicesMapping> ArchetypeFragmentMapping;
 
 	bool bAllowParallelExecution = false;
diff --git a/Engine/Plugins/Runtime/MassEntity/Source/MassEntity/Public/MassEntitySubsystem.h b/Engine/Plugins/Runtime/MassEntity/Source/MassEntity/Public/MassEntitySubsystem.h
--- a/Engine/Plugins/Runtime/MassEntity/Source/MassEntity/Public/MassEntitySubsystem.h
+++ b/Engine/Plugins/Runtime/MassEntity/Source/MassEntity/Public/MassEntitySubsystem.h
@@ -17,8 +17,7 @@ struct FMassEntityQuery;
 struct FMassExecutionContext;
 struct FMassArchetypeData;
 struct FMassCommandBuffer;
-struct FArchetypeChunkCollection;
-struct FMassArchetypeChunk;
+struct FMassArchetypeSubChunks;
 class FOutputDevice;
 enum class EMassFragmentAccess : uint8;
 
@@ -83,7 +82,7 @@ public:
 	 * A special, relaxed but slower version of CreateArchetype functions that allows FragmentAngTagsList to contain 
 	 * both fragments and tags. 
 	 */
-	FArchetypeHandle CreateArchetype(TConstArrayView<const UScriptStruct*> FragmentsAndTagsList);
+	FMassArchetypeHandle CreateArchetype(TConstArrayView<const UScriptStruct*> FragmentsAndTagsList);
 
 	/**
 	 * CreateArchetype from a composition descriptor and initial values
@@ -92,7 +91,7 @@ public:
 	 * @param SharedFragmentValues are the actual pointer to shared fragments
 	 * @return a handle of a new archetype 
 	 */
-	FArchetypeHandle CreateArchetype(const FMassArchetypeCompositionDescriptor& Composition, const FMassArchetypeSharedFragmentValues& SharedFragmentValues);
+	FMassArchetypeHandle CreateArchetype(const FMassArchetypeCompositionDescriptor& Composition, const FMassArchetypeSharedFragmentValues& SharedFragmentValues);
 
 	/** 
 	 *  Creates an archetype like SourceArchetype + NewFragmentList. 
@@ -104,11 +103,11 @@ public:
 	 *   types that SourceArchetype doesn't already have. If the caller cannot guarantee it use of AddFragment functions
 	 *   family is recommended.
 	 */
-	FArchetypeHandle CreateArchetype(const TSharedPtr<FMassArchetypeData>& SourceArchetype, const FMassFragmentBitSet& NewFragmentList);
+	FMassArchetypeHandle CreateArchetype(const TSharedPtr<FMassArchetypeData>& SourceArchetype, const FMassFragmentBitSet& NewFragmentList);
 
-	FArchetypeHandle GetArchetypeForEntity(FMassEntityHandle Entity) const;
+	FMassArchetypeHandle GetArchetypeForEntity(FMassEntityHandle Entity) const;
 	/** Method to iterate on all the fragment types of an archetype */
-	static void ForEachArchetypeFragmentType(const FArchetypeHandle Archetype, TFunction< void(const UScriptStruct* /*FragmentType*/)> Function);
+	static void ForEachArchetypeFragmentType(const FMassArchetypeHandle Archetype, TFunction< void(const UScriptStruct* /*FragmentType*/)> Function);
 
 	/**
 	 * Go through all archetypes and compact entities
@@ -120,7 +119,7 @@ public:
 	 * Creates fully built entity ready to be used by the subsystem
 	 * @param Archetype you want this entity to be
 	 * @return FMassEntityHandle id of the newly created entity */
-	FMassEntityHandle CreateEntity(const FArchetypeHandle Archetype);
+	FMassEntityHandle CreateEntity(const FMassArchetypeHandle Archetype);
 
 	/**
 	 * Creates fully built entity ready to be used by the subsystem
@@ -138,11 +137,11 @@ public:
 		{}
 		~FEntityCreationContext() { if (OnSpawningFinished) OnSpawningFinished(*this); }
 
-		const FArchetypeChunkCollection& GetChunkCollection() const { return ChunkCollection; }
+		const FMassArchetypeSubChunks& GetChunkCollection() const { return ChunkCollection; }
 	private:
 		friend UMassEntitySubsystem;
 		int32 NumberSpawned;
-		FArchetypeChunkCollection ChunkCollection;
+		FMassArchetypeSubChunks ChunkCollection;
 		TFunction<void(FEntityCreationContext&)> OnSpawningFinished;
 	};
 
@@ -151,7 +150,7 @@ public:
 	 *  @param Count number of entities to create
 	 *  @param OutEntities the newly created entities are appended to given array, i.e. the pre-existing content of OutEntities won't be affected by the call
 	 *  @return a creation context that will notify all the interested observers about newly created fragments once the context is released */
-	TSharedRef<FEntityCreationContext> BatchCreateEntities(const FArchetypeHandle Archetype, const int32 Count, TArray<FMassEntityHandle>& OutEntities);
+	TSharedRef<FEntityCreationContext> BatchCreateEntities(const FMassArchetypeHandle Archetype, const int32 Count, TArray<FMassEntityHandle>& OutEntities);
 
 	/**
 	 * Destroys a fully built entity, use ReleaseReservedEntity if entity was not yet built.
@@ -167,7 +166,7 @@ public:
 	 * Builds an entity for it to be ready to be used by the subsystem
 	 * @param Entity to build which was retrieved with ReserveEntity() method
 	 * @param Archetype you want this entity to be*/
-	void BuildEntity(FMassEntityHandle Entity, const FArchetypeHandle Archetype);
+	void BuildEntity(FMassEntityHandle Entity, const FMassArchetypeHandle Archetype);
 
 	/**
 	 * Builds an entity for it to be ready to be used by the subsystem
@@ -186,7 +185,7 @@ public:
 	 */
 	void BatchDestroyEntities(TConstArrayView<FMassEntityHandle> InEntities);
 
-	void BatchDestroyEntityChunks(const FArchetypeChunkCollection& Chunks);
+	void BatchDestroyEntityChunks(const FMassArchetypeSubChunks& Chunks);
 
 	void AddFragmentToEntity(FMassEntityHandle Entity, const UScriptStruct* FragmentType);
 
@@ -213,23 +212,23 @@ public:
 	void AddCompositionToEntity_GetDelta(FMassEntityHandle Entity, FMassArchetypeCompositionDescriptor& InOutDescriptor);
 	void RemoveCompositionFromEntity(FMassEntityHandle Entity, const FMassArchetypeCompositionDescriptor& InDescriptor);
 
-	const FMassArchetypeCompositionDescriptor& GetArchetypeComposition(const FArchetypeHandle& ArchetypeHandle) const;
+	const FMassArchetypeCompositionDescriptor& GetArchetypeComposition(const FMassArchetypeHandle& ArchetypeHandle) const;
 
 	/** 
 	 * Moves an entity over to a new archetype by copying over fragments common to both archetypes
 	 * @param Entity is the entity to move 
 	 * @param NewArchetypeHandle the handle to the new archetype
 	 */
-	void MoveEntityToAnotherArchetype(FMassEntityHandle Entity, FArchetypeHandle NewArchetypeHandle);
+	void MoveEntityToAnotherArchetype(FMassEntityHandle Entity, FMassArchetypeHandle NewArchetypeHandle);
 
 	/** Copies values from FragmentInstanceList over to Entity's fragment. Caller is responsible for ensuring that 
 	 *  the given entity does have given fragments. Failing this assumption will cause a check-fail.*/
 	void SetEntityFragmentsValues(FMassEntityHandle Entity, TArrayView<const FInstancedStruct> FragmentInstanceList);
 
 	/** Copies values from FragmentInstanceList over to fragments of given entities collection. The caller is responsible 
-	 *  for ensuring that the given entity archetype (FArchetypeChunkCollection.Archetype) does have given fragments. 
+	 *  for ensuring that the given entity archetype (FMassArchetypeSubChunks .Archetype) does have given fragments. 
 	 *  Failing this assumption will cause a check-fail. */
-	static void BatchSetEntityFragmentsValues(const FArchetypeChunkCollection& SparseEntities, TArrayView<const FInstancedStruct> FragmentInstanceList);
+	static void BatchSetEntityFragmentsValues(const FMassArchetypeSubChunks& SparseEntities, TArrayView<const FInstancedStruct> FragmentInstanceList);
 
 	// Return true if it is an valid built entity
 	bool IsEntityActive(FMassEntityHandle Entity) const 
@@ -325,26 +324,26 @@ public:
 	void DebugPrintEntity(int32 Index, FOutputDevice& Ar, const TCHAR* InPrefix = TEXT("")) const;
 	void DebugPrintEntity(FMassEntityHandle Entity, FOutputDevice& Ar, const TCHAR* InPrefix = TEXT("")) const;
 	void DebugPrintArchetypes(FOutputDevice& Ar, const bool bIncludeEmpty = true) const;
-	static void DebugGetStringDesc(const FArchetypeHandle& Archetype, FOutputDevice& Ar);
+	static void DebugGetStringDesc(const FMassArchetypeHandle& Archetype, FOutputDevice& Ar);
 	void DebugGetArchetypesStringDetails(FOutputDevice& Ar, const bool bIncludeEmpty = true);
-	void DebugGetArchetypeFragmentTypes(const FArchetypeHandle& Archetype, TArray<const UScriptStruct*>& InOutFragmentList) const;
-	int32 DebugGetArchetypeEntitiesCount(const FArchetypeHandle& Archetype) const;
-	int32 DebugGetArchetypeEntitiesCountPerChunk(const FArchetypeHandle& Archetype) const;
+	void DebugGetArchetypeFragmentTypes(const FMassArchetypeHandle& Archetype, TArray<const UScriptStruct*>& InOutFragmentList) const;
+	int32 DebugGetArchetypeEntitiesCount(const FMassArchetypeHandle& Archetype) const;
+	int32 DebugGetArchetypeEntitiesCountPerChunk(const FMassArchetypeHandle& Archetype) const;
 	int32 DebugGetEntityCount() const { return Entities.Num() - NumReservedEntities - EntityFreeIndexList.Num(); }
 	int32 DebugGetArchetypesCount() const { return FragmentHashToArchetypeMap.Num(); }
 	void DebugRemoveAllEntities();
 	void DebugForceArchetypeDataVersionBump() { ++ArchetypeDataVersion; }
-	void DebugGetArchetypeStrings(const FArchetypeHandle& Archetype, TArray<FName>& OutFragmentNames, TArray<FName>& OutTagNames);
+	void DebugGetArchetypeStrings(const FMassArchetypeHandle& Archetype, TArray<FName>& OutFragmentNames, TArray<FName>& OutTagNames);
 	FMassEntityHandle DebugGetEntityIndexHandle(const int32 EntityIndex) const { return Entities.IsValidIndex(EntityIndex) ? FMassEntityHandle(EntityIndex, Entities[EntityIndex].SerialNumber) : FMassEntityHandle(); }
 #endif // WITH_MASSENTITY_DEBUG
 
 protected:
-	void GetValidArchetypes(const FMassEntityQuery& Query, TArray<FArchetypeHandle>& OutValidArchetypes);
+	void GetValidArchetypes(const FMassEntityQuery& Query, TArray<FMassArchetypeHandle>& OutValidArchetypes);
 	
-	FArchetypeHandle InternalCreateSiblingArchetype(const TSharedPtr<FMassArchetypeData>& SourceArchetype, const FMassTagBitSet& OverrideTags);
+	FMassArchetypeHandle InternalCreateSiblingArchetype(const TSharedPtr<FMassArchetypeData>& SourceArchetype, const FMassTagBitSet& OverrideTags);
 
 private:
-	void InternalBuildEntity(FMassEntityHandle Entity, const FArchetypeHandle Archetype);
+	void InternalBuildEntity(FMassEntityHandle Entity, const FMassArchetypeHandle Archetype);
 	void InternalReleaseEntity(FMassEntityHandle Entity);
 
 	/** 
@@ -411,7 +410,7 @@ private:
 
 		bool operator==(const UScriptStruct* FragmentType) const { return Requirement.StructType == FragmentType; }
 	};
-	using FFragmentView=TFragmentView<TArrayView<FMassFragment>>;
+	using FFragmentView = TFragmentView<TArrayView<FMassFragment>>;
 	TArray<FFragmentView, TInlineAllocator<8>> FragmentViews;
 
 	using FChunkFragmentView = TFragmentView<FStructView>;
@@ -429,7 +428,7 @@ private:
 	
 	/** If set this indicates the exact archetype and its chunks to be processed. 
 	 *  @todo this data should live somewhere else, preferably be just a parameter to Query.ForEachEntityChunk function */
-	FArchetypeChunkCollection ChunkCollection;
+	FMassArchetypeSubChunks ChunkCollection;
 	
 	/** @todo rename to "payload" */
 	FInstancedStruct AuxData;
@@ -443,7 +442,7 @@ private:
 	
 	/** @todo update comment, and look if we If true the EntitySystem will flush the deferred commands stored in DeferredCommandBuffer just after executing 
 	 *  the given system. If False then the party calling UEntitySubsystem::ExecuteSystem is responsible for manually
-	 *  calling FLWCCommandBuffer.ReplayBufferAgainstSystem() */
+	 *  calling FMassCommandBuffer.ReplayBufferAgainstSystem() */
 	bool bFlushDeferredCommands = true;
 
 	TArrayView<FFragmentView> GetMutableRequirements() { return FragmentViews; }
@@ -470,8 +469,8 @@ public:
 	 *  immediate commands flushing */
 	void SetFlushDeferredCommands(const bool bNewFlushDeferredCommands) { bFlushDeferredCommands = bNewFlushDeferredCommands; } 
 	void SetDeferredCommandBuffer(const TSharedPtr<FMassCommandBuffer>& InDeferredCommandBuffer) { DeferredCommandBuffer = InDeferredCommandBuffer; }
-	void SetChunkCollection(const FArchetypeChunkCollection& InChunkCollection);
-	void SetChunkCollection(FArchetypeChunkCollection&& InChunkCollection);
+	void SetChunkCollection(const FMassArchetypeSubChunks& InChunkCollection);
+	void SetChunkCollection(FMassArchetypeSubChunks&& InChunkCollection);
 	void ClearChunkCollection() { ChunkCollection.Reset(); }
 	void SetAuxData(const FInstancedStruct& InAuxData) { AuxData = InAuxData; }
 
@@ -619,7 +618,7 @@ public:
 	}
 
 	/** Sparse chunk related operation */
-	const FArchetypeChunkCollection& GetChunkCollection() const { return ChunkCollection; }
+	const FMassArchetypeSubChunks& GetChunkCollection() const { return ChunkCollection; }
 
 	const FInstancedStruct& GetAuxData() const { return AuxData; }
 	FInstancedStruct& GetMutableAuxData() { return AuxData; }
@@ -634,7 +633,10 @@ public:
 	void FlushDeferred(UMassEntitySubsystem& EntitySystem) const;
 
 	void ClearExecutionData();
-	void SetCurrentArchetypeData(FMassArchetypeData& ArchetypeData);
+	void SetCurrentArchetypesTagBitSet(const FMassTagBitSet& BitSet)
+	{
+		CurrentArchetypesTagBitSet = BitSet;
+	}
 
 protected:
 	void SetRequirements(TConstArrayView<FMassFragmentRequirement> InRequirements, 
diff --git a/Engine/Plugins/Runtime/MassEntity/Source/MassEntity/Public/MassEntityUtils.h b/Engine/Plugins/Runtime/MassEntity/Source/MassEntity/Public/MassEntityUtils.h
--- a/Engine/Plugins/Runtime/MassEntity/Source/MassEntity/Public/MassEntityUtils.h
+++ b/Engine/Plugins/Runtime/MassEntity/Source/MassEntity/Public/MassEntityUtils.h
@@ -16,10 +16,10 @@ namespace UE::Mass::Utils
 MASSENTITY_API extern EProcessorExecutionFlags GetProcessorExecutionFlagsForWold(const UWorld& World);
 
 /** 
- * Fills OutChunkCollections with per-archetype FArchetypeChunkCollection instances. 
+ * Fills OutChunkCollections with per-archetype FMassArchetypeSubChunks instances. 
  * @param DuplicatesHandling used to inform the function whether to expect duplicates.
  */
 MASSENTITY_API extern void CreateSparseChunks(const UMassEntitySubsystem& EntitySystem, const TConstArrayView<FMassEntityHandle> Entities
-	, const FArchetypeChunkCollection::EDuplicatesHandling DuplicatesHandling, TArray<FArchetypeChunkCollection>& OutChunkCollections);
+	, const FMassArchetypeSubChunks::EDuplicatesHandling DuplicatesHandling, TArray<FMassArchetypeSubChunks >& OutChunkCollections);
 
 } // namespace UE::Mass::Utils
\ No newline at end of file
diff --git a/Engine/Plugins/Runtime/MassEntity/Source/MassEntity/Public/MassEntityView.h b/Engine/Plugins/Runtime/MassEntity/Source/MassEntity/Public/MassEntityView.h
--- a/Engine/Plugins/Runtime/MassEntity/Source/MassEntity/Public/MassEntityView.h
+++ b/Engine/Plugins/Runtime/MassEntity/Source/MassEntity/Public/MassEntityView.h
@@ -9,8 +9,7 @@
 
 class UMassEntitySubsystem;
 struct FMassArchetypeData;
-struct FArchetypeHandle;
-struct FMassArchetypeFragmentConfig;
+struct FMassArchetypeHandle;
 
 /** 
  * The type representing a single entity in a single archetype. It's of a very transient nature so we guarantee it's 
@@ -29,7 +28,7 @@ struct MASSENTITY_API FMassEntityView
 	 *  UMassEntitySubsystem-flavored constructor is recommended since it will first find the appropriate archetype for
 	 *  Entity. 
 	 */
-	FMassEntityView(const FArchetypeHandle& ArchetypeHandle, FMassEntityHandle Entity);
+	FMassEntityView(const FMassArchetypeHandle& ArchetypeHandle, FMassEntityHandle Entity);
 
 	/** 
 	 *  Finds the archetype Entity belongs to and then resolves against it. The caller is responsible for ensuring
@@ -136,6 +135,6 @@ protected:
 
 private:
 	FMassEntityHandle Entity;
-	FInternalEntityHandle EntityHandle;
+	FMassRawEntityInChunkData EntityHandle;
 	FMassArchetypeData* Archetype = nullptr;
 };
diff --git a/Engine/Plugins/Runtime/MassEntity/Source/MassEntity/Public/MassExecutor.h b/Engine/Plugins/Runtime/MassEntity/Source/MassEntity/Public/MassExecutor.h
--- a/Engine/Plugins/Runtime/MassEntity/Source/MassEntity/Public/MassExecutor.h
+++ b/Engine/Plugins/Runtime/MassEntity/Source/MassEntity/Public/MassExecutor.h
@@ -10,7 +10,7 @@
 struct FMassRuntimePipeline;
 struct FMassProcessingContext;
 struct FMassEntityHandle;
-struct FArchetypeChunkCollection;
+struct FMassArchetypeSubChunks;
 class UMassProcessor;
 
 
@@ -25,18 +25,18 @@ namespace UE::Mass::Executor
 	/** Similar to the Run function, but instead of using all the entities hosted by ProcessingContext.EntitySubsystem 
 	 *  it is processing only the entities given by EntityID via the Entities input parameter. 
 	 *  Note that all the entities need to be of Archetype archetype. 
-	 *  Under the hood the function converts Archetype-Entities pair to FArchetypeChunkCollection and calls the other flavor of RunSparse
+	 *  Under the hood the function converts Archetype-Entities pair to FMassArchetypeSubChunks and calls the other flavor of RunSparse
 	 */
-	MASSENTITY_API void RunSparse(FMassRuntimePipeline& RuntimePipeline, FMassProcessingContext& ProcessingContext, FArchetypeHandle Archetype, TConstArrayView<FMassEntityHandle> Entities);
+	MASSENTITY_API void RunSparse(FMassRuntimePipeline& RuntimePipeline, FMassProcessingContext& ProcessingContext, FMassArchetypeHandle Archetype, TConstArrayView<FMassEntityHandle> Entities);
 
 	/** Similar to the Run function, but instead of using all the entities hosted by ProcessingContext.EntitySubsystem 
 	 *  it is processing only the entities given by SparseEntities input parameter.
 	 *  @todo rename
 	 */
-	MASSENTITY_API void RunSparse(FMassRuntimePipeline& RuntimePipeline, FMassProcessingContext& ProcessingContext, const FArchetypeChunkCollection& ChunkCollection);
+	MASSENTITY_API void RunSparse(FMassRuntimePipeline& RuntimePipeline, FMassProcessingContext& ProcessingContext, const FMassArchetypeSubChunks& ChunkCollection);
 
 	/** Executes given Processors array view. This function gets called under the hood by the rest of Run* functions */
-	MASSENTITY_API void RunProcessorsView(TArrayView<UMassProcessor*> Processors, FMassProcessingContext& ProcessingContext, const FArchetypeChunkCollection* ChunkCollection = nullptr);
+	MASSENTITY_API void RunProcessorsView(TArrayView<UMassProcessor*> Processors, FMassProcessingContext& ProcessingContext, const FMassArchetypeSubChunks* ChunkCollection = nullptr);
 
 	/** 
 	 *  Triggers tasks executing Processor (and potentially it's children) and returns the task graph event representing 
diff --git a/Engine/Plugins/Runtime/MassEntity/Source/MassEntity/Public/MassObserverManager.h b/Engine/Plugins/Runtime/MassEntity/Source/MassEntity/Public/MassObserverManager.h
--- a/Engine/Plugins/Runtime/MassEntity/Source/MassEntity/Public/MassObserverManager.h
+++ b/Engine/Plugins/Runtime/MassEntity/Source/MassEntity/Public/MassObserverManager.h
@@ -7,7 +7,7 @@
 
 
 class UMassEntitySubsystem;
-struct FArchetypeChunkCollection;
+struct FMassArchetypeSubChunks;
 
 /** 
  * A type that encapsulates logic related to notifying interested parties of entity composition changes. Upon creation it
@@ -32,21 +32,21 @@ public:
 	bool HasOnAddedObserversForFragments(const FMassFragmentBitSet& InQueriedBitSet) const { return ObservedAddFragments.HasAny(InQueriedBitSet); }
 	bool HasOnRemovedObserversForFragments(const FMassFragmentBitSet& InQueriedBitSet) const { return ObservedRemoveFragments.HasAny(InQueriedBitSet); }
 
-	bool OnPostEntitiesCreated(const FArchetypeChunkCollection& ChunkCollection);
-	bool OnPostEntitiesCreated(FMassProcessingContext& ProcessingContext, const FArchetypeChunkCollection& ChunkCollection);
+	bool OnPostEntitiesCreated(const FMassArchetypeSubChunks& ChunkCollection);
+	bool OnPostEntitiesCreated(FMassProcessingContext& ProcessingContext, const FMassArchetypeSubChunks& ChunkCollection);
 
-	bool OnPreEntitiesDestroyed(const FArchetypeChunkCollection& ChunkCollection);
-	bool OnPreEntitiesDestroyed(FMassProcessingContext& ProcessingContext, const FArchetypeChunkCollection& ChunkCollection);
+	bool OnPreEntitiesDestroyed(const FMassArchetypeSubChunks& ChunkCollection);
+	bool OnPreEntitiesDestroyed(FMassProcessingContext& ProcessingContext, const FMassArchetypeSubChunks& ChunkCollection);
 
 	bool OnPostCompositionAdded(const FMassEntityHandle Entity, const FMassArchetypeCompositionDescriptor& Composition);
 	bool OnPreCompositionRemoved(const FMassEntityHandle Entity, const FMassArchetypeCompositionDescriptor& Composition);
 
-	void OnPostFragmentAdded(const UScriptStruct& FragmentType, const FArchetypeChunkCollection& ChunkCollection)
+	void OnPostFragmentAdded(const UScriptStruct& FragmentType, const FMassArchetypeSubChunks& ChunkCollection)
 	{
 		HandleSingleFragmentImpl(FragmentType, ChunkCollection, ObservedAddFragments, OnFragmentAddedObservers);
 	}
 
-	void OnPreFragmentRemoved(const UScriptStruct& FragmentType, const FArchetypeChunkCollection& ChunkCollection)
+	void OnPreFragmentRemoved(const UScriptStruct& FragmentType, const FMassArchetypeSubChunks& ChunkCollection)
 	{
 		HandleSingleFragmentImpl(FragmentType, ChunkCollection, ObservedRemoveFragments, OnFragmentRemovedObservers);
 	}
@@ -56,9 +56,9 @@ protected:
 	explicit FMassObserverManager(UMassEntitySubsystem& Owner);
 
 	void Initialize();
-	void HandleFragmentsImpl(FMassProcessingContext& ProcessingContext, const FArchetypeChunkCollection& ChunkCollection
+	void HandleFragmentsImpl(FMassProcessingContext& ProcessingContext, const FMassArchetypeSubChunks& ChunkCollection
 		, const FMassFragmentBitSet& FragmentsBitSet, TMap<const UScriptStruct*, FMassRuntimePipeline>& HandlersContainer);
-	void HandleSingleFragmentImpl(const UScriptStruct& FragmentType, const FArchetypeChunkCollection& ChunkCollection
+	void HandleSingleFragmentImpl(const UScriptStruct& FragmentType, const FMassArchetypeSubChunks& ChunkCollection
 		, const FMassFragmentBitSet& FragmentFilterBitSet, TMap<const UScriptStruct*, FMassRuntimePipeline>& HandlersContainer);
 
 	FOnObservedFragmentTypesChanged OnObservedFragmentTypesChangedDelegate;
diff --git a/Engine/Plugins/Runtime/MassEntity/Source/MassEntityTestSuite/Private/LWComponentTestFarmPlot.cpp b/Engine/Plugins/Runtime/MassEntity/Source/MassEntityTestSuite/Private/LWComponentTestFarmPlot.cpp
--- a/Engine/Plugins/Runtime/MassEntity/Source/MassEntityTestSuite/Private/LWComponentTestFarmPlot.cpp
+++ b/Engine/Plugins/Runtime/MassEntity/Source/MassEntityTestSuite/Private/LWComponentTestFarmPlot.cpp
@@ -64,7 +64,7 @@ ALWFragmentTestFarmPlot::ALWFragmentTestFarmPlot()
 	HarvestIconISMC->SetCollisionProfileName(UCollisionProfile::NoCollision_ProfileName);
 }
 
-void ALWFragmentTestFarmPlot::AddItemToGrid(UMassEntitySubsystem* EntitySystem, uint16 X, uint16 Y, FArchetypeHandle Archetype, uint16 VisualIndex)
+void ALWFragmentTestFarmPlot::AddItemToGrid(UMassEntitySubsystem* EntitySystem, uint16 X, uint16 Y, FMassArchetypeHandle Archetype, uint16 VisualIndex)
 {
 	FMassEntityHandle NewItem = EntitySystem->CreateEntity(Archetype);
 	PlantedSquares[X + Y * GridWidth] = NewItem;
@@ -93,8 +93,8 @@ void ALWFragmentTestFarmPlot::BeginPlay()
 
 	UMassEntitySubsystem* EntitySystem = UWorld::GetSubsystem<UMassEntitySubsystem>(GetWorld());
 
-	FArchetypeHandle CropArchetype = EntitySystem->CreateArchetype(TArray<const UScriptStruct*>{ FFarmWaterFragment::StaticStruct(), FFarmCropFragment::StaticStruct(), FHarvestTimerFragment::StaticStruct(), FFarmVisualFragment::StaticStruct(), FFarmGridCellData::StaticStruct() });
-	FArchetypeHandle FlowerArchetype = EntitySystem->CreateArchetype(TArray<const UScriptStruct*>{ FFarmWaterFragment::StaticStruct(), FFarmFlowerFragment::StaticStruct(), FHarvestTimerFragment::StaticStruct(), FFarmVisualFragment::StaticStruct(), FFarmGridCellData::StaticStruct() });
+	FMassArchetypeHandle CropArchetype = EntitySystem->CreateArchetype(TArray<const UScriptStruct*>{ FFarmWaterFragment::StaticStruct(), FFarmCropFragment::StaticStruct(), FHarvestTimerFragment::StaticStruct(), FFarmVisualFragment::StaticStruct(), FFarmGridCellData::StaticStruct() });
+	FMassArchetypeHandle FlowerArchetype = EntitySystem->CreateArchetype(TArray<const UScriptStruct*>{ FFarmWaterFragment::StaticStruct(), FFarmFlowerFragment::StaticStruct(), FHarvestTimerFragment::StaticStruct(), FFarmVisualFragment::StaticStruct(), FFarmGridCellData::StaticStruct() });
 
 	PerFrameSystems.Add(NewObject<UFarmWaterUpdateSystem>(this));
 
diff --git a/Engine/Plugins/Runtime/MassEntity/Source/MassEntityTestSuite/Private/MassChunkCollectionTest.cpp b/Engine/Plugins/Runtime/MassEntity/Source/MassEntityTestSuite/Private/MassChunkCollectionTest.cpp
--- a/Engine/Plugins/Runtime/MassEntity/Source/MassEntityTestSuite/Private/MassChunkCollectionTest.cpp
+++ b/Engine/Plugins/Runtime/MassEntity/Source/MassEntityTestSuite/Private/MassChunkCollectionTest.cpp
@@ -67,8 +67,8 @@ struct FChunkCollection_CreateBasic : FChunkCollectionTestBase
 		EntitiesSubSet.Add(Entities[11]);
 		EntitiesSubSet.Add(Entities[12]);
 
-		FArchetypeChunkCollection ChunkCollection(FloatsArchetype, EntitiesSubSet, FArchetypeChunkCollection::NoDuplicates);
-		TArrayView<const FArchetypeChunkCollection::FChunkInfo> Chunks = ChunkCollection.GetChunks();
+		FMassArchetypeSubChunks ChunkCollection(FloatsArchetype, EntitiesSubSet, FMassArchetypeSubChunks::NoDuplicates);
+		FMassArchetypeSubChunks::FConstSubChunkArrayView Chunks = ChunkCollection.GetChunks();
 		AITEST_EQUAL("The predicted sub-chunk count should match", Chunks.Num(), 4);
 		AITEST_EQUAL("The [10-13] chunk should be first and start at 10", Chunks[0].SubchunkStart, 10);
 		AITEST_EQUAL("The [10-13] chunk should be first and have a length of 4", Chunks[0].Length, 4);
@@ -91,17 +91,17 @@ struct FChunkCollection_CreateOrderInvariant : FChunkCollectionTestBase
 		TArray<FMassEntityHandle> EntitiesSubSet(&Entities[10], 30);
 		EntitiesSubSet.RemoveAt(10, 1, false);
 
-		FArchetypeChunkCollection CollectionFromOrdered(FloatsArchetype, EntitiesSubSet, FArchetypeChunkCollection::NoDuplicates);
+		FMassArchetypeSubChunks CollectionFromOrdered(FloatsArchetype, EntitiesSubSet, FMassArchetypeSubChunks::NoDuplicates);
 
 		FRandomStream Rand(0);
 		Shuffle(Rand, EntitiesSubSet);
 
-		FArchetypeChunkCollection CollectionFromRandom(FloatsArchetype, EntitiesSubSet, FArchetypeChunkCollection::NoDuplicates);
+		FMassArchetypeSubChunks CollectionFromRandom(FloatsArchetype, EntitiesSubSet, FMassArchetypeSubChunks::NoDuplicates);
 
 		AITEST_TRUE("The resulting chunk collection should be the same regardless of the order of input entities", CollectionFromOrdered.IsSame(CollectionFromRandom));
 		
 		// just to roughly make sure the result is what we expect
-		TArrayView<const FArchetypeChunkCollection::FChunkInfo> Chunks = CollectionFromOrdered.GetChunks();
+		FMassArchetypeSubChunks::FConstSubChunkArrayView Chunks = CollectionFromOrdered.GetChunks();
 		AITEST_EQUAL("The result should contain two chunks", Chunks.Num(), 2);
 
 		return true;
@@ -128,8 +128,8 @@ struct FChunkCollection_CreateCrossChunk : FEntityTestBase
 			EntitiesSubCollection.Add(Entities[EntitiesPerChunk - i]);
 		}
 		
-		FArchetypeChunkCollection ChunkCollection(FloatsArchetype, EntitiesSubCollection, FArchetypeChunkCollection::NoDuplicates);
-		TArrayView<const FArchetypeChunkCollection::FChunkInfo> Chunks = ChunkCollection.GetChunks();
+		FMassArchetypeSubChunks ChunkCollection(FloatsArchetype, EntitiesSubCollection, FMassArchetypeSubChunks::NoDuplicates);
+		FMassArchetypeSubChunks::FConstSubChunkArrayView Chunks = ChunkCollection.GetChunks();
 		AITEST_EQUAL("The given continuous range should get split in two", Chunks.Num(), 2);
 		AITEST_EQUAL("The part in first archetype\'s chunk should contain 9 elements", Chunks[0].Length, 9);
 		AITEST_EQUAL("The part in second archetype\'s chunk should contain 10 elements", Chunks[1].Length, 10);
@@ -151,8 +151,8 @@ struct FChunkCollection_CreateWithDuplicates : FChunkCollectionTestBase
 		EntitiesWithDuplicates.Add(Entities[2]);
 		EntitiesWithDuplicates.Add(Entities[2]);
 
-		FArchetypeChunkCollection ChunkCollection(FloatsArchetype, EntitiesWithDuplicates, FArchetypeChunkCollection::FoldDuplicates);
-		TArrayView<const FArchetypeChunkCollection::FChunkInfo> Chunks = ChunkCollection.GetChunks();
+		FMassArchetypeSubChunks ChunkCollection(FloatsArchetype, EntitiesWithDuplicates, FMassArchetypeSubChunks::FoldDuplicates);
+		FMassArchetypeSubChunks::FConstSubChunkArrayView Chunks = ChunkCollection.GetChunks();
 		AITEST_EQUAL("The result should have a single subchunk", Chunks.Num(), 1);
 		AITEST_EQUAL("The resulting subchunk should be of length 3", Chunks[0].Length, 3);
 
diff --git a/Engine/Plugins/Runtime/MassEntity/Source/MassEntityTestSuite/Private/MassEntityTest.cpp b/Engine/Plugins/Runtime/MassEntity/Source/MassEntityTestSuite/Private/MassEntityTest.cpp
--- a/Engine/Plugins/Runtime/MassEntity/Source/MassEntityTestSuite/Private/MassEntityTest.cpp
+++ b/Engine/Plugins/Runtime/MassEntity/Source/MassEntityTestSuite/Private/MassEntityTest.cpp
@@ -51,8 +51,8 @@ struct FEntityTest_ArchetypeEquivalence : FEntityTestBase
 	{
 		TArray<const UScriptStruct*> FragmentsA = { FTestFragment_Float::StaticStruct(), FTestFragment_Int::StaticStruct() };
 		TArray<const UScriptStruct*> FragmentsB = { FTestFragment_Int::StaticStruct(), FTestFragment_Float::StaticStruct() };
-		const FArchetypeHandle ArchetypeA = EntitySubsystem->CreateArchetype(FragmentsA);
-		const FArchetypeHandle ArchetypeB = EntitySubsystem->CreateArchetype(FragmentsB);
+		const FMassArchetypeHandle ArchetypeA = EntitySubsystem->CreateArchetype(FragmentsA);
+		const FMassArchetypeHandle ArchetypeB = EntitySubsystem->CreateArchetype(FragmentsB);
 
 		AITEST_EQUAL("Archetype creation is expected to be independent of fragments ordering", ArchetypeA, ArchetypeB);
 
@@ -69,9 +69,9 @@ struct FEntityTest_MultipleEntitiesCreation : FEntityTestBase
 		CA_ASSUME(EntitySubsystem);
 		int32 Counts[] = { 10, 100, 1000 };
 		int32 TotalCreatedCount = 0;
-		FArchetypeHandle Archetypes[] = { FloatsArchetype, IntsArchetype, FloatsIntsArchetype };
+		FMassArchetypeHandle Archetypes[] = { FloatsArchetype, IntsArchetype, FloatsIntsArchetype };
 
-		for (int ArchetypeIndex = 0; ArchetypeIndex < (sizeof(Archetypes) / sizeof(FArchetypeHandle)); ++ArchetypeIndex)
+		for (int ArchetypeIndex = 0; ArchetypeIndex < (sizeof(Archetypes) / sizeof(FMassArchetypeHandle)); ++ArchetypeIndex)
 		{
 			for (int i = 0; i < Counts[ArchetypeIndex]; ++i)
 			{
@@ -318,7 +318,7 @@ struct FEntityTest_ReleaseEntity : FEntityTestBase
 		AITEST_TRUE("The reserved entity should be a valid entity", EntitySubsystem->IsEntityValid(ReservedEntity));
 		AITEST_FALSE("The reserved entity should not be a valid entity", EntitySubsystem->IsEntityBuilt(ReservedEntity));
 		AITEST_EQUAL("There should only be one entity across the whole system", EntitySubsystem->DebugGetEntityCount(), 1);
-		AITEST_EQUAL("The entity should not get associated to any archetype", EntitySubsystem->GetArchetypeForEntity(ReservedEntity), FArchetypeHandle());
+		AITEST_EQUAL("The entity should not get associated to any archetype", EntitySubsystem->GetArchetypeForEntity(ReservedEntity), FMassArchetypeHandle());
 		EntitySubsystem->ReleaseReservedEntity(ReservedEntity);
 		AITEST_EQUAL("There should not be any entity across the whole system", EntitySubsystem->DebugGetEntityCount(), 0);
 		return true;
@@ -343,7 +343,7 @@ struct FEntityTest_ReserveAPreviouslyBuiltEntity : FEntityTestBase
 		AITEST_TRUE("The reserved entity should be a valid entity", EntitySubsystem->IsEntityValid(ReservedEntity));
 		AITEST_FALSE("The reserved entity should not be a valid entity", EntitySubsystem->IsEntityBuilt(ReservedEntity));
 		AITEST_EQUAL("There should only be one entity across the whole system", EntitySubsystem->DebugGetEntityCount(), 1);
-		AITEST_EQUAL("The entity should not get associated to any archetype", EntitySubsystem->GetArchetypeForEntity(ReservedEntity), FArchetypeHandle());
+		AITEST_EQUAL("The entity should not get associated to any archetype", EntitySubsystem->GetArchetypeForEntity(ReservedEntity), FMassArchetypeHandle());
 		EntitySubsystem->BuildEntity(ReservedEntity, FloatsArchetype);
 		AITEST_TRUE("The reserved entity should not be a valid entity", EntitySubsystem->IsEntityBuilt(ReservedEntity));
 		AITEST_EQUAL("The entity should get associated to the right archetype", EntitySubsystem->GetArchetypeForEntity(ReservedEntity), FloatsArchetype);
diff --git a/Engine/Plugins/Runtime/MassEntity/Source/MassEntityTestSuite/Private/MassExecutionTest.cpp b/Engine/Plugins/Runtime/MassEntity/Source/MassEntityTestSuite/Private/MassExecutionTest.cpp
--- a/Engine/Plugins/Runtime/MassEntity/Source/MassEntityTestSuite/Private/MassExecutionTest.cpp
+++ b/Engine/Plugins/Runtime/MassEntity/Source/MassEntityTestSuite/Private/MassExecutionTest.cpp
@@ -161,7 +161,7 @@ struct FExecution_Sparse : FEntityTestBase
 			Pipeline.SetProcessors(MoveTemp(Processors));
 		}
 
-		FArchetypeChunkCollection ChunkCollection(FloatsArchetype);
+		FMassArchetypeSubChunks ChunkCollection(FloatsArchetype);
 		// nothing should break. The actual result of processing is getting tested in MassProcessorTests.cpp
 		
 		UE::Mass::Executor::RunSparse(Pipeline, ProcessingContext, ChunkCollection);
diff --git a/Engine/Plugins/Runtime/MassEntity/Source/MassEntityTestSuite/Private/MassQueryTest.cpp b/Engine/Plugins/Runtime/MassEntity/Source/MassEntityTestSuite/Private/MassQueryTest.cpp
--- a/Engine/Plugins/Runtime/MassEntity/Source/MassEntityTestSuite/Private/MassQueryTest.cpp
+++ b/Engine/Plugins/Runtime/MassEntity/Source/MassEntityTestSuite/Private/MassQueryTest.cpp
@@ -215,7 +215,7 @@ struct FQueryTest_ExecuteSparse : FEntityTestBase
 
 		FMassExecutionContext ExecContext;
 		FMassEntityQuery().AddRequirement<FTestFragment_Float>(EMassFragmentAccess::ReadWrite)
-			.ForEachEntityChunk(FArchetypeChunkCollection(FloatsArchetype, EntitiesToProcess, FArchetypeChunkCollection::NoDuplicates)
+			.ForEachEntityChunk(FMassArchetypeSubChunks(FloatsArchetype, EntitiesToProcess, FMassArchetypeSubChunks::NoDuplicates)
 								, *EntitySubsystem, ExecContext, [&TotalProcessed](FMassExecutionContext& Context)
 			{
 				TotalProcessed += Context.GetNumEntities();
@@ -254,7 +254,7 @@ struct FQueryTest_TagPresent : FEntityTestBase
 		CA_ASSUME(EntitySubsystem);
 
 		TArray<const UScriptStruct*> Fragments = {FTestFragment_Float::StaticStruct(), FTestFragment_Tag::StaticStruct()};
-		const FArchetypeHandle FloatsTagArchetype = EntitySubsystem->CreateArchetype(Fragments);
+		const FMassArchetypeHandle FloatsTagArchetype = EntitySubsystem->CreateArchetype(Fragments);
 
 		FMassEntityQuery Query;
 		Query.AddRequirement<FTestFragment_Float>(EMassFragmentAccess::ReadWrite);
@@ -277,7 +277,7 @@ struct FQueryTest_TagAbsent : FEntityTestBase
 		CA_ASSUME(EntitySubsystem);
 
 		TArray<const UScriptStruct*> Fragments = { FTestFragment_Float::StaticStruct(), FTestFragment_Tag::StaticStruct() };
-		const FArchetypeHandle FloatsTagArchetype = EntitySubsystem->CreateArchetype(Fragments);
+		const FMassArchetypeHandle FloatsTagArchetype = EntitySubsystem->CreateArchetype(Fragments);
 
 		FMassEntityQuery Query;
 		Query.AddRequirement<FTestFragment_Float>(EMassFragmentAccess::ReadWrite);
@@ -476,11 +476,11 @@ struct FQueryTest_AnyFragment : FEntityTestBase
 		CA_ASSUME(EntitySubsystem);
 
 		// From FEntityTestBase:
-		// FArchetypeHandle FloatsArchetype;
-		// FArchetypeHandle IntsArchetype;
-		// FArchetypeHandle FloatsIntsArchetype;
-		const FArchetypeHandle BoolArchetype = EntitySubsystem->CreateArchetype({ FTestFragment_Bool::StaticStruct() });
-		const FArchetypeHandle BoolFloatArchetype = EntitySubsystem->CreateArchetype({ FTestFragment_Bool::StaticStruct(), FTestFragment_Float::StaticStruct() });
+		// FMassArchetypeHandle FloatsArchetype;
+		// FMassArchetypeHandle IntsArchetype;
+		// FMassArchetypeHandle FloatsIntsArchetype;
+		const FMassArchetypeHandle BoolArchetype = EntitySubsystem->CreateArchetype({ FTestFragment_Bool::StaticStruct() });
+		const FMassArchetypeHandle BoolFloatArchetype = EntitySubsystem->CreateArchetype({ FTestFragment_Bool::StaticStruct(), FTestFragment_Float::StaticStruct() });
 
 		FMassEntityQuery Query;
 		Query.AddRequirement<FTestFragment_Int>(EMassFragmentAccess::ReadWrite, EMassFragmentPresence::Any);
@@ -518,11 +518,11 @@ struct FQueryTest_AnyTag : FEntityTestBase
 	{
 		CA_ASSUME(EntitySubsystem);
 
-		const FArchetypeHandle ABArchetype = EntitySubsystem->CreateArchetype({ FTestFragment_Int::StaticStruct(), FTestTag_A::StaticStruct(), FTestTag_B::StaticStruct() });
-		const FArchetypeHandle ACArchetype = EntitySubsystem->CreateArchetype({ FTestFragment_Int::StaticStruct(), FTestTag_A::StaticStruct(), FTestTag_C::StaticStruct() });
-		const FArchetypeHandle BCArchetype = EntitySubsystem->CreateArchetype({ FTestFragment_Int::StaticStruct(), FTestTag_B::StaticStruct(), FTestTag_C::StaticStruct() });
-		const FArchetypeHandle BDArchetype = EntitySubsystem->CreateArchetype({ FTestFragment_Int::StaticStruct(), FTestTag_B::StaticStruct(), FTestTag_D::StaticStruct() });
-		const FArchetypeHandle FloatACArchetype = EntitySubsystem->CreateArchetype({ FTestFragment_Float::StaticStruct(), FTestTag_A::StaticStruct(), FTestTag_C::StaticStruct() });
+		const FMassArchetypeHandle ABArchetype = EntitySubsystem->CreateArchetype({ FTestFragment_Int::StaticStruct(), FTestTag_A::StaticStruct(), FTestTag_B::StaticStruct() });
+		const FMassArchetypeHandle ACArchetype = EntitySubsystem->CreateArchetype({ FTestFragment_Int::StaticStruct(), FTestTag_A::StaticStruct(), FTestTag_C::StaticStruct() });
+		const FMassArchetypeHandle BCArchetype = EntitySubsystem->CreateArchetype({ FTestFragment_Int::StaticStruct(), FTestTag_B::StaticStruct(), FTestTag_C::StaticStruct() });
+		const FMassArchetypeHandle BDArchetype = EntitySubsystem->CreateArchetype({ FTestFragment_Int::StaticStruct(), FTestTag_B::StaticStruct(), FTestTag_D::StaticStruct() });
+		const FMassArchetypeHandle FloatACArchetype = EntitySubsystem->CreateArchetype({ FTestFragment_Float::StaticStruct(), FTestTag_A::StaticStruct(), FTestTag_C::StaticStruct() });
 
 		FMassEntityQuery Query;
 		// at least one fragment requirement needs to be present for the query to be valid
diff --git a/Engine/Plugins/Runtime/MassEntity/Source/MassEntityTestSuite/Public/LWComponentTestFarmPlot.h b/Engine/Plugins/Runtime/MassEntity/Source/MassEntityTestSuite/Public/LWComponentTestFarmPlot.h
--- a/Engine/Plugins/Runtime/MassEntity/Source/MassEntityTestSuite/Public/LWComponentTestFarmPlot.h
+++ b/Engine/Plugins/Runtime/MassEntity/Source/MassEntityTestSuite/Public/LWComponentTestFarmPlot.h
@@ -337,7 +337,7 @@ private:
 	TArray<int32> FreeHarvestIconIndicies;
 
 private:
-	void AddItemToGrid(UMassEntitySubsystem* EntitySystem, uint16 X, uint16 Y, FArchetypeHandle Archetype, uint16 VisualIndex);
+	void AddItemToGrid(UMassEntitySubsystem* EntitySystem, uint16 X, uint16 Y, FMassArchetypeHandle Archetype, uint16 VisualIndex);
 
 public:
 	ALWFragmentTestFarmPlot();
diff --git a/Engine/Plugins/Runtime/MassEntity/Source/MassEntityTestSuite/Public/MassEntityTestTypes.h b/Engine/Plugins/Runtime/MassEntity/Source/MassEntityTestSuite/Public/MassEntityTestTypes.h
--- a/Engine/Plugins/Runtime/MassEntity/Source/MassEntityTestSuite/Public/MassEntityTestTypes.h
+++ b/Engine/Plugins/Runtime/MassEntity/Source/MassEntityTestSuite/Public/MassEntityTestTypes.h
@@ -166,10 +166,10 @@ const int TestIntValue = 123456;
 
 struct FEntityTestBase : FExecutionTestBase
 {
-	FArchetypeHandle EmptyArchetype;
-	FArchetypeHandle FloatsArchetype;
-	FArchetypeHandle IntsArchetype;
-	FArchetypeHandle FloatsIntsArchetype;
+	FMassArchetypeHandle EmptyArchetype;
+	FMassArchetypeHandle FloatsArchetype;
+	FMassArchetypeHandle IntsArchetype;
+	FMassArchetypeHandle FloatsIntsArchetype;
 
 	FInstancedStruct InstanceInt;
 
diff --git a/Engine/Plugins/Runtime/MassGameplay/Source/MassActors/Private/MassAgentSubsystem.cpp b/Engine/Plugins/Runtime/MassGameplay/Source/MassActors/Private/MassAgentSubsystem.cpp
--- a/Engine/Plugins/Runtime/MassGameplay/Source/MassActors/Private/MassAgentSubsystem.cpp
+++ b/Engine/Plugins/Runtime/MassGameplay/Source/MassActors/Private/MassAgentSubsystem.cpp
@@ -137,7 +137,7 @@ void UMassAgentSubsystem::UpdateAgentComponent(const UMassAgentComponent& AgentC
 	}
 
 	const FMassEntityHandle Entity = AgentComp.GetEntityHandle();
-	const FArchetypeHandle CurrentArchetypeHandle = EntitySystem->GetArchetypeForEntity(Entity);
+	const FMassArchetypeHandle CurrentArchetypeHandle = EntitySystem->GetArchetypeForEntity(Entity);
 	if (CurrentArchetypeHandle == EntityTemplate->GetArchetype())
 	{
 		UE_VLOG(this, LogMassActor, Log, TEXT("%s called for %s but no archetype changes have been found")
@@ -359,7 +359,7 @@ void UMassAgentSubsystem::HandlePendingInitialization()
 			
 			if (EntityTemplate->GetObjectFragmentInitializers().Num())
 			{
-				const FArchetypeHandle ArchetypeHandle = EntitySystem->GetArchetypeForEntity(PuppetEntity);
+				const FMassArchetypeHandle ArchetypeHandle = EntitySystem->GetArchetypeForEntity(PuppetEntity);
 				FMassEntityView EntityView(ArchetypeHandle, PuppetEntity);
 				FMassAgentSubsystemHelper::InitializeAgentComponentFragments(*AgentComp, EntityView, EMassTranslationDirection::MassToActor, EntityTemplate->GetObjectFragmentInitializers());
 			}
diff --git a/Engine/Plugins/Runtime/MassGameplay/Source/MassReplication/Public/MassClientBubbleHandler.h b/Engine/Plugins/Runtime/MassGameplay/Source/MassReplication/Public/MassClientBubbleHandler.h
--- a/Engine/Plugins/Runtime/MassGameplay/Source/MassReplication/Public/MassClientBubbleHandler.h
+++ b/Engine/Plugins/Runtime/MassGameplay/Source/MassReplication/Public/MassClientBubbleHandler.h
@@ -538,7 +538,7 @@ void TClientBubbleHandlerBase<AgentArrayItem>::PostReplicatedAddEntitiesHelper(c
 
 		const FMassEntityTemplate* MassEntityTemplate = SpawnerSubsystem->GetMassEntityTemplate(TemplateID);
 		check(MassEntityTemplate);
-		const FArchetypeHandle& ArchetypeHandle = MassEntityTemplate->GetArchetype();
+		const FMassArchetypeHandle& ArchetypeHandle = MassEntityTemplate->GetArchetype();
 
 
 		FMassExecutionContext ExecContext;
@@ -550,7 +550,7 @@ void TClientBubbleHandlerBase<AgentArrayItem>::PostReplicatedAddEntitiesHelper(c
 
 		int32 AgentsSpawnIdx = 0;
 
-		Query.ForEachEntityChunk(FArchetypeChunkCollection(ArchetypeHandle, Entities, FArchetypeChunkCollection::NoDuplicates)
+		Query.ForEachEntityChunk(FMassArchetypeSubChunks(ArchetypeHandle, Entities, FMassArchetypeSubChunks::NoDuplicates)
 								, *EntitySystem, ExecContext, [&AgentsSpawn, &AgentsSpawnIdx, this, ReplicationManager, &ExecContext, &CacheFragmentViewsForSpawnQuery, &SetSpawnedEntityData, &EntitySystem](FMassExecutionContext& Context)
 			{
 				CacheFragmentViewsForSpawnQuery(ExecContext);
diff --git a/Engine/Plugins/Runtime/MassGameplay/Source/MassSignals/Private/MassSignalProcessorBase.cpp b/Engine/Plugins/Runtime/MassGameplay/Source/MassSignals/Private/MassSignalProcessorBase.cpp
--- a/Engine/Plugins/Runtime/MassGameplay/Source/MassSignals/Private/MassSignalProcessorBase.cpp
+++ b/Engine/Plugins/Runtime/MassGameplay/Source/MassSignals/Private/MassSignalProcessorBase.cpp
@@ -51,7 +51,7 @@ void UMassSignalProcessorBase::Execute(UMassEntitySubsystem& EntitySubsystem, FM
 	if (EntityQuery.GetArchetypes().Num() > 0)
 	{
 		// EntitySet stores unique array of entities per specified archetype.
-		// FArchetypeChunkCollection expects an array of entities, a set is used to detect unique ones.
+		// FMassArchetypeSubChunks expects an array of entities, a set is used to detect unique ones.
 		struct FEntitySet
 		{
 			void Reset()
@@ -59,12 +59,12 @@ void UMassSignalProcessorBase::Execute(UMassEntitySubsystem& EntitySubsystem, FM
 				Entities.Reset();
 			}
 
-			FArchetypeHandle Archetype;
+			FMassArchetypeHandle Archetype;
 			TArray<FMassEntityHandle> Entities;
 		};
 		TArray<FEntitySet> EntitySets;
 
-		for (const FArchetypeHandle& Archetype : EntityQuery.GetArchetypes())
+		for (const FMassArchetypeHandle& Archetype : EntityQuery.GetArchetypes())
 		{
 			FEntitySet& Set = EntitySets.AddDefaulted_GetRef();
 			Set.Archetype = Archetype;
@@ -97,11 +97,11 @@ void UMassSignalProcessorBase::Execute(UMassEntitySubsystem& EntitySubsystem, FM
 				for (const FMassEntityHandle Entity : Entities)
 				{
 					// Add to set of supported archetypes. Dont process if we don't care about the type.
-					FArchetypeHandle Archetype = EntitySubsystem.GetArchetypeForEntity(Entity);
+					FMassArchetypeHandle Archetype = EntitySubsystem.GetArchetypeForEntity(Entity);
 					FEntitySet* Set = PrevSet->Archetype == Archetype ? PrevSet : EntitySets.FindByPredicate([&Archetype](const FEntitySet& Set) { return Archetype == Set.Archetype; });
 					if (Set != nullptr)
 					{
-						// We don't care about duplicates here, the FArchetypeChunkCollection creation below will handle it
+						// We don't care about duplicates here, the FMassArchetypeSubChunks creation below will handle it
 						Set->Entities.Add(Entity);
 						SignalNameLookup.AddSignalToEntity(Entity, SignalFlag);
 						PrevSet = Set;
@@ -117,7 +117,7 @@ void UMassSignalProcessorBase::Execute(UMassEntitySubsystem& EntitySubsystem, FM
 			{
 				if (Set.Entities.Num() > 0)
 				{
-					Context.SetChunkCollection(FArchetypeChunkCollection(Set.Archetype, Set.Entities, FArchetypeChunkCollection::FoldDuplicates));
+					Context.SetChunkCollection(FMassArchetypeSubChunks(Set.Archetype, Set.Entities, FMassArchetypeSubChunks::FoldDuplicates));
 					SignalEntities(EntitySubsystem, Context, SignalNameLookup);
 					Context.ClearChunkCollection();
 				}
diff --git a/Engine/Plugins/Runtime/MassGameplay/Source/MassSpawner/Private/MassEntityTemplate.cpp b/Engine/Plugins/Runtime/MassGameplay/Source/MassSpawner/Private/MassEntityTemplate.cpp
--- a/Engine/Plugins/Runtime/MassGameplay/Source/MassSpawner/Private/MassEntityTemplate.cpp
+++ b/Engine/Plugins/Runtime/MassGameplay/Source/MassSpawner/Private/MassEntityTemplate.cpp
@@ -16,7 +16,7 @@ FString FMassEntityTemplateID::ToString() const
 //----------------------------------------------------------------------//
 //  FMassEntityTemplate
 //----------------------------------------------------------------------//
-void FMassEntityTemplate::SetArchetype(const FArchetypeHandle& InArchetype)
+void FMassEntityTemplate::SetArchetype(const FMassArchetypeHandle& InArchetype)
 {
 	check(InArchetype.IsValid());
 	Archetype = InArchetype;
diff --git a/Engine/Plugins/Runtime/MassGameplay/Source/MassSpawner/Private/MassEntityTemplateRegistry.cpp b/Engine/Plugins/Runtime/MassGameplay/Source/MassSpawner/Private/MassEntityTemplateRegistry.cpp
--- a/Engine/Plugins/Runtime/MassGameplay/Source/MassSpawner/Private/MassEntityTemplateRegistry.cpp
+++ b/Engine/Plugins/Runtime/MassGameplay/Source/MassSpawner/Private/MassEntityTemplateRegistry.cpp
@@ -151,7 +151,7 @@ void UMassEntityTemplateRegistry::InitializeEntityTemplate(FMassEntityTemplate&
 	// Sort anything there is to sort for later comparison purposes
 	OutTemplate.Sort();
 
-	const FArchetypeHandle ArchetypeHandle = EntitySys->CreateArchetype(OutTemplate.GetCompositionDescriptor(), OutTemplate.GetSharedFragmentValues());
+	const FMassArchetypeHandle ArchetypeHandle = EntitySys->CreateArchetype(OutTemplate.GetCompositionDescriptor(), OutTemplate.GetSharedFragmentValues());
 	OutTemplate.SetArchetype(ArchetypeHandle);
 }
 
diff --git a/Engine/Plugins/Runtime/MassGameplay/Source/MassSpawner/Private/MassSpawnerSubsystem.cpp b/Engine/Plugins/Runtime/MassGameplay/Source/MassSpawner/Private/MassSpawnerSubsystem.cpp
--- a/Engine/Plugins/Runtime/MassGameplay/Source/MassSpawner/Private/MassSpawnerSubsystem.cpp
+++ b/Engine/Plugins/Runtime/MassGameplay/Source/MassSpawner/Private/MassSpawnerSubsystem.cpp
@@ -116,10 +116,10 @@ void UMassSpawnerSubsystem::DestroyEntities(const FMassEntityTemplateID Template
 	check(World);
 
 
-	TArray<FArchetypeChunkCollection> ChunkCollections;
-	UE::Mass::Utils::CreateSparseChunks(*EntitySystem, Entities, FArchetypeChunkCollection::NoDuplicates, ChunkCollections);
+	TArray<FMassArchetypeSubChunks > ChunkCollections;
+	UE::Mass::Utils::CreateSparseChunks(*EntitySystem, Entities, FMassArchetypeSubChunks::NoDuplicates, ChunkCollections);
 
-	for (const FArchetypeChunkCollection& Chunks : ChunkCollections)
+	for (const FMassArchetypeSubChunks& Chunks : ChunkCollections)
 	{
 		EntitySystem->BatchDestroyEntityChunks(Chunks);
 	}
diff --git a/Engine/Plugins/Runtime/MassGameplay/Source/MassSpawner/Public/MassEntityTemplate.h b/Engine/Plugins/Runtime/MassGameplay/Source/MassSpawner/Public/MassEntityTemplate.h
--- a/Engine/Plugins/Runtime/MassGameplay/Source/MassSpawner/Public/MassEntityTemplate.h
+++ b/Engine/Plugins/Runtime/MassGameplay/Source/MassSpawner/Public/MassEntityTemplate.h
@@ -73,8 +73,8 @@ struct MASSSPAWNER_API FMassEntityTemplate
 
 	FMassEntityTemplate() = default;
 	/** InArchetype is expected to be valid. The function will crash-check it. */
-	void SetArchetype(const FArchetypeHandle& InArchetype);
-	const FArchetypeHandle& GetArchetype() const { return Archetype; }
+	void SetArchetype(const FMassArchetypeHandle& InArchetype);
+	const FMassArchetypeHandle& GetArchetype() const { return Archetype; }
 
 	TConstArrayView<FObjectFragmentInitializerFunction> GetObjectFragmentInitializers() const { return ObjectInitializers; }
 	TArray<FObjectFragmentInitializerFunction>& GetMutableObjectFragmentInitializers() { return ObjectInitializers; }
@@ -235,7 +235,7 @@ struct MASSSPAWNER_API FMassEntityTemplate
 	}
 
 private:
-	FArchetypeHandle Archetype;
+	FMassArchetypeHandle Archetype;
 
 	FMassArchetypeCompositionDescriptor Composition;
 	FMassArchetypeSharedFragmentValues SharedFragmentValues;
