Lowered the severity of "Attempted to execute FMassEntityQuery with an incompatible Archetype" log message from FMassEntityQuery

diff --git a/Engine/Plugins/Runtime/MassEntity/Source/MassEntity/Private/MassEntityQuery.cpp b/Engine/Plugins/Runtime/MassEntity/Source/MassEntity/Private/MassEntityQuery.cpp
--- a/Engine/Plugins/Runtime/MassEntity/Source/MassEntity/Private/MassEntityQuery.cpp
+++ b/Engine/Plugins/Runtime/MassEntity/Source/MassEntity/Private/MassEntityQuery.cpp
@@ -142,8 +142,7 @@ void FMassEntityQuery::ForEachEntityChunk(UMassEntitySubsystem& EntitySubsystem,
 		// verify the archetype matches requirements
 		if (DoesArchetypeMatchRequirements(ExecutionContext.GetChunkCollection().GetArchetype()) == false)
 		{
-			// mz@todo add a unit test for this message
-			UE_VLOG_UELOG(&EntitySubsystem, LogMass, Error, TEXT("Attempted to execute FMassEntityQuery with an incompatible Archetype: %s")
+			UE_VLOG_UELOG(&EntitySubsystem, LogMass, Log, TEXT("Attempted to execute FMassEntityQuery with an incompatible Archetype: %s")
 				, *DebugGetArchetypeCompatibilityDescription(ExecutionContext.GetChunkCollection().GetArchetype()));
 			return;
 		}
@@ -204,8 +203,7 @@ void FMassEntityQuery::ParallelForEachEntityChunk(UMassEntitySubsystem& EntitySu
 		FMassArchetypeHandle ArchetypeHandle = ExecutionContext.GetChunkCollection().GetArchetype();
 		if (DoesArchetypeMatchRequirements(ArchetypeHandle) == false)
 		{
-			// mz@todo add a unit test for this message
-			UE_VLOG_UELOG(&EntitySubsystem, LogMass, Error, TEXT("Attempted to execute FMassEntityQuery with an incompatible Archetype: %s")
+			UE_VLOG_UELOG(&EntitySubsystem, LogMass, Log, TEXT("Attempted to execute FMassEntityQuery with an incompatible Archetype: %s")
 				, *DebugGetArchetypeCompatibilityDescription(ExecutionContext.GetChunkCollection().GetArchetype()));
 			return;
 		}
