Added a couple of missing DirtyCachedData calls in FMassEntityQuery's "Add Requirement" methods

diff --git a/Engine/Plugins/Runtime/MassEntity/Source/MassEntity/Public/MassEntityQuery.h b/Engine/Plugins/Runtime/MassEntity/Source/MassEntity/Public/MassEntityQuery.h
--- a/Engine/Plugins/Runtime/MassEntity/Source/MassEntity/Public/MassEntityQuery.h
+++ b/Engine/Plugins/Runtime/MassEntity/Source/MassEntity/Public/MassEntityQuery.h
@@ -197,6 +197,7 @@ public:
 			RequiredNoneTags.Add(TagType);
 			break;
 		}
+		DirtyCachedData();
 	}
 
 	template<typename T>
@@ -216,7 +217,7 @@ public:
 				RequiredNoneTags.Add<T>();
 				break;
 		}
-
+		DirtyCachedData();
 		return *this;
 	}
 
@@ -251,7 +252,7 @@ public:
 				RequiredNoneChunkFragments.Add<T>();
 				break;
 		}
-
+		DirtyCachedData();
 		return *this;
 	}
 
@@ -277,7 +278,7 @@ public:
 			RequiredNoneSharedFragments.Add<T>();
 			break;
 		}
-
+		DirtyCachedData();
 		return *this;
 	}
 
@@ -303,7 +304,7 @@ public:
 			RequiredNoneSharedFragments.Add<T>();
 			break;
 		}
-
+		DirtyCachedData();
 		return *this;
 	}
 
