Fixed MassEntity unit test for shipping editor build

diff --git a/Engine/Plugins/Runtime/MassEntity/Source/MassEntityTestSuite/Private/MassChunkCollectionTest.cpp b/Engine/Plugins/Runtime/MassEntity/Source/MassEntityTestSuite/Private/MassChunkCollectionTest.cpp
--- a/Engine/Plugins/Runtime/MassEntity/Source/MassEntityTestSuite/Private/MassChunkCollectionTest.cpp
+++ b/Engine/Plugins/Runtime/MassEntity/Source/MassEntityTestSuite/Private/MassChunkCollectionTest.cpp
@@ -113,6 +113,7 @@ struct FChunkCollection_CreateCrossChunk : FEntityTestBase
 {
 	virtual bool InstantTest() override
 	{
+#if WITH_MASSENTITY_DEBUG
 		TArray<FMassEntityHandle> Entities;
 		const int32 EntitiesPerChunk = EntitySubsystem->DebugGetArchetypeEntitiesCountPerChunk(FloatsArchetype);
 
@@ -133,7 +134,7 @@ struct FChunkCollection_CreateCrossChunk : FEntityTestBase
 		AITEST_EQUAL("The given continuous range should get split in two", Chunks.Num(), 2);
 		AITEST_EQUAL("The part in first archetype\'s chunk should contain 9 elements", Chunks[0].Length, 9);
 		AITEST_EQUAL("The part in second archetype\'s chunk should contain 10 elements", Chunks[1].Length, 10);
-
+#endif // WITH_MASSENTITY_DEBUG
 		return true;
 	}
 };
