Abstract-sort, a generic sorting wrapper that can sort an arbitrary collection of things as long as it's provided with a number of elements, a function for comparing items, and a function to swap elements. The sorting is performed mostly in-place - there's a helper array of indices allocated first.

diff --git a/Engine/Plugins/Runtime/MassEntity/Source/MassEntity/Public/MassEntityUtils.h b/Engine/Plugins/Runtime/MassEntity/Source/MassEntity/Public/MassEntityUtils.h
--- a/Engine/Plugins/Runtime/MassEntity/Source/MassEntity/Public/MassEntityUtils.h
+++ b/Engine/Plugins/Runtime/MassEntity/Source/MassEntity/Public/MassEntityUtils.h
@@ -22,4 +22,59 @@ MASSENTITY_API extern EProcessorExecutionFlags GetProcessorExecutionFlagsForWold
 MASSENTITY_API extern void CreateSparseChunks(const UMassEntitySubsystem& EntitySystem, const TConstArrayView<FMassEntityHandle> Entities
 	, const FMassArchetypeSubChunks::EDuplicatesHandling DuplicatesHandling, TArray<FMassArchetypeSubChunks >& OutChunkCollections);
 
-} // namespace UE::Mass::Utils
\ No newline at end of file
+/**
+* AbstractSort is a sorting function that only needs to know how many items there are, how to compare items
+* at individual locations - where location is in [0, NumElements) - and how to swap two elements at given locations.
+* The main use case is to sort multiple arrays while keeping them in sync. For example:
+*
+* TArray<float> Lead = { 3.1, 0.2, 2.6, 1.0 };
+* TArray<UObject*> Payload = { A, B, C, D };
+*
+* AbstractSort(Lead.Num()											// NumElements
+* 	, [&Lead](const int32 LHS, const int32 RHS)					// Predicate
+*		{
+*			return Lead[LHS] < Lead[RHS];
+*		}
+* 	, [&Lead, &Payload](const int32 A, const int32 B)			// SwapFunctor
+*	 	{
+*			Swap(Lead[A], Lead[B]);
+* 			Swap(Payload[A], Payload[B]);
+*		}
+* );
+*/
+template<typename TPred, typename TSwap>
+inline void AbstractSort(const int32 NumElements, TPred&& Predicate, TSwap&& SwapFunctor)
+{
+	if (NumElements == 0)
+	{
+		return;
+	}
+
+	TArray<int32> Indices;
+	Indices.AddUninitialized(NumElements);
+	int i = 0;
+	do
+	{
+		Indices[i] = i;
+	} while (++i < NumElements);
+
+	Indices.Sort(Predicate);
+
+	for (i = 0; i < NumElements; ++i)
+	{
+		int32 SwapFromIndex = Indices[i];
+		while (SwapFromIndex < i)
+		{
+			SwapFromIndex = Indices[SwapFromIndex];
+		}
+
+		if (SwapFromIndex != i)
+		{
+			SwapFunctor(i, SwapFromIndex);
+		}
+	}
+}
+
+
+} // namespace UE::Mass::Utils
+
