Extended MassObserverManager with a capability to register tag observers and broadcast to them when the observed action has been performed.
Updated MassObserverManager's API to better accommodate both types of observed mass items.
Note: while MassObserverManager can observe tags the MassEntitySubsystem nor the MassCommandBuffer don't utilize this functionality just yet. This will come next with batchified MassEntitySubsystem API changes

diff --git a/Engine/Plugins/Runtime/MassEntity/Source/MassEntity/Private/MassCommandBuffer.cpp b/Engine/Plugins/Runtime/MassEntity/Source/MassEntity/Private/MassCommandBuffer.cpp
--- a/Engine/Plugins/Runtime/MassEntity/Source/MassEntity/Private/MassCommandBuffer.cpp
+++ b/Engine/Plugins/Runtime/MassEntity/Source/MassEntity/Private/MassCommandBuffer.cpp
@@ -107,7 +107,7 @@ void FMassCommandBuffer::ReplayBufferAgainstSystem(UMassEntitySubsystem* EntityS
 				for (FMassArchetypeSubChunks& Collection : ChunkCollections)
 				{
 					check(It.Key);
-					ObserverManager.OnPreFragmentRemoved(*It.Key, Collection);
+					ObserverManager.OnPreItemRemoved(*It.Key, Collection);
 				}
 			}
 		}
@@ -165,7 +165,7 @@ void FMassCommandBuffer::ReplayBufferAgainstSystem(UMassEntitySubsystem* EntityS
 			for (FMassArchetypeSubChunks& Collection : ChunkCollections)
 			{
 				check(It.Key);
-				ObserverManager.OnPostFragmentAdded(*It.Key, Collection);
+				ObserverManager.OnPostItemAdded(*It.Key, Collection);
 			}
 		}
 	}
diff --git a/Engine/Plugins/Runtime/MassEntity/Source/MassEntity/Private/MassEntitySubsystem.cpp b/Engine/Plugins/Runtime/MassEntity/Source/MassEntity/Private/MassEntitySubsystem.cpp
--- a/Engine/Plugins/Runtime/MassEntity/Source/MassEntity/Private/MassEntitySubsystem.cpp
+++ b/Engine/Plugins/Runtime/MassEntity/Source/MassEntity/Private/MassEntitySubsystem.cpp
@@ -348,7 +348,7 @@ TSharedRef<UMassEntitySubsystem::FEntityCreationContext> UMassEntitySubsystem::B
 	// @todo this could probably be optimized since one would assume we're adding elements to OutEntities in order.
 	// Then again, if that's the case, the sorting will be almost instant
 	new (&CreationContext->ChunkCollection)FMassArchetypeSubChunks(Archetype, MakeArrayView(&OutEntities[OutEntities.Num() - Count], Count), FMassArchetypeSubChunks::NoDuplicates);
-	if (ObserverManager.HasOnAddedObserversForFragments(ArchetypePtr->GetCompositionDescriptor().Fragments))
+	if (ObserverManager.HasOnAddedObserversForBitSet(ArchetypePtr->GetCompositionDescriptor().Fragments))
 	{
 		CreationContext->OnSpawningFinished = [this](FEntityCreationContext& Context){
 			ObserverManager.OnPostEntitiesCreated(Context.ChunkCollection);
diff --git a/Engine/Plugins/Runtime/MassEntity/Source/MassEntity/Private/MassObserverManager.cpp b/Engine/Plugins/Runtime/MassEntity/Source/MassEntity/Private/MassObserverManager.cpp
--- a/Engine/Plugins/Runtime/MassEntity/Source/MassEntity/Private/MassObserverManager.cpp
+++ b/Engine/Plugins/Runtime/MassEntity/Source/MassEntity/Private/MassObserverManager.cpp
@@ -6,6 +6,34 @@
 #include "MassProcessingTypes.h"
 #include "MassObserverRegistry.h"
 
+namespace UE::Mass::ObserverManager::Private
+{
+// a helper function to reduce code duplication in FMassObserverManager::Initialize
+template<typename TBitSet>
+void SetUpObservers(UMassEntitySubsystem& EntitySubsystem, const TMap<const UScriptStruct*, FMassProcessorClassCollection>& RegisteredObserverTypes, TBitSet& ObservedBitSet, FMassItemObserversMap& Observers)
+{
+	ObservedBitSet.Reset();
+
+	for (auto It : RegisteredObserverTypes)
+	{
+		if (It.Value.ClassCollection.Num() == 0)
+		{
+			continue;
+		}
+
+		ObservedBitSet.Add(*It.Key);
+		FMassRuntimePipeline& Pipeline = (*Observers).FindOrAdd(It.Key);
+
+		for (const TSubclassOf<UMassProcessor>& ProcessorClass : It.Value.ClassCollection)
+		{
+			Pipeline.AppendProcessor(ProcessorClass, EntitySubsystem);
+		}
+		Pipeline.Initialize(EntitySubsystem);
+	}
+};
+
+} // UE::Mass::ObserverManager::Private
+
 //----------------------------------------------------------------------//
 // FMassObserverManager
 //----------------------------------------------------------------------//
@@ -26,44 +54,11 @@ void FMassObserverManager::Initialize()
 	// instantiate initializers
 	const UMassObserverRegistry& Registry = UMassObserverRegistry::Get();
 
-	ObservedAddFragments.Reset();
-	ObservedRemoveFragments.Reset();
-	OnFragmentAddedObservers.Reset();
-	OnFragmentRemovedObservers.Reset();
-
-	for (auto It : Registry.FragmentInitializersMap)
-	{
-		if (It.Value.ClassCollection.Num() == 0)
-		{
-			continue;
-		}
-
-		ObservedAddFragments.Add(*It.Key);
-		FMassRuntimePipeline& Pipeline = OnFragmentAddedObservers.FindOrAdd(It.Key);
-
-		for (const TSubclassOf<UMassProcessor>& ProcessorClass : It.Value.ClassCollection)
-		{
-			Pipeline.AppendProcessor(ProcessorClass, EntitySubsystem);
-		}
-		Pipeline.Initialize(EntitySubsystem);
-	}
-
-	for (auto It : Registry.FragmentDeinitializersMap)
-	{
-		if (It.Value.ClassCollection.Num() == 0)
-		{
-			continue;
-		}
-
-		ObservedRemoveFragments.Add(*It.Key);
-		FMassRuntimePipeline& Pipeline = OnFragmentRemovedObservers.FindOrAdd(It.Key);
-
-		for (const TSubclassOf<UMassProcessor>& ProcessorClass : It.Value.ClassCollection)
-		{
-			Pipeline.AppendProcessor(ProcessorClass, EntitySubsystem);
-		}
-		Pipeline.Initialize(EntitySubsystem);
-	}
+	using UE::Mass::ObserverManager::Private::SetUpObservers;
+	SetUpObservers(EntitySubsystem, Registry.FragmentInitializersMap, ObservedFragments[(uint8)FMassObservedOperation::Add], FragmentObservers[(uint8)FMassObservedOperation::Add]);
+	SetUpObservers(EntitySubsystem, Registry.FragmentDeinitializersMap, ObservedFragments[(uint8)FMassObservedOperation::Remove], FragmentObservers[(uint8)FMassObservedOperation::Remove]);
+	SetUpObservers(EntitySubsystem, Registry.TagAddedObserversMap, ObservedTags[(uint8)FMassObservedOperation::Add], TagObservers[(uint8)FMassObservedOperation::Add]);
+	SetUpObservers(EntitySubsystem, Registry.TagRemovedObserversMap, ObservedTags[(uint8)FMassObservedOperation::Remove], TagObservers[(uint8)FMassObservedOperation::Remove]);
 }
 
 bool FMassObserverManager::OnPostEntitiesCreated(const FMassArchetypeSubChunks& ChunkCollection)
@@ -89,11 +84,16 @@ bool FMassObserverManager::OnPostEntitiesCreated(FMassProcessingContext& Process
 
 	check(ProcessingContext.EntitySubsystem);
 	const FMassArchetypeCompositionDescriptor& ArchetypeComposition = ProcessingContext.EntitySubsystem->GetArchetypeComposition(ChunkCollection.GetArchetype());
-	const FMassFragmentBitSet Overlap = ObservedAddFragments.GetOverlap(ArchetypeComposition.Fragments);
+
+
+	const FMassFragmentBitSet Overlap = ObservedFragments[(uint8)FMassObservedOperation::Add].GetOverlap(ArchetypeComposition.Fragments);
 
 	if (Overlap.IsEmpty() == false)
 	{
-		HandleFragmentsImpl(ProcessingContext, ChunkCollection, Overlap, OnFragmentAddedObservers);
+		TArray<const UScriptStruct*> OverlapTypes;
+		Overlap.ExportTypes(OverlapTypes);
+
+		HandleFragmentsImpl(ProcessingContext, ChunkCollection, MakeArrayView(OverlapTypes), FragmentObservers[(uint8)FMassObservedOperation::Add]);
 		return true;
 	}
 
@@ -120,69 +120,135 @@ bool FMassObserverManager::OnPreEntitiesDestroyed(FMassProcessingContext& Proces
 
 	check(ProcessingContext.EntitySubsystem);
 	const FMassArchetypeCompositionDescriptor& ArchetypeComposition = ProcessingContext.EntitySubsystem->GetArchetypeComposition(ChunkCollection.GetArchetype());
-	const FMassFragmentBitSet Overlap = ObservedRemoveFragments.GetOverlap(ArchetypeComposition.Fragments);
+	
+	return OnCompositionChanged(ChunkCollection, ArchetypeComposition, FMassObservedOperation::Remove, &ProcessingContext);
+}
 
-	if (Overlap.IsEmpty() == false)
+bool FMassObserverManager::OnCompositionChanged(const FMassArchetypeSubChunks& ChunkCollection, const FMassArchetypeCompositionDescriptor& CompositionDelta, const FMassObservedOperation Operation, FMassProcessingContext* InProcessingContext)
+{
+	const FMassFragmentBitSet FragmentOverlap = ObservedFragments[(uint8)Operation].GetOverlap(CompositionDelta.Fragments);
+	const bool bHasFragmentsOverlap = !FragmentOverlap.IsEmpty();
+	const FMassTagBitSet TagOverlap = ObservedTags[(uint8)Operation].GetOverlap(CompositionDelta.Tags);
+	const bool bHasTagsOverlap = !TagOverlap.IsEmpty();
+
+	if (bHasFragmentsOverlap || bHasTagsOverlap)
 	{
-		HandleFragmentsImpl(ProcessingContext, ChunkCollection, Overlap, OnFragmentRemovedObservers);
-		return true;
+		FMassProcessingContext LocalContext(EntitySubsystem, /*DeltaSeconds=*/0.f);
+		FMassProcessingContext* ProcessingContext = InProcessingContext ? InProcessingContext : &LocalContext;
+		TArray<const UScriptStruct*> ObservedTypesOverlap;
+
+		if (bHasFragmentsOverlap)
+		{
+			FragmentOverlap.ExportTypes(ObservedTypesOverlap);
+
+			HandleFragmentsImpl(*ProcessingContext, ChunkCollection, ObservedTypesOverlap, FragmentObservers[(uint8)Operation]);
+		}
+
+		if (bHasTagsOverlap)
+		{
+			ObservedTypesOverlap.Reset();
+			TagOverlap.ExportTypes(ObservedTypesOverlap);
+
+			HandleFragmentsImpl(*ProcessingContext, ChunkCollection, ObservedTypesOverlap, TagObservers[(uint8)Operation]);
+		}
 	}
 
-	return false;
+	return bHasFragmentsOverlap || bHasTagsOverlap;
 }
 
-bool FMassObserverManager::OnPostCompositionAdded(const FMassEntityHandle Entity, const FMassArchetypeCompositionDescriptor& Composition)
+bool FMassObserverManager::OnCompositionChanged(const FMassEntityHandle Entity, const FMassArchetypeCompositionDescriptor& CompositionDelta, const FMassObservedOperation Operation)
 {
-	const FMassFragmentBitSet Overlap = ObservedAddFragments.GetOverlap(Composition.Fragments);
-	if (Overlap.IsEmpty() == false)
+	const FMassFragmentBitSet FragmentOverlap = ObservedFragments[(uint8)Operation].GetOverlap(CompositionDelta.Fragments);
+	const bool bHasFragmentsOverlap = !FragmentOverlap.IsEmpty();
+	const FMassTagBitSet TagOverlap = ObservedTags[(uint8)Operation].GetOverlap(CompositionDelta.Tags);
+	const bool bHasTagsOverlap = !TagOverlap.IsEmpty();
+
+	if (bHasFragmentsOverlap || bHasTagsOverlap)
 	{
+		TArray<const UScriptStruct*> ObservedTypesOverlap;
 		FMassProcessingContext ProcessingContext(EntitySubsystem, /*DeltaSeconds=*/0.f);
 		const FMassArchetypeHandle ArchetypeHandle = EntitySubsystem.GetArchetypeForEntity(Entity);
-		HandleFragmentsImpl(ProcessingContext, FMassArchetypeSubChunks(ArchetypeHandle, MakeArrayView(&Entity, 1)
-			, FMassArchetypeSubChunks::NoDuplicates), Overlap, OnFragmentAddedObservers);
-		return true;
+
+		if (bHasFragmentsOverlap)
+		{
+			FragmentOverlap.ExportTypes(ObservedTypesOverlap);
+
+			HandleFragmentsImpl(ProcessingContext, FMassArchetypeSubChunks(ArchetypeHandle, MakeArrayView(&Entity, 1)
+				, FMassArchetypeSubChunks::NoDuplicates), ObservedTypesOverlap, FragmentObservers[(uint8)Operation]);
+		}
+
+		if (bHasTagsOverlap)
+		{
+			ObservedTypesOverlap.Reset();
+			TagOverlap.ExportTypes(ObservedTypesOverlap);
+
+			HandleFragmentsImpl(ProcessingContext, FMassArchetypeSubChunks(ArchetypeHandle, MakeArrayView(&Entity, 1)
+				, FMassArchetypeSubChunks::NoDuplicates), ObservedTypesOverlap, TagObservers[(uint8)Operation]);
+		}
 	}
 
-	return false;
+	return bHasFragmentsOverlap || bHasTagsOverlap;
 }
 
-bool FMassObserverManager::OnPreCompositionRemoved(const FMassEntityHandle Entity, const FMassArchetypeCompositionDescriptor& Composition)
+void FMassObserverManager::OnSingleItemOperation(const UScriptStruct& ItemType, const FMassArchetypeSubChunks& ChunkCollection, const FMassObservedOperation Operation)
 {
-	const FMassFragmentBitSet Overlap = ObservedRemoveFragments.GetOverlap(Composition.Fragments);
-	if (Overlap.IsEmpty() == false)
+	check(ItemType.IsChildOf(FMassFragment::StaticStruct()) || ItemType.IsChildOf(FMassTag::StaticStruct()));
+
+	if (ItemType.IsChildOf(FMassFragment::StaticStruct()))
 	{
-		FMassProcessingContext ProcessingContext(EntitySubsystem, /*DeltaSeconds=*/0.f);
-		const FMassArchetypeHandle ArchetypeHandle = EntitySubsystem.GetArchetypeForEntity(Entity);
-		HandleFragmentsImpl(ProcessingContext, FMassArchetypeSubChunks(ArchetypeHandle, MakeArrayView(&Entity, 1)
-			, FMassArchetypeSubChunks::NoDuplicates), Overlap, OnFragmentRemovedObservers);
-		return true;
+		if (ObservedFragments[(uint8)Operation].Contains(ItemType))
+		{
+			HandleSingleItemImpl(ItemType, ChunkCollection, FragmentObservers[(uint8)Operation]);
+		}
+	}
+	else if (ObservedTags[(uint8)Operation].Contains(ItemType))
+	{
+		HandleSingleItemImpl(ItemType, ChunkCollection, TagObservers[(uint8)Operation]);
 	}
-
-	return false;
 }
 
-void FMassObserverManager::HandleFragmentsImpl(FMassProcessingContext& ProcessingContext, const FMassArchetypeSubChunks& ChunkCollection, const FMassFragmentBitSet& FragmentsBitSet, TMap<const UScriptStruct*, FMassRuntimePipeline>& HandlersContainer)
-{
-	TArray<const UScriptStruct*> Fragments;
-	FragmentsBitSet.ExportTypes(Fragments);
-	
-	for (const UScriptStruct* FragmentType : Fragments)
+void FMassObserverManager::HandleFragmentsImpl(FMassProcessingContext& ProcessingContext, const FMassArchetypeSubChunks& ChunkCollection
+	, TArrayView<const UScriptStruct*> ObservedTypes
+	/*, const FMassFragmentBitSet& FragmentsBitSet*/, FMassItemObserversMap& HandlersContainer)
+{	
+	check(ObservedTypes.Num() > 0);
+
+	for (const UScriptStruct* Type : ObservedTypes)
 	{		
-		ProcessingContext.AuxData.InitializeAs(FragmentType);
-		FMassRuntimePipeline& Pipeline = HandlersContainer.FindChecked(FragmentType);
+		ProcessingContext.AuxData.InitializeAs(Type);
+		FMassRuntimePipeline& Pipeline = (*HandlersContainer).FindChecked(Type);
 
 		UE::Mass::Executor::RunProcessorsView(Pipeline.Processors, ProcessingContext, &ChunkCollection);
 	}
 }
 
-void FMassObserverManager::HandleSingleFragmentImpl(const UScriptStruct& FragmentType, const FMassArchetypeSubChunks& ChunkCollection, const FMassFragmentBitSet& FragmentFilterBitSet, TMap<const UScriptStruct*, FMassRuntimePipeline>& HandlersContainer)
+void FMassObserverManager::HandleSingleItemImpl(const UScriptStruct& FragmentType, const FMassArchetypeSubChunks& ChunkCollection, FMassItemObserversMap& HandlersContainer)
 {
-	if (FragmentFilterBitSet.Contains(FragmentType))
-	{
-		FMassProcessingContext ProcessingContext(EntitySubsystem, /*DeltaSeconds=*/0.f);
-		ProcessingContext.AuxData.InitializeAs(&FragmentType);
-		FMassRuntimePipeline& Pipeline = HandlersContainer.FindChecked(&FragmentType	);
+	FMassProcessingContext ProcessingContext(EntitySubsystem, /*DeltaSeconds=*/0.f);
+	ProcessingContext.AuxData.InitializeAs(&FragmentType);
+	FMassRuntimePipeline& Pipeline = (*HandlersContainer).FindChecked(&FragmentType);
 
-		UE::Mass::Executor::RunProcessorsView(Pipeline.Processors, ProcessingContext, &ChunkCollection);
+	UE::Mass::Executor::RunProcessorsView(Pipeline.Processors, ProcessingContext, &ChunkCollection);
+}
+
+void FMassObserverManager::AddObserverInstance(const UScriptStruct& ItemType, const FMassObservedOperation Operation, UMassProcessor& ObserverProcessor)
+{
+	checkSlow(ItemType.IsChildOf(FMassFragment::StaticStruct()) || ItemType.IsChildOf(FMassTag::StaticStruct()));
+
+	FMassRuntimePipeline* Pipeline = nullptr;
+
+	if (ItemType.IsChildOf(FMassFragment::StaticStruct()))
+	{
+		Pipeline = &(*FragmentObservers[(uint8)Operation]).FindOrAdd(&ItemType);
+		ObservedFragments[(uint8)Operation].Add(ItemType);
 	}
+	else
+	{
+		Pipeline = &(*TagObservers[(uint8)Operation]).FindOrAdd(&ItemType);
+		ObservedTags[(uint8)Operation].Add(ItemType);
+	}
+	Pipeline->AppendProcessor(ObserverProcessor);
+
+	// calling initialize to ensure the given processor is related to the same EntitySubsystem
+	ObserverProcessor.Initialize(EntitySubsystem);
 }
diff --git a/Engine/Plugins/Runtime/MassEntity/Source/MassEntity/Private/MassObserverProcessor.cpp b/Engine/Plugins/Runtime/MassEntity/Source/MassEntity/Private/MassObserverProcessor.cpp
--- a/Engine/Plugins/Runtime/MassEntity/Source/MassEntity/Private/MassObserverProcessor.cpp
+++ b/Engine/Plugins/Runtime/MassEntity/Source/MassEntity/Private/MassObserverProcessor.cpp
@@ -30,7 +30,7 @@ void UMassObserverProcessor::PostInitProperties()
 void UMassFragmentInitializer::Register()
 {
 	check(FragmentType);
-	UMassObserverRegistry::GetMutable().RegisterFragmentInitializer(*FragmentType, GetClass());
+	UMassObserverRegistry::GetMutable().RegisterFragmentAddedObserver(*FragmentType, GetClass());
 }
 
 //----------------------------------------------------------------------//
@@ -39,5 +39,5 @@ void UMassFragmentInitializer::Register()
 void UMassFragmentDeinitializer::Register()
 {
 	check(FragmentType);
-	UMassObserverRegistry::GetMutable().RegisterFragmentDeinitializer(*FragmentType, GetClass());
+	UMassObserverRegistry::GetMutable().RegisterFragmentRemovedObserver(*FragmentType, GetClass());
 }
diff --git a/Engine/Plugins/Runtime/MassEntity/Source/MassEntity/Private/MassObserverRegistry.cpp b/Engine/Plugins/Runtime/MassEntity/Source/MassEntity/Private/MassObserverRegistry.cpp
--- a/Engine/Plugins/Runtime/MassEntity/Source/MassEntity/Private/MassObserverRegistry.cpp
+++ b/Engine/Plugins/Runtime/MassEntity/Source/MassEntity/Private/MassObserverRegistry.cpp
@@ -11,16 +11,28 @@ UMassObserverRegistry::UMassObserverRegistry()
 	check(HasAnyFlags(RF_ClassDefaultObject));
 }
 
-void UMassObserverRegistry::RegisterFragmentInitializer(const UScriptStruct& FragmentType, TSubclassOf<UMassProcessor> FragmentInitializerClass)
+void UMassObserverRegistry::RegisterFragmentAddedObserver(const UScriptStruct& FragmentType, TSubclassOf<UMassProcessor> ObserverClass)
 {
-	check(FragmentInitializerClass);
-	FragmentInitializersMap.FindOrAdd(&FragmentType).ClassCollection.AddUnique(FragmentInitializerClass);
+	check(ObserverClass);
+	FragmentInitializersMap.FindOrAdd(&FragmentType).ClassCollection.AddUnique(ObserverClass);
 }
 
-void UMassObserverRegistry::RegisterFragmentDeinitializer(const UScriptStruct& FragmentType, TSubclassOf<UMassProcessor> FragmentDeinitializerClass)
+void UMassObserverRegistry::RegisterFragmentRemovedObserver(const UScriptStruct& FragmentType, TSubclassOf<UMassProcessor> ObserverClass)
 {
-	check(FragmentDeinitializerClass);
-	FragmentDeinitializersMap.FindOrAdd(&FragmentType).ClassCollection.AddUnique(FragmentDeinitializerClass);
+	check(ObserverClass);
+	FragmentDeinitializersMap.FindOrAdd(&FragmentType).ClassCollection.AddUnique(ObserverClass);
+}
+
+void UMassObserverRegistry::RegisterTagAddedObserver(const UScriptStruct& FragmentType, TSubclassOf<UMassProcessor> ObserverClass)
+{
+	check(ObserverClass);
+	TagAddedObserversMap.FindOrAdd(&FragmentType).ClassCollection.AddUnique(ObserverClass);
+}
+
+void UMassObserverRegistry::RegisterTagRemovedObserver(const UScriptStruct& FragmentType, TSubclassOf<UMassProcessor> ObserverClass)
+{
+	check(ObserverClass);
+	TagRemovedObserversMap.FindOrAdd(&FragmentType).ClassCollection.AddUnique(ObserverClass);
 }
 
 TConstArrayView<TSubclassOf<UMassProcessor>> UMassObserverRegistry::GetFragmentInitializers(const UScriptStruct& FragmentType) const
diff --git a/Engine/Plugins/Runtime/MassEntity/Source/MassEntity/Public/MassObserverManager.h b/Engine/Plugins/Runtime/MassEntity/Source/MassEntity/Public/MassObserverManager.h
--- a/Engine/Plugins/Runtime/MassEntity/Source/MassEntity/Public/MassObserverManager.h
+++ b/Engine/Plugins/Runtime/MassEntity/Source/MassEntity/Public/MassObserverManager.h
@@ -8,6 +8,37 @@
 
 class UMassEntitySubsystem;
 struct FMassArchetypeSubChunks;
+class UMassProcessor;
+struct FMassProcessingContext;
+
+enum class FMassObservedOperation : uint8
+{
+	Add,
+	Remove,
+	// @todo Keeping this here as a indication of design intent. For now we handle entity destruction like removal, but 
+	// there might be coputationaly expensive cases where we might want to avoid for soon-to-be-dead entities. 
+	// Destroy,
+	MAX
+};
+
+/** 
+ * A wrapper type for a TMap to support having array-of-maps UPROPERTY members inFMassObserverManager
+ */
+USTRUCT()
+struct FMassItemObserversMap
+{
+	GENERATED_BODY()
+
+	// a helper accessor simplifying access while still keeping Container private
+	TMap<const UScriptStruct*, FMassRuntimePipeline>& operator*()
+	{
+		return Container;
+	}
+
+private:
+	UPROPERTY()
+	TMap<const UScriptStruct*, FMassRuntimePipeline> Container;
+};
 
 /** 
  * A type that encapsulates logic related to notifying interested parties of entity composition changes. Upon creation it
@@ -22,15 +53,30 @@ struct MASSENTITY_API FMassObserverManager
 public:
 	FMassObserverManager();	
 
-	DECLARE_MULTICAST_DELEGATE_TwoParams(FOnObservedFragmentTypesChanged, const FMassFragmentBitSet& /*NewObservedAddSet*/, const FMassFragmentBitSet& /*NewObservedRemoveSet*/);
+	const FMassFragmentBitSet& GetObservedAddFragmentsBitSet() const { return ObservedFragments[(uint8)FMassObservedOperation::Add]; }
+	const FMassFragmentBitSet& GetObservedRemoveFragmentsBitSet() const { return ObservedFragments[(uint8)FMassObservedOperation::Remove]; }
+	const FMassTagBitSet& GetObservedAddTagsBitSet() const { return ObservedTags[(uint8)FMassObservedOperation::Add]; }
+	const FMassTagBitSet& GetObservedRemoveTagsBitSet() const { return ObservedTags[(uint8)FMassObservedOperation::Remove]; }
 
-	FOnObservedFragmentTypesChanged& GetOnObservedFragmentTypesChangedDelegate() { return OnObservedFragmentTypesChangedDelegate; }
+	bool HasOnAddedObserversForBitSet(const FMassFragmentBitSet& InQueriedBitSet) const { return ObservedFragments[(uint8)FMassObservedOperation::Add].HasAny(InQueriedBitSet); }
+	bool HasOnAddedObserversForBitSet(const FMassTagBitSet& InQueriedBitSet) const { return ObservedTags[(uint8)FMassObservedOperation::Add].HasAny(InQueriedBitSet); }
+	bool HasOnRemovedObserversForBitSet(const FMassFragmentBitSet& InQueriedBitSet) const { return ObservedFragments[(uint8)FMassObservedOperation::Remove].HasAny(InQueriedBitSet); }
+	bool HasOnRemovedObserversForBitSet(const FMassTagBitSet& InQueriedBitSet) const { return ObservedTags[(uint8)FMassObservedOperation::Remove].HasAny(InQueriedBitSet); }
 
-	const FMassFragmentBitSet& GetObservedAddFragmentsBitSet() const { return ObservedAddFragments; }
-	const FMassFragmentBitSet& GetObservedRemoveFragmentsBitSet() const { return ObservedRemoveFragments; }
+	bool HasObserversForBitSet(const FMassFragmentBitSet& InQueriedBitSet, const FMassObservedOperation Operation) const
+	{
+		return ObservedFragments[(uint8)Operation].HasAny(InQueriedBitSet);
+	}
+
+	bool HasObserversForBitSet(const FMassTagBitSet& InQueriedBitSet, const FMassObservedOperation Operation) const
+	{
+		return ObservedTags[(uint8)Operation].HasAny(InQueriedBitSet);
+	}
 
-	bool HasOnAddedObserversForFragments(const FMassFragmentBitSet& InQueriedBitSet) const { return ObservedAddFragments.HasAny(InQueriedBitSet); }
-	bool HasOnRemovedObserversForFragments(const FMassFragmentBitSet& InQueriedBitSet) const { return ObservedRemoveFragments.HasAny(InQueriedBitSet); }
+	bool HasObserversForComposition(const FMassArchetypeCompositionDescriptor& Composition, const FMassObservedOperation Operation) const
+	{
+		return HasObserversForBitSet(Composition.Fragments, Operation) || HasObserversForBitSet(Composition.Tags, Operation);
+	}
 
 	bool OnPostEntitiesCreated(const FMassArchetypeSubChunks& ChunkCollection);
 	bool OnPostEntitiesCreated(FMassProcessingContext& ProcessingContext, const FMassArchetypeSubChunks& ChunkCollection);
@@ -38,18 +84,33 @@ public:
 	bool OnPreEntitiesDestroyed(const FMassArchetypeSubChunks& ChunkCollection);
 	bool OnPreEntitiesDestroyed(FMassProcessingContext& ProcessingContext, const FMassArchetypeSubChunks& ChunkCollection);
 
-	bool OnPostCompositionAdded(const FMassEntityHandle Entity, const FMassArchetypeCompositionDescriptor& Composition);
-	bool OnPreCompositionRemoved(const FMassEntityHandle Entity, const FMassArchetypeCompositionDescriptor& Composition);
-
-	void OnPostFragmentAdded(const UScriptStruct& FragmentType, const FMassArchetypeSubChunks& ChunkCollection)
+	bool OnCompositionChanged(const FMassEntityHandle Entity, const FMassArchetypeCompositionDescriptor& CompositionDelta, const FMassObservedOperation Operation);
+	bool OnPostCompositionAdded(const FMassEntityHandle Entity, const FMassArchetypeCompositionDescriptor& Composition)
 	{
-		HandleSingleFragmentImpl(FragmentType, ChunkCollection, ObservedAddFragments, OnFragmentAddedObservers);
+		return OnCompositionChanged(Entity, Composition, FMassObservedOperation::Add);
 	}
+	bool OnPreCompositionRemoved(const FMassEntityHandle Entity, const FMassArchetypeCompositionDescriptor& Composition)
+	{
+		return OnCompositionChanged(Entity, Composition, FMassObservedOperation::Remove);
+	}
+
+	// @todo caller will need to recreate ChunkCollection if this call is done after modifications to entities, so it 
+	// would be beneficial to ask the system "if there's anything observing the change" before actually reconstructing 
+	// the chunk collection. Unless the new ChunkCollection gets created as a side-product of the batched modification 
+	// operation
+	bool OnCompositionChanged(const FMassArchetypeSubChunks& ChunkCollection, const FMassArchetypeCompositionDescriptor& Composition, const FMassObservedOperation Operation, FMassProcessingContext* InProcessingContext = nullptr);
 
-	void OnPreFragmentRemoved(const UScriptStruct& FragmentType, const FMassArchetypeSubChunks& ChunkCollection)
+	void OnPostItemAdded(const UScriptStruct& ItemType, const FMassArchetypeSubChunks& ChunkCollection)
+	{
+		OnSingleItemOperation(ItemType, ChunkCollection, FMassObservedOperation::Add);
+	}
+	void OnPreItemRemoved(const UScriptStruct& ItemType, const FMassArchetypeSubChunks& ChunkCollection)
 	{
-		HandleSingleFragmentImpl(FragmentType, ChunkCollection, ObservedRemoveFragments, OnFragmentRemovedObservers);
+		OnSingleItemOperation(ItemType, ChunkCollection, FMassObservedOperation::Remove);
 	}
+	void OnSingleItemOperation(const UScriptStruct& ItemType, const FMassArchetypeSubChunks& ChunkCollection, const FMassObservedOperation Operation);
+
+	void AddObserverInstance(const UScriptStruct& ItemType, const FMassObservedOperation Operation, UMassProcessor& ObserverProcessor);
 
 protected:
 	friend UMassEntitySubsystem;
@@ -57,20 +118,17 @@ protected:
 
 	void Initialize();
 	void HandleFragmentsImpl(FMassProcessingContext& ProcessingContext, const FMassArchetypeSubChunks& ChunkCollection
-		, const FMassFragmentBitSet& FragmentsBitSet, TMap<const UScriptStruct*, FMassRuntimePipeline>& HandlersContainer);
-	void HandleSingleFragmentImpl(const UScriptStruct& FragmentType, const FMassArchetypeSubChunks& ChunkCollection
-		, const FMassFragmentBitSet& FragmentFilterBitSet, TMap<const UScriptStruct*, FMassRuntimePipeline>& HandlersContainer);
-
-	FOnObservedFragmentTypesChanged OnObservedFragmentTypesChangedDelegate;
+		, TArrayView<const UScriptStruct*> ObservedTypes, FMassItemObserversMap& HandlersContainer);
+	void HandleSingleItemImpl(const UScriptStruct& FragmentType, const FMassArchetypeSubChunks& ChunkCollection, FMassItemObserversMap& HandlersContainer);
 
-	FMassFragmentBitSet ObservedAddFragments;
-	FMassFragmentBitSet ObservedRemoveFragments;
+	FMassFragmentBitSet ObservedFragments[(uint8)FMassObservedOperation::MAX];
+	FMassTagBitSet ObservedTags[(uint8)FMassObservedOperation::MAX];
 
 	UPROPERTY()
-	TMap<const UScriptStruct*, FMassRuntimePipeline> OnFragmentAddedObservers;
+	FMassItemObserversMap FragmentObservers[(uint8)FMassObservedOperation::MAX];
 
 	UPROPERTY()
-	TMap<const UScriptStruct*, FMassRuntimePipeline> OnFragmentRemovedObservers;
+	FMassItemObserversMap TagObservers[(uint8)FMassObservedOperation::MAX];
 
 	/** 
 	 * The owning EntitySubsystem. No need for it to be a UPROPERTY since by design we don't support creation of 
@@ -86,4 +144,4 @@ struct TStructOpsTypeTraits<FMassObserverManager> : public TStructOpsTypeTraitsB
 	{
 		WithCopy = false,
 	};
-};
\ No newline at end of file
+};
diff --git a/Engine/Plugins/Runtime/MassEntity/Source/MassEntity/Public/MassObserverRegistry.h b/Engine/Plugins/Runtime/MassEntity/Source/MassEntity/Public/MassObserverRegistry.h
--- a/Engine/Plugins/Runtime/MassEntity/Source/MassEntity/Public/MassObserverRegistry.h
+++ b/Engine/Plugins/Runtime/MassEntity/Source/MassEntity/Public/MassObserverRegistry.h
@@ -35,7 +35,7 @@ public:
 	{
 		ensureMsgf(CDO.HasAnyFlags(RF_ClassDefaultObject) && CDO.GetClass()->HasAnyClassFlags(CLASS_Abstract) == false
 			, TEXT("Only CDOs are expected to be registered with UMassObserverRegistry::RegisterInitializer"));
-		GetMutable().RegisterFragmentInitializer(*TFragment::StaticStruct(), CDO.GetClass());
+		GetMutable().RegisterFragmentAddedObserver(*TFragment::StaticStruct(), CDO.GetClass());
 	}
 
 	template<typename TFragment>
@@ -43,11 +43,13 @@ public:
 	{
 		ensureMsgf(CDO.HasAnyFlags(RF_ClassDefaultObject) && CDO.GetClass()->HasAnyClassFlags(CLASS_Abstract) == false
 			, TEXT("Only CDOs are expected to be registered with UMassObserverRegistry::RegisterDeinitializer"));
-		GetMutable().RegisterFragmentDeinitializer(*TFragment::StaticStruct(), CDO.GetClass());
+		GetMutable().RegisterFragmentRemovedObserver(*TFragment::StaticStruct(), CDO.GetClass());
 	}
 
-	void RegisterFragmentInitializer(const UScriptStruct& FragmentType, TSubclassOf<UMassProcessor> FragmentInitializerClass);
-	void RegisterFragmentDeinitializer(const UScriptStruct& FragmentType, TSubclassOf<UMassProcessor> FragmentDeinitializerClass);
+	void RegisterFragmentAddedObserver(const UScriptStruct& FragmentType, TSubclassOf<UMassProcessor> FragmentInitializerClass);
+	void RegisterFragmentRemovedObserver(const UScriptStruct& FragmentType, TSubclassOf<UMassProcessor> FragmentDeinitializerClass);
+	void RegisterTagAddedObserver(const UScriptStruct& FragmentType, TSubclassOf<UMassProcessor> FragmentInitializerClass);
+	void RegisterTagRemovedObserver(const UScriptStruct& FragmentType, TSubclassOf<UMassProcessor> FragmentDeinitializerClass);
 
 protected:
 	friend FMassObserverManager;
@@ -60,4 +62,10 @@ protected:
 
 	UPROPERTY()
 	TMap<const UScriptStruct*, FMassProcessorClassCollection> FragmentDeinitializersMap;
+
+	UPROPERTY()
+	TMap<const UScriptStruct*, FMassProcessorClassCollection> TagAddedObserversMap;
+	
+	UPROPERTY()
+	TMap<const UScriptStruct*, FMassProcessorClassCollection> TagRemovedObserversMap;
 };
diff --git a/Engine/Plugins/Runtime/MassEntity/Source/MassEntityTestSuite/Private/MassObserverTest.cpp b/Engine/Plugins/Runtime/MassEntity/Source/MassEntityTestSuite/Private/MassObserverTest.cpp
new file mode 100644
--- /dev/null
+++ b/Engine/Plugins/Runtime/MassEntity/Source/MassEntityTestSuite/Private/MassObserverTest.cpp
@@ -0,0 +1,197 @@
+// Copyright Epic Games, Inc. All Rights Reserved.
+
+#include "Engine/World.h"
+#include "MassEntitySubsystem.h"
+#include "MassProcessingTypes.h"
+#include "MassEntityTestTypes.h"
+#include "MassEntityTypes.h"
+#include "MassCommandBuffer.h"
+
+#define LOCTEXT_NAMESPACE "MassTest"
+
+PRAGMA_DISABLE_OPTIMIZATION
+
+//----------------------------------------------------------------------//
+// tests 
+//----------------------------------------------------------------------//
+namespace FMassObserverTest
+{
+
+auto EntityIndexSorted = [](const FMassEntityHandle& A, const FMassEntityHandle& B)
+{
+	return A.Index < B.Index;
+};
+
+struct FTagBaseOperation : FEntityTestBase
+{
+	TArray<FMassEntityHandle> AffectedEntities;
+	UMassTestProcessorBase* TagObserver = nullptr;
+	FMassObservedOperation OperationObserved = FMassObservedOperation::MAX;
+	TArray<FMassEntityHandle> EntitiesInt;
+	TArray<FMassEntityHandle> EntitiesIntsFloat;
+	TArray<FMassEntityHandle> ExpectedEntities;
+	// @return signifies if the test can continue
+	virtual bool PerformOperation() = 0;
+
+	virtual bool SetUp() override
+	{
+		if (FEntityTestBase::SetUp())
+		{
+			//TArray<FMassEntityHandle>* ContainerPtr = &AffectedEntities;
+			TagObserver = NewObject<UMassTestProcessorBase>(EntitySubsystem);
+			TagObserver->TestGetQuery().AddRequirement<FTestFragment_Int>(EMassFragmentAccess::ReadOnly);
+			TagObserver->TestGetQuery().AddTagRequirement<FTestTag_A>(EMassFragmentPresence::All);
+			TagObserver->ExecutionFunction = [this](UMassEntitySubsystem& InEntitySubsystem, FMassExecutionContext& Context)
+			{
+				TagObserver->TestGetQuery().ForEachEntityChunk(InEntitySubsystem, Context, [this](FMassExecutionContext& Context)
+					{
+						AffectedEntities.Append(Context.GetEntities().GetData(), Context.GetEntities().Num());
+					});
+			};
+
+			return true;
+		}
+		return false;
+	}
+
+	virtual bool InstantTest() override
+	{
+		FMassObserverManager& ObserverManager = EntitySubsystem->GetObserverManager();
+		ObserverManager.AddObserverInstance(*FTestTag_A::StaticStruct(), OperationObserved, *TagObserver);
+
+		EntitySubsystem->BatchCreateEntities(IntsArchetype, 3, EntitiesInt);
+		EntitySubsystem->BatchCreateEntities(FloatsIntsArchetype, 3, EntitiesIntsFloat);
+
+		if (PerformOperation())
+		{
+			EntitySubsystem->Defer().ReplayBufferAgainstSystem(EntitySubsystem);
+			AITEST_EQUAL(TEXT("The tag observer is expected to be run for predicted number of entities"), AffectedEntities.Num(), ExpectedEntities.Num());
+			
+			ExpectedEntities.Sort(EntityIndexSorted);
+			AffectedEntities.Sort(EntityIndexSorted);
+
+			for (int i = 0; i < ExpectedEntities.Num(); ++i)
+			{
+				AITEST_EQUAL(TEXT("Expected and affected sets should be the same"), AffectedEntities[i], ExpectedEntities[i]);
+			}
+
+			/*AITEST_EQUAL(TEXT("The tag observer is expected to be run for a single entity"), AffectedEntities.Num(), 1);
+			AITEST_EQUAL(TEXT("The tag observer is expected to be run for the entity that received the tag"), AffectedEntities[0], Entities[1]);*/
+		}
+
+		return true;
+	}
+};
+
+struct FTag_SingleEntitySingleArchetypeAdd : FTagBaseOperation
+{
+	FTag_SingleEntitySingleArchetypeAdd() { OperationObserved = FMassObservedOperation::Add; }
+	virtual bool PerformOperation() override 
+	{
+		ExpectedEntities = { EntitiesInt[1] };
+		EntitySubsystem->Defer().AddTag<FTestTag_A>(EntitiesInt[1]);
+		return true;
+	}
+};
+IMPLEMENT_AI_INSTANT_TEST(FTag_SingleEntitySingleArchetypeAdd, "System.Mass.Observer.Tag.SingleEntitySingleArchetypeAdd");
+
+struct FTag_SingleEntitySingleArchetypeRemove : FTagBaseOperation
+{
+	FTag_SingleEntitySingleArchetypeRemove() { OperationObserved = FMassObservedOperation::Remove; }
+	virtual bool PerformOperation() override
+	{
+		ExpectedEntities = { EntitiesInt[1] };
+
+		EntitySubsystem->Defer().AddTag<FTestTag_A>(EntitiesInt[1]);
+		EntitySubsystem->Defer().ReplayBufferAgainstSystem(EntitySubsystem);
+		// since we're only observing tag removal we don't expect AffectedEntities to contain any data at this point
+		AITEST_EQUAL(TEXT("Tag addition is not being observed and is not expected to produce results yet"), AffectedEntities.Num(), 0);
+		EntitySubsystem->Defer().RemoveTag<FTestTag_A>(EntitiesInt[1]);
+		return true;
+	}
+};
+IMPLEMENT_AI_INSTANT_TEST(FTag_SingleEntitySingleArchetypeRemove, "System.Mass.Observer.Tag.SingleEntitySingleArchetypeRemove");
+
+struct FTag_SingleEntitySingleArchetypeDestroy : FTagBaseOperation
+{
+	FTag_SingleEntitySingleArchetypeDestroy() { OperationObserved = FMassObservedOperation::Remove; }
+	virtual bool PerformOperation() override
+	{
+		ExpectedEntities = { EntitiesInt[1] };
+		EntitySubsystem->Defer().AddTag<FTestTag_A>(EntitiesInt[1]); 
+		EntitySubsystem->Defer().ReplayBufferAgainstSystem(EntitySubsystem);
+		// since we're only observing tag removal we don't expect AffectedEntities to contain any data at this point
+		AITEST_EQUAL(TEXT("Tag addition is not being observed and is not expected to produce results yet"), AffectedEntities.Num(), 0);
+		EntitySubsystem->Defer().DestroyEntity(EntitiesInt[1]);
+		return true;
+	}
+};
+IMPLEMENT_AI_INSTANT_TEST(FTag_SingleEntitySingleArchetypeDestroy, "System.Mass.Observer.Tag.SingleEntitySingleArchetypeDestroy");
+
+struct FTag_MultipleArchetypeAdd : FTagBaseOperation
+{
+	FTag_MultipleArchetypeAdd() { OperationObserved = FMassObservedOperation::Add; }
+
+	virtual bool PerformOperation() override
+	{
+		ExpectedEntities = { EntitiesInt[0], EntitiesInt[2], EntitiesIntsFloat[1] };
+		for (const FMassEntityHandle& ModifiedEntity : ExpectedEntities)
+		{
+			EntitySubsystem->Defer().AddTag<FTestTag_A>(ModifiedEntity);
+		}
+		return true;
+	}
+};
+IMPLEMENT_AI_INSTANT_TEST(FTag_MultipleArchetypeAdd, "System.Mass.Observer.Tag.MultipleArchetypesAdd");
+
+struct FTag_MultipleArchetypeRemove : FTagBaseOperation
+{
+	FTag_MultipleArchetypeRemove() { OperationObserved = FMassObservedOperation::Remove; }
+
+	virtual bool PerformOperation() override
+	{
+		ExpectedEntities = { EntitiesInt[0], EntitiesInt[2], EntitiesIntsFloat[1] };
+		for (const FMassEntityHandle& ModifiedEntity : ExpectedEntities)
+		{
+			EntitySubsystem->Defer().AddTag<FTestTag_A>(ModifiedEntity);
+		}
+		EntitySubsystem->Defer().ReplayBufferAgainstSystem(EntitySubsystem);
+		// since we're only observing tag removal we don't expect AffectedEntities to contain any data at this point
+		AITEST_EQUAL(TEXT("Tag addition is not being observed and is not expected to produce results yet"), AffectedEntities.Num(), 0);
+		for (const FMassEntityHandle& ModifiedEntity : ExpectedEntities)
+		{
+			EntitySubsystem->Defer().RemoveTag<FTestTag_A>(ModifiedEntity);
+		}
+		return true;
+	}
+};
+IMPLEMENT_AI_INSTANT_TEST(FTag_MultipleArchetypeRemove, "System.Mass.Observer.Tag.MultipleArchetypesRemove");
+
+struct FTag_MultipleArchetypeDestroy : FTagBaseOperation
+{
+	FTag_MultipleArchetypeDestroy() { OperationObserved = FMassObservedOperation::Remove; }
+
+	virtual bool PerformOperation() override
+	{
+		ExpectedEntities = { EntitiesInt[0], EntitiesInt[2], EntitiesIntsFloat[1] };
+		for (const FMassEntityHandle& ModifiedEntity : ExpectedEntities)
+		{
+			EntitySubsystem->Defer().AddTag<FTestTag_A>(ModifiedEntity);
+		}
+		EntitySubsystem->Defer().ReplayBufferAgainstSystem(EntitySubsystem);
+		// since we're only observing tag removal we don't expect AffectedEntities to contain any data at this point
+		AITEST_EQUAL(TEXT("Tag addition is not being observed and is not expected to produce results yet"), AffectedEntities.Num(), 0);
+		for (const FMassEntityHandle& ModifiedEntity : ExpectedEntities)
+		{
+			EntitySubsystem->Defer().DestroyEntity(ModifiedEntity);
+		}
+		return true;
+	}
+};
+IMPLEMENT_AI_INSTANT_TEST(FTag_MultipleArchetypeDestroy, "System.Mass.Observer.Tag.MultipleArchetypesDestroy");
+
+} // FMassObserverTest
+
+PRAGMA_ENABLE_OPTIMIZATION
+
+#undef LOCTEXT_NAMESPACE
