Fixed Mass' observation not working for Tags.
Also fixed edge cases where executing observers affected the top-level command buffer due to ProcessingContext missconfig.
Also fixed the "ObservedTypes" not getting copied during FMassCommandBuffer::MoveAppend

diff --git a/Engine/Plugins/Runtime/MassEntity/Source/MassEntity/Private/MassCommandBuffer.cpp b/Engine/Plugins/Runtime/MassEntity/Source/MassEntity/Private/MassCommandBuffer.cpp
--- a/Engine/Plugins/Runtime/MassEntity/Source/MassEntity/Private/MassCommandBuffer.cpp
+++ b/Engine/Plugins/Runtime/MassEntity/Source/MassEntity/Private/MassCommandBuffer.cpp
@@ -67,15 +67,45 @@ void GetCommandStatNames(FStructView Entry, FString& OutName, ANSIName*& OutANSI
 #endif
 } // UE::FLWCCommand
 
+
+//////////////////////////////////////////////////////////////////////
+// FMassObservedTypeCollection
+
+void FMassObservedTypeCollection::Append(const FMassObservedTypeCollection& Other)
+{
+	for (auto It : Other.Types)
+	{
+		TArray<FMassEntityHandle> Contents = Types.FindOrAdd(It.Key);
+		Contents.Append(It.Value);
+	}
+}
+
 //////////////////////////////////////////////////////////////////////
 // FMassCommandsObservedTypes
 
 void FMassCommandsObservedTypes::Reset()
 {
-	FragmentsToAdd.Reset();
-	FragmentsToRemove.Reset();
+	for (FMassObservedTypeCollection& Collection : Fragments)
+	{
+		Collection.Reset();
+	};
+	for (FMassObservedTypeCollection& Collection : Tags)
+	{
+		Collection.Reset();
+	};
 }
 
+void FMassCommandsObservedTypes::Append(const FMassCommandsObservedTypes& Other)
+{
+	for (int i = 0; i < (int)EMassObservedOperation::MAX; ++i)
+	{
+		Fragments[i].Append(Other.Fragments[i]);
+	}
+	for (int i = 0; i < (int)EMassObservedOperation::MAX; ++i)
+	{
+		Tags[i].Append(Other.Tags[i]);
+	}
+}
 
 //////////////////////////////////////////////////////////////////////
 // FMassCommandBuffer
@@ -84,6 +114,9 @@ void FMassCommandBuffer::ReplayBufferAgainstSystem(UMassEntitySubsystem* EntityS
 {
 	check(EntitySystem);
 
+	check(!bIsFlushing);
+	TGuardValue FlushingGuard(bIsFlushing, true);
+
 	// short-circuit exit
 	if (EntitiesToDestroy.Num() == 0 && PendingCommands.IsEmpty())
 	{
@@ -92,40 +125,47 @@ void FMassCommandBuffer::ReplayBufferAgainstSystem(UMassEntitySubsystem* EntityS
 
 	FMassObserverManager& ObserverManager = EntitySystem->GetObserverManager();
 
-	const FMassFragmentBitSet& ObservedAddFragments = ObserverManager.GetObservedAddFragmentsBitSet();
-	const FMassFragmentBitSet& ObservedRemoveFragments = ObserverManager.GetObservedRemoveFragmentsBitSet();
+	const FMassFragmentBitSet* ObservedFragments = ObserverManager.GetObservedFragmentBitSets();
+	const FMassTagBitSet* ObservedTags = ObserverManager.GetObservedTagBitSets();
 
-	if (ObservedRemoveFragments.IsEmpty() == false)
+	for (auto It : ObservedTypes.GetObservedFragments(EMassObservedOperation::Remove))
 	{
-		for (auto It : ObservedTypes.GetFragmentsToRemove())
+		check(It.Key);
+		if (ObservedFragments[(uint8)EMassObservedOperation::Remove].Contains(*It.Key))
 		{
-			check(It.Key);
-			if (ObservedRemoveFragments.Contains(*It.Key))
+			TArray<FMassArchetypeSubChunks > ChunkCollections;
+			UE::Mass::Utils::CreateSparseChunks(*EntitySystem, It.Value, FMassArchetypeSubChunks::FoldDuplicates, ChunkCollections);
+			for (FMassArchetypeSubChunks& Collection : ChunkCollections)
 			{
-				TArray<FMassArchetypeSubChunks > ChunkCollections;
-				UE::Mass::Utils::CreateSparseChunks(*EntitySystem, It.Value, FMassArchetypeSubChunks::FoldDuplicates, ChunkCollections);
-				for (FMassArchetypeSubChunks& Collection : ChunkCollections)
-				{
-					check(It.Key);
-					ObserverManager.OnPreItemRemoved(*It.Key, Collection);
-				}
+				check(It.Key);
+				ObserverManager.OnPreFragmentOrTagRemoved(*It.Key, Collection);
 			}
 		}
+	}
 	
-		TArray<FMassArchetypeSubChunks > EntityChunksToDestroy;
-		if (EntitiesToDestroy.Num())
+	for (auto It : ObservedTypes.GetObservedTags(EMassObservedOperation::Remove))
+	{
+		check(It.Key);
+		if (ObservedTags[(uint8)EMassObservedOperation::Remove].Contains(*It.Key))
 		{
-			UE::Mass::Utils::CreateSparseChunks(*EntitySystem, EntitiesToDestroy, FMassArchetypeSubChunks::FoldDuplicates, EntityChunksToDestroy);
-			for (FMassArchetypeSubChunks& Collection : EntityChunksToDestroy)
+			TArray<FMassArchetypeSubChunks > ChunkCollections;
+			UE::Mass::Utils::CreateSparseChunks(*EntitySystem, It.Value, FMassArchetypeSubChunks::FoldDuplicates, ChunkCollections);
+			for (FMassArchetypeSubChunks& Collection : ChunkCollections)
 			{
-				EntitySystem->BatchDestroyEntityChunks(Collection);
+				check(It.Key);
+				ObserverManager.OnPreFragmentOrTagRemoved(*It.Key, Collection);
 			}
 		}
 	}
-	else if (EntitiesToDestroy.Num())
+	
+	TArray<FMassArchetypeSubChunks > EntityChunksToDestroy;
+	if (EntitiesToDestroy.Num())
 	{
-		TRACE_CPUPROFILER_EVENT_SCOPE_STR("ECS Deferred Commands Destroy Entities");
-		EntitySystem->BatchDestroyEntities(MakeArrayView(EntitiesToDestroy));
+		UE::Mass::Utils::CreateSparseChunks(*EntitySystem, EntitiesToDestroy, FMassArchetypeSubChunks::FoldDuplicates, EntityChunksToDestroy);
+		for (FMassArchetypeSubChunks& Collection : EntityChunksToDestroy)
+		{
+			EntitySystem->BatchDestroyEntityChunks(Collection);
+		}
 	}
 	EntitiesToDestroy.Reset();
 
@@ -155,17 +195,32 @@ void FMassCommandBuffer::ReplayBufferAgainstSystem(UMassEntitySubsystem* EntityS
 	// Using Clear() instead of Reset(), as otherwise the chunks moved into the PendingCommands in MoveAppend() can accumulate.
 	PendingCommands.Clear();
 
-	for (auto It : ObservedTypes.GetFragmentsToAdd())
+	for (auto It : ObservedTypes.GetObservedFragments(EMassObservedOperation::Add))
+	{
+		check(It.Key);
+		if (ObservedFragments[(uint8)EMassObservedOperation::Add].Contains(*It.Key))
+		{
+			TArray<FMassArchetypeSubChunks > ChunkCollections;
+			UE::Mass::Utils::CreateSparseChunks(*EntitySystem, It.Value, FMassArchetypeSubChunks::FoldDuplicates, ChunkCollections);
+			for (FMassArchetypeSubChunks& Collection : ChunkCollections)
+			{
+				check(It.Key);
+				ObserverManager.OnPostFragmentOrTagAdded(*It.Key, Collection);
+			}
+		}
+	}
+
+	for (auto It : ObservedTypes.GetObservedTags(EMassObservedOperation::Add))
 	{
 		check(It.Key);
-		if (ObservedAddFragments.Contains(*It.Key))
+		if (ObservedTags[(uint8)EMassObservedOperation::Add].Contains(*It.Key))
 		{
 			TArray<FMassArchetypeSubChunks > ChunkCollections;
 			UE::Mass::Utils::CreateSparseChunks(*EntitySystem, It.Value, FMassArchetypeSubChunks::FoldDuplicates, ChunkCollections);
 			for (FMassArchetypeSubChunks& Collection : ChunkCollections)
 			{
 				check(It.Key);
-				ObserverManager.OnPostItemAdded(*It.Key, Collection);
+				ObserverManager.OnPostFragmentOrTagAdded(*It.Key, Collection);
 			}
 		}
 	}
@@ -183,6 +238,7 @@ void FMassCommandBuffer::MoveAppend(FMassCommandBuffer& Other)
 		UE_MT_SCOPED_WRITE_ACCESS(PendingCommandsDetector);
 		PendingCommands.Append(MoveTemp(Other.PendingCommands));
 		EntitiesToDestroy.Append(MoveTemp(Other.EntitiesToDestroy));
+		ObservedTypes.Append(Other.ObservedTypes);
 	}
 }
 
diff --git a/Engine/Plugins/Runtime/MassEntity/Source/MassEntity/Private/MassEntitySubsystem.cpp b/Engine/Plugins/Runtime/MassEntity/Source/MassEntity/Private/MassEntitySubsystem.cpp
--- a/Engine/Plugins/Runtime/MassEntity/Source/MassEntity/Private/MassEntitySubsystem.cpp
+++ b/Engine/Plugins/Runtime/MassEntity/Source/MassEntity/Private/MassEntitySubsystem.cpp
@@ -348,7 +348,7 @@ TSharedRef<UMassEntitySubsystem::FEntityCreationContext> UMassEntitySubsystem::B
 	// @todo this could probably be optimized since one would assume we're adding elements to OutEntities in order.
 	// Then again, if that's the case, the sorting will be almost instant
 	new (&CreationContext->ChunkCollection)FMassArchetypeSubChunks(Archetype, MakeArrayView(&OutEntities[OutEntities.Num() - Count], Count), FMassArchetypeSubChunks::NoDuplicates);
-	if (ObserverManager.HasOnAddedObserversForBitSet(ArchetypePtr->GetCompositionDescriptor().Fragments))
+	if (ObserverManager.HasObserversForBitSet(ArchetypePtr->GetCompositionDescriptor().Fragments, EMassObservedOperation::Add))
 	{
 		CreationContext->OnSpawningFinished = [this](FEntityCreationContext& Context){
 			ObserverManager.OnPostEntitiesCreated(Context.ChunkCollection);
diff --git a/Engine/Plugins/Runtime/MassEntity/Source/MassEntity/Private/MassExecutor.cpp b/Engine/Plugins/Runtime/MassEntity/Source/MassEntity/Private/MassExecutor.cpp
--- a/Engine/Plugins/Runtime/MassEntity/Source/MassEntity/Private/MassExecutor.cpp
+++ b/Engine/Plugins/Runtime/MassEntity/Source/MassEntity/Private/MassExecutor.cpp
@@ -115,9 +115,24 @@ void RunProcessorsView(TArrayView<UMassProcessor*> Processors, FMassProcessingCo
 		
 		ExecutionContext.SetFlushDeferredCommands(true);
 		// append the commands added from other, non-processor sources (like MassAgentSubsystem)
+		ensure(!ProcessingContext.EntitySubsystem->Defer().IsFlushing());
+		ensure(!ExecutionContext.Defer().IsFlushing());
 		ExecutionContext.Defer().MoveAppend(ProcessingContext.EntitySubsystem->Defer());
 		ExecutionContext.FlushDeferred(*ProcessingContext.EntitySubsystem);
 	}
+	// else make sure we don't just lose the commands. Append to the command buffer requested via
+	// ProcessingContext.CommandBuffer or to the default EntitySubsystem's command buffer.
+	else if (CommandBuffer != ProcessingContext.CommandBuffer)
+	{
+		if (ProcessingContext.CommandBuffer)
+		{
+			ProcessingContext.CommandBuffer->MoveAppend(*CommandBuffer.Get());
+		}
+		else if (CommandBuffer.Get() != &ProcessingContext.EntitySubsystem->Defer())
+		{
+			ProcessingContext.EntitySubsystem->Defer().MoveAppend(*CommandBuffer.Get());
+		}
+	}
 }
 
 struct FMassExecutorDoneTask
diff --git a/Engine/Plugins/Runtime/MassEntity/Source/MassEntity/Private/MassObserverManager.cpp b/Engine/Plugins/Runtime/MassEntity/Source/MassEntity/Private/MassObserverManager.cpp
--- a/Engine/Plugins/Runtime/MassEntity/Source/MassEntity/Private/MassObserverManager.cpp
+++ b/Engine/Plugins/Runtime/MassEntity/Source/MassEntity/Private/MassObserverManager.cpp
@@ -10,7 +10,7 @@ namespace UE::Mass::ObserverManager::Private
 {
 // a helper function to reduce code duplication in FMassObserverManager::Initialize
 template<typename TBitSet>
-void SetUpObservers(UMassEntitySubsystem& EntitySubsystem, const TMap<const UScriptStruct*, FMassProcessorClassCollection>& RegisteredObserverTypes, TBitSet& ObservedBitSet, FMassItemObserversMap& Observers)
+void SetUpObservers(UMassEntitySubsystem& EntitySubsystem, const TMap<const UScriptStruct*, FMassProcessorClassCollection>& RegisteredObserverTypes, TBitSet& ObservedBitSet, FMassObserversMap& Observers)
 {
 	ObservedBitSet.Reset();
 
@@ -55,10 +55,11 @@ void FMassObserverManager::Initialize()
 	const UMassObserverRegistry& Registry = UMassObserverRegistry::Get();
 
 	using UE::Mass::ObserverManager::Private::SetUpObservers;
-	SetUpObservers(EntitySubsystem, Registry.FragmentInitializersMap, ObservedFragments[(uint8)FMassObservedOperation::Add], FragmentObservers[(uint8)FMassObservedOperation::Add]);
-	SetUpObservers(EntitySubsystem, Registry.FragmentDeinitializersMap, ObservedFragments[(uint8)FMassObservedOperation::Remove], FragmentObservers[(uint8)FMassObservedOperation::Remove]);
-	SetUpObservers(EntitySubsystem, Registry.TagAddedObserversMap, ObservedTags[(uint8)FMassObservedOperation::Add], TagObservers[(uint8)FMassObservedOperation::Add]);
-	SetUpObservers(EntitySubsystem, Registry.TagRemovedObserversMap, ObservedTags[(uint8)FMassObservedOperation::Remove], TagObservers[(uint8)FMassObservedOperation::Remove]);
+	for (int i = 0; i < (int)EMassObservedOperation::MAX; ++i)
+	{
+		SetUpObservers(EntitySubsystem, *Registry.FragmentObservers[i], ObservedFragments[i], FragmentObservers[i]);
+		SetUpObservers(EntitySubsystem, *Registry.TagObservers[i], ObservedTags[i], TagObservers[i]);
+	}
 }
 
 bool FMassObserverManager::OnPostEntitiesCreated(const FMassArchetypeSubChunks& ChunkCollection)
@@ -84,16 +85,14 @@ bool FMassObserverManager::OnPostEntitiesCreated(FMassProcessingContext& Process
 
 	check(ProcessingContext.EntitySubsystem);
 	const FMassArchetypeCompositionDescriptor& ArchetypeComposition = ProcessingContext.EntitySubsystem->GetArchetypeComposition(ChunkCollection.GetArchetype());
-
-
-	const FMassFragmentBitSet Overlap = ObservedFragments[(uint8)FMassObservedOperation::Add].GetOverlap(ArchetypeComposition.Fragments);
+	const FMassFragmentBitSet Overlap = ObservedFragments[(uint8)EMassObservedOperation::Add].GetOverlap(ArchetypeComposition.Fragments);
 
 	if (Overlap.IsEmpty() == false)
 	{
 		TArray<const UScriptStruct*> OverlapTypes;
 		Overlap.ExportTypes(OverlapTypes);
 
-		HandleFragmentsImpl(ProcessingContext, ChunkCollection, MakeArrayView(OverlapTypes), FragmentObservers[(uint8)FMassObservedOperation::Add]);
+		HandleFragmentsImpl(ProcessingContext, ChunkCollection, MakeArrayView(OverlapTypes), FragmentObservers[(uint8)EMassObservedOperation::Add]);
 		return true;
 	}
 
@@ -121,10 +120,10 @@ bool FMassObserverManager::OnPreEntitiesDestroyed(FMassProcessingContext& Proces
 	check(ProcessingContext.EntitySubsystem);
 	const FMassArchetypeCompositionDescriptor& ArchetypeComposition = ProcessingContext.EntitySubsystem->GetArchetypeComposition(ChunkCollection.GetArchetype());
 	
-	return OnCompositionChanged(ChunkCollection, ArchetypeComposition, FMassObservedOperation::Remove, &ProcessingContext);
+	return OnCompositionChanged(ChunkCollection, ArchetypeComposition, EMassObservedOperation::Remove, &ProcessingContext);
 }
 
-bool FMassObserverManager::OnCompositionChanged(const FMassArchetypeSubChunks& ChunkCollection, const FMassArchetypeCompositionDescriptor& CompositionDelta, const FMassObservedOperation Operation, FMassProcessingContext* InProcessingContext)
+bool FMassObserverManager::OnCompositionChanged(const FMassArchetypeSubChunks& ChunkCollection, const FMassArchetypeCompositionDescriptor& CompositionDelta, const EMassObservedOperation Operation, FMassProcessingContext* InProcessingContext)
 {
 	const FMassFragmentBitSet FragmentOverlap = ObservedFragments[(uint8)Operation].GetOverlap(CompositionDelta.Fragments);
 	const bool bHasFragmentsOverlap = !FragmentOverlap.IsEmpty();
@@ -134,6 +133,7 @@ bool FMassObserverManager::OnCompositionChanged(const FMassArchetypeSubChunks& C
 	if (bHasFragmentsOverlap || bHasTagsOverlap)
 	{
 		FMassProcessingContext LocalContext(EntitySubsystem, /*DeltaSeconds=*/0.f);
+		LocalContext.bFlushCommandBuffer = false;
 		FMassProcessingContext* ProcessingContext = InProcessingContext ? InProcessingContext : &LocalContext;
 		TArray<const UScriptStruct*> ObservedTypesOverlap;
 
@@ -156,7 +156,7 @@ bool FMassObserverManager::OnCompositionChanged(const FMassArchetypeSubChunks& C
 	return bHasFragmentsOverlap || bHasTagsOverlap;
 }
 
-bool FMassObserverManager::OnCompositionChanged(const FMassEntityHandle Entity, const FMassArchetypeCompositionDescriptor& CompositionDelta, const FMassObservedOperation Operation)
+bool FMassObserverManager::OnCompositionChanged(const FMassEntityHandle Entity, const FMassArchetypeCompositionDescriptor& CompositionDelta, const EMassObservedOperation Operation)
 {
 	const FMassFragmentBitSet FragmentOverlap = ObservedFragments[(uint8)Operation].GetOverlap(CompositionDelta.Fragments);
 	const bool bHasFragmentsOverlap = !FragmentOverlap.IsEmpty();
@@ -167,6 +167,7 @@ bool FMassObserverManager::OnCompositionChanged(const FMassEntityHandle Entity,
 	{
 		TArray<const UScriptStruct*> ObservedTypesOverlap;
 		FMassProcessingContext ProcessingContext(EntitySubsystem, /*DeltaSeconds=*/0.f);
+		ProcessingContext.bFlushCommandBuffer = false;
 		const FMassArchetypeHandle ArchetypeHandle = EntitySubsystem.GetArchetypeForEntity(Entity);
 
 		if (bHasFragmentsOverlap)
@@ -190,26 +191,26 @@ bool FMassObserverManager::OnCompositionChanged(const FMassEntityHandle Entity,
 	return bHasFragmentsOverlap || bHasTagsOverlap;
 }
 
-void FMassObserverManager::OnSingleItemOperation(const UScriptStruct& ItemType, const FMassArchetypeSubChunks& ChunkCollection, const FMassObservedOperation Operation)
+void FMassObserverManager::OnFragmentOrTagOperation(const UScriptStruct& FragmentOrTagType, const FMassArchetypeSubChunks& ChunkCollection, const EMassObservedOperation Operation)
 {
-	check(ItemType.IsChildOf(FMassFragment::StaticStruct()) || ItemType.IsChildOf(FMassTag::StaticStruct()));
+	check(FragmentOrTagType.IsChildOf(FMassFragment::StaticStruct()) || FragmentOrTagType.IsChildOf(FMassTag::StaticStruct()));
 
-	if (ItemType.IsChildOf(FMassFragment::StaticStruct()))
+	if (FragmentOrTagType.IsChildOf(FMassFragment::StaticStruct()))
 	{
-		if (ObservedFragments[(uint8)Operation].Contains(ItemType))
+		if (ObservedFragments[(uint8)Operation].Contains(FragmentOrTagType))
 		{
-			HandleSingleItemImpl(ItemType, ChunkCollection, FragmentObservers[(uint8)Operation]);
+			HandleSingleEntityImpl(FragmentOrTagType, ChunkCollection, FragmentObservers[(uint8)Operation]);
 		}
 	}
-	else if (ObservedTags[(uint8)Operation].Contains(ItemType))
+	else if (ObservedTags[(uint8)Operation].Contains(FragmentOrTagType))
 	{
-		HandleSingleItemImpl(ItemType, ChunkCollection, TagObservers[(uint8)Operation]);
+		HandleSingleEntityImpl(FragmentOrTagType, ChunkCollection, TagObservers[(uint8)Operation]);
 	}
 }
 
 void FMassObserverManager::HandleFragmentsImpl(FMassProcessingContext& ProcessingContext, const FMassArchetypeSubChunks& ChunkCollection
 	, TArrayView<const UScriptStruct*> ObservedTypes
-	/*, const FMassFragmentBitSet& FragmentsBitSet*/, FMassItemObserversMap& HandlersContainer)
+	/*, const FMassFragmentBitSet& FragmentsBitSet*/, FMassObserversMap& HandlersContainer)
 {	
 	check(ObservedTypes.Num() > 0);
 
@@ -222,30 +223,31 @@ void FMassObserverManager::HandleFragmentsImpl(FMassProcessingContext& Processin
 	}
 }
 
-void FMassObserverManager::HandleSingleItemImpl(const UScriptStruct& FragmentType, const FMassArchetypeSubChunks& ChunkCollection, FMassItemObserversMap& HandlersContainer)
+void FMassObserverManager::HandleSingleEntityImpl(const UScriptStruct& FragmentType, const FMassArchetypeSubChunks& ChunkCollection, FMassObserversMap& HandlersContainer)
 {
 	FMassProcessingContext ProcessingContext(EntitySubsystem, /*DeltaSeconds=*/0.f);
+	ProcessingContext.bFlushCommandBuffer = false;
 	ProcessingContext.AuxData.InitializeAs(&FragmentType);
 	FMassRuntimePipeline& Pipeline = (*HandlersContainer).FindChecked(&FragmentType);
 
 	UE::Mass::Executor::RunProcessorsView(Pipeline.Processors, ProcessingContext, &ChunkCollection);
 }
 
-void FMassObserverManager::AddObserverInstance(const UScriptStruct& ItemType, const FMassObservedOperation Operation, UMassProcessor& ObserverProcessor)
+void FMassObserverManager::AddObserverInstance(const UScriptStruct& FragmentOrTagType, const EMassObservedOperation Operation, UMassProcessor& ObserverProcessor)
 {
-	checkSlow(ItemType.IsChildOf(FMassFragment::StaticStruct()) || ItemType.IsChildOf(FMassTag::StaticStruct()));
+	checkSlow(FragmentOrTagType.IsChildOf(FMassFragment::StaticStruct()) || FragmentOrTagType.IsChildOf(FMassTag::StaticStruct()));
 
 	FMassRuntimePipeline* Pipeline = nullptr;
 
-	if (ItemType.IsChildOf(FMassFragment::StaticStruct()))
+	if (FragmentOrTagType.IsChildOf(FMassFragment::StaticStruct()))
 	{
-		Pipeline = &(*FragmentObservers[(uint8)Operation]).FindOrAdd(&ItemType);
-		ObservedFragments[(uint8)Operation].Add(ItemType);
+		Pipeline = &(*FragmentObservers[(uint8)Operation]).FindOrAdd(&FragmentOrTagType);
+		ObservedFragments[(uint8)Operation].Add(FragmentOrTagType);
 	}
 	else
 	{
-		Pipeline = &(*TagObservers[(uint8)Operation]).FindOrAdd(&ItemType);
-		ObservedTags[(uint8)Operation].Add(ItemType);
+		Pipeline = &(*TagObservers[(uint8)Operation]).FindOrAdd(&FragmentOrTagType);
+		ObservedTags[(uint8)Operation].Add(FragmentOrTagType);
 	}
 	Pipeline->AppendProcessor(ObserverProcessor);
 
diff --git a/Engine/Plugins/Runtime/MassEntity/Source/MassEntity/Private/MassObserverRegistry.cpp b/Engine/Plugins/Runtime/MassEntity/Source/MassEntity/Private/MassObserverRegistry.cpp
--- a/Engine/Plugins/Runtime/MassEntity/Source/MassEntity/Private/MassObserverRegistry.cpp
+++ b/Engine/Plugins/Runtime/MassEntity/Source/MassEntity/Private/MassObserverRegistry.cpp
@@ -11,38 +11,18 @@ UMassObserverRegistry::UMassObserverRegistry()
 	check(HasAnyFlags(RF_ClassDefaultObject));
 }
 
-void UMassObserverRegistry::RegisterFragmentAddedObserver(const UScriptStruct& FragmentType, TSubclassOf<UMassProcessor> ObserverClass)
+void UMassObserverRegistry::RegisterObserver(const UScriptStruct& ObservedType, const EMassObservedOperation Operation, TSubclassOf<UMassProcessor> ObserverClass)
 {
 	check(ObserverClass);
-	FragmentInitializersMap.FindOrAdd(&FragmentType).ClassCollection.AddUnique(ObserverClass);
+	checkSlow(ObservedType.IsChildOf(FMassFragment::StaticStruct()) || ObservedType.IsChildOf(FMassTag::StaticStruct()));
+
+	if (ObservedType.IsChildOf(FMassFragment::StaticStruct()))
+	{
+		(*FragmentObservers[(uint8)Operation]).FindOrAdd(&ObservedType).ClassCollection.AddUnique(ObserverClass);
+	}
+	else
+	{
+		(*TagObservers[(uint8)Operation]).FindOrAdd(&ObservedType).ClassCollection.AddUnique(ObserverClass);
+	}
 }
 
-void UMassObserverRegistry::RegisterFragmentRemovedObserver(const UScriptStruct& FragmentType, TSubclassOf<UMassProcessor> ObserverClass)
-{
-	check(ObserverClass);
-	FragmentDeinitializersMap.FindOrAdd(&FragmentType).ClassCollection.AddUnique(ObserverClass);
-}
-
-void UMassObserverRegistry::RegisterTagAddedObserver(const UScriptStruct& FragmentType, TSubclassOf<UMassProcessor> ObserverClass)
-{
-	check(ObserverClass);
-	TagAddedObserversMap.FindOrAdd(&FragmentType).ClassCollection.AddUnique(ObserverClass);
-}
-
-void UMassObserverRegistry::RegisterTagRemovedObserver(const UScriptStruct& FragmentType, TSubclassOf<UMassProcessor> ObserverClass)
-{
-	check(ObserverClass);
-	TagRemovedObserversMap.FindOrAdd(&FragmentType).ClassCollection.AddUnique(ObserverClass);
-}
-
-TConstArrayView<TSubclassOf<UMassProcessor>> UMassObserverRegistry::GetFragmentInitializers(const UScriptStruct& FragmentType) const
-{
-	const FMassProcessorClassCollection* InitializersCollection = FragmentInitializersMap.Find(&FragmentType);
-	return InitializersCollection ? InitializersCollection->ClassCollection : TConstArrayView<TSubclassOf<UMassProcessor>>();
-}
-
-TConstArrayView<TSubclassOf<UMassProcessor>> UMassObserverRegistry::GetFragmentDeinitializers(const UScriptStruct& FragmentType) const
-{
-	const FMassProcessorClassCollection* DeinitializersCollection = FragmentDeinitializersMap.Find(&FragmentType);
-	return DeinitializersCollection ? DeinitializersCollection->ClassCollection : TConstArrayView<TSubclassOf<UMassProcessor>>();
-}
diff --git a/Engine/Plugins/Runtime/MassEntity/Source/MassEntity/Public/MassCommandBuffer.h b/Engine/Plugins/Runtime/MassEntity/Source/MassEntity/Public/MassCommandBuffer.h
--- a/Engine/Plugins/Runtime/MassEntity/Source/MassEntity/Public/MassCommandBuffer.h
+++ b/Engine/Plugins/Runtime/MassEntity/Source/MassEntity/Public/MassCommandBuffer.h
@@ -23,12 +23,12 @@ namespace ECommandBufferOperationType
 
 struct FMassObservedTypeCollection
 {
-	void Add(const UScriptStruct& TypeType, FMassEntityHandle Entity)
+	void Add(const UScriptStruct& Type, FMassEntityHandle Entity)
 	{
-		if (&TypeType != LastAddedType)
+		if (&Type != LastAddedType)
 		{
-			LastAddedType = &TypeType;
-			LastUsedCollection = &Types.FindOrAdd(&TypeType);
+			LastAddedType = &Type;
+			LastUsedCollection = &Types.FindOrAdd(&Type);
 		}
 		CA_ASSUME(LastUsedCollection);
 		LastUsedCollection->Add(Entity);
@@ -41,6 +41,8 @@ struct FMassObservedTypeCollection
 		Types.Reset();
 	}
 
+	void Append(const FMassObservedTypeCollection& Other);
+
 	const TMap<const UScriptStruct*, TArray<FMassEntityHandle>>& GetTypes() const 
 	{ 
 		return Types; 
@@ -55,30 +57,42 @@ private:
 struct FMassCommandsObservedTypes
 {
 	void Reset();
-	void FragmentAdded(const UScriptStruct* TypeType, FMassEntityHandle Entity)
+	void FragmentAdded(const UScriptStruct* Type, FMassEntityHandle Entity)
+	{
+		check(Type);
+		Fragments[(uint8)EMassObservedOperation::Add].Add(*Type, Entity);
+	}
+	void FragmentRemoved(const UScriptStruct* Type, FMassEntityHandle Entity)
 	{
-		check(TypeType);
-		FragmentsToAdd.Add(*TypeType, Entity);
+		check(Type);
+		Fragments[(uint8)EMassObservedOperation::Remove].Add(*Type, Entity);
 	}
-	void FragmentRemoved(const UScriptStruct* TypeType, FMassEntityHandle Entity)
+	void TagAdded(const UScriptStruct* Type, FMassEntityHandle Entity)
 	{
-		check(TypeType);
-		FragmentsToRemove.Add(*TypeType, Entity);
+		check(Type);
+		Tags[(uint8)EMassObservedOperation::Add].Add(*Type, Entity);
+	}
+	void TagRemoved(const UScriptStruct* Type, FMassEntityHandle Entity)
+	{
+		check(Type);
+		Tags[(uint8)EMassObservedOperation::Remove].Add(*Type, Entity);
 	}
 
-	const TMap<const UScriptStruct*, TArray<FMassEntityHandle>>& GetFragmentsToAdd() const
+	const TMap<const UScriptStruct*, TArray<FMassEntityHandle>>& GetObservedFragments(const EMassObservedOperation Operation) const
 	{
-		return FragmentsToAdd.GetTypes();
+		return Fragments[(uint8)Operation].GetTypes();
 	}
 
-	const TMap<const UScriptStruct*, TArray<FMassEntityHandle>>& GetFragmentsToRemove() const
+	const TMap<const UScriptStruct*, TArray<FMassEntityHandle>>& GetObservedTags(const EMassObservedOperation Operation) const
 	{
-		return FragmentsToRemove.GetTypes();
+		return Tags[(uint8)Operation].GetTypes();
 	}
 
+	void Append(const FMassCommandsObservedTypes& Other);
+
 protected:
-	FMassObservedTypeCollection FragmentsToAdd;
-	FMassObservedTypeCollection FragmentsToRemove;
+	FMassObservedTypeCollection Fragments[(uint8)EMassObservedOperation::MAX];
+	FMassObservedTypeCollection Tags[(uint8)EMassObservedOperation::MAX];
 };
 
 
@@ -396,7 +410,7 @@ struct MASSENTITY_API FCommandAddTag: public FCommandBufferEntryBase
 
 	enum
 	{
-		Type = ECommandBufferOperationType::None
+		Type = ECommandBufferOperationType::Add
 	};
 
 	FCommandAddTag() = default;
@@ -405,6 +419,11 @@ struct MASSENTITY_API FCommandAddTag: public FCommandBufferEntryBase
 		, StructParam(InStruct)
 	{}
 
+	void AppendAffectedEntitiesPerType(FMassCommandsObservedTypes& ObservedTypes)
+	{
+		ObservedTypes.TagAdded(StructParam, TargetEntity);
+	}
+
 protected:
 	virtual void Execute(UMassEntitySubsystem& System) const override;
 
@@ -421,7 +440,7 @@ struct MASSENTITY_API FCommandRemoveTag : public FCommandBufferEntryBase
 
 	enum
 	{
-		Type = ECommandBufferOperationType::None
+		Type = ECommandBufferOperationType::Remove
 	};
 
 	FCommandRemoveTag() = default;
@@ -430,6 +449,11 @@ struct MASSENTITY_API FCommandRemoveTag : public FCommandBufferEntryBase
 		, StructParam(InStruct)
 	{}
 
+	void AppendAffectedEntitiesPerType(FMassCommandsObservedTypes& ObservedTypes)
+	{
+		ObservedTypes.TagRemoved(StructParam, TargetEntity);
+	}
+
 protected:
 	virtual void Execute(UMassEntitySubsystem& System) const override;
 
@@ -447,7 +471,7 @@ struct MASSENTITY_API FCommandSwapTags : public FCommandBufferEntryBase
 
 	enum
 	{
-		Type = ECommandBufferOperationType::None //ECommandBufferOperationType::Add | ECommandBufferOperationType::Remove
+		Type = ECommandBufferOperationType::Add | ECommandBufferOperationType::Remove
 	};
 
 	FCommandSwapTags() = default;
@@ -460,6 +484,12 @@ struct MASSENTITY_API FCommandSwapTags : public FCommandBufferEntryBase
 		checkf((InNewTagType == nullptr) || InNewTagType->IsChildOf(FMassTag::StaticStruct()), TEXT("FCommandSwapTags works only with tags while '%s' is not one."), *GetPathNameSafe(InNewTagType));
 	}
 
+	void AppendAffectedEntitiesPerType(FMassCommandsObservedTypes& ObservedTypes)
+	{
+		ObservedTypes.TagRemoved(OldTagType, TargetEntity);
+		ObservedTypes.TagAdded(NewTagType, TargetEntity);
+	}
+
 protected:
 	virtual void Execute(UMassEntitySubsystem& System) const override;
 
@@ -590,6 +620,7 @@ public:
 	void MoveAppend(FMassCommandBuffer& InOutOther);
 
 	bool HasPendingCommands() const { return PendingCommands.Num() > 0 || EntitiesToDestroy.Num() > 0; }
+	bool IsFlushing() const { return bIsFlushing; }
 
 private:
 	FInstancedStructStream PendingCommands;
@@ -599,4 +630,7 @@ private:
 	TArray<FMassEntityHandle> EntitiesToDestroy;
 
 	FMassCommandsObservedTypes ObservedTypes;
+
+	/** Indicates that this specific MassCommandBuffer is currently flushing its contents */
+	bool bIsFlushing = false;
 };
diff --git a/Engine/Plugins/Runtime/MassEntity/Source/MassEntity/Public/MassObserverManager.h b/Engine/Plugins/Runtime/MassEntity/Source/MassEntity/Public/MassObserverManager.h
--- a/Engine/Plugins/Runtime/MassEntity/Source/MassEntity/Public/MassObserverManager.h
+++ b/Engine/Plugins/Runtime/MassEntity/Source/MassEntity/Public/MassObserverManager.h
@@ -11,21 +11,11 @@ struct FMassArchetypeSubChunks;
 class UMassProcessor;
 struct FMassProcessingContext;
 
-enum class FMassObservedOperation : uint8
-{
-	Add,
-	Remove,
-	// @todo Keeping this here as a indication of design intent. For now we handle entity destruction like removal, but 
-	// there might be coputationaly expensive cases where we might want to avoid for soon-to-be-dead entities. 
-	// Destroy,
-	MAX
-};
-
 /** 
- * A wrapper type for a TMap to support having array-of-maps UPROPERTY members inFMassObserverManager
+ * A wrapper type for a TMap to support having array-of-maps UPROPERTY members in FMassObserverManager
  */
 USTRUCT()
-struct FMassItemObserversMap
+struct FMassObserversMap
 {
 	GENERATED_BODY()
 
@@ -53,27 +43,29 @@ struct MASSENTITY_API FMassObserverManager
 public:
 	FMassObserverManager();	
 
-	const FMassFragmentBitSet& GetObservedAddFragmentsBitSet() const { return ObservedFragments[(uint8)FMassObservedOperation::Add]; }
-	const FMassFragmentBitSet& GetObservedRemoveFragmentsBitSet() const { return ObservedFragments[(uint8)FMassObservedOperation::Remove]; }
-	const FMassTagBitSet& GetObservedAddTagsBitSet() const { return ObservedTags[(uint8)FMassObservedOperation::Add]; }
-	const FMassTagBitSet& GetObservedRemoveTagsBitSet() const { return ObservedTags[(uint8)FMassObservedOperation::Remove]; }
-
-	bool HasOnAddedObserversForBitSet(const FMassFragmentBitSet& InQueriedBitSet) const { return ObservedFragments[(uint8)FMassObservedOperation::Add].HasAny(InQueriedBitSet); }
-	bool HasOnAddedObserversForBitSet(const FMassTagBitSet& InQueriedBitSet) const { return ObservedTags[(uint8)FMassObservedOperation::Add].HasAny(InQueriedBitSet); }
-	bool HasOnRemovedObserversForBitSet(const FMassFragmentBitSet& InQueriedBitSet) const { return ObservedFragments[(uint8)FMassObservedOperation::Remove].HasAny(InQueriedBitSet); }
-	bool HasOnRemovedObserversForBitSet(const FMassTagBitSet& InQueriedBitSet) const { return ObservedTags[(uint8)FMassObservedOperation::Remove].HasAny(InQueriedBitSet); }
+	const FMassFragmentBitSet* GetObservedFragmentBitSets() const { return ObservedFragments; }
+	const FMassFragmentBitSet& GetObservedFragmentsBitSet(const EMassObservedOperation Operation) const 
+	{ 
+		return ObservedFragments[(uint8)Operation]; 
+	}
 
-	bool HasObserversForBitSet(const FMassFragmentBitSet& InQueriedBitSet, const FMassObservedOperation Operation) const
+	const FMassTagBitSet* GetObservedTagBitSets() const { return ObservedTags; }
+	const FMassTagBitSet& GetObservedTagsBitSet(const EMassObservedOperation Operation) const 
+	{ 
+		return ObservedTags[(uint8)Operation]; 
+	}
+	
+	bool HasObserversForBitSet(const FMassFragmentBitSet& InQueriedBitSet, const EMassObservedOperation Operation) const
 	{
 		return ObservedFragments[(uint8)Operation].HasAny(InQueriedBitSet);
 	}
 
-	bool HasObserversForBitSet(const FMassTagBitSet& InQueriedBitSet, const FMassObservedOperation Operation) const
+	bool HasObserversForBitSet(const FMassTagBitSet& InQueriedBitSet, const EMassObservedOperation Operation) const
 	{
 		return ObservedTags[(uint8)Operation].HasAny(InQueriedBitSet);
 	}
 
-	bool HasObserversForComposition(const FMassArchetypeCompositionDescriptor& Composition, const FMassObservedOperation Operation) const
+	bool HasObserversForComposition(const FMassArchetypeCompositionDescriptor& Composition, const EMassObservedOperation Operation) const
 	{
 		return HasObserversForBitSet(Composition.Fragments, Operation) || HasObserversForBitSet(Composition.Tags, Operation);
 	}
@@ -84,33 +76,31 @@ public:
 	bool OnPreEntitiesDestroyed(const FMassArchetypeSubChunks& ChunkCollection);
 	bool OnPreEntitiesDestroyed(FMassProcessingContext& ProcessingContext, const FMassArchetypeSubChunks& ChunkCollection);
 
-	bool OnCompositionChanged(const FMassEntityHandle Entity, const FMassArchetypeCompositionDescriptor& CompositionDelta, const FMassObservedOperation Operation);
+	bool OnCompositionChanged(const FMassEntityHandle Entity, const FMassArchetypeCompositionDescriptor& CompositionDelta, const EMassObservedOperation Operation);
 	bool OnPostCompositionAdded(const FMassEntityHandle Entity, const FMassArchetypeCompositionDescriptor& Composition)
 	{
-		return OnCompositionChanged(Entity, Composition, FMassObservedOperation::Add);
+		return OnCompositionChanged(Entity, Composition, EMassObservedOperation::Add);
 	}
 	bool OnPreCompositionRemoved(const FMassEntityHandle Entity, const FMassArchetypeCompositionDescriptor& Composition)
 	{
-		return OnCompositionChanged(Entity, Composition, FMassObservedOperation::Remove);
+		return OnCompositionChanged(Entity, Composition, EMassObservedOperation::Remove);
 	}
 
-	// @todo caller will need to recreate ChunkCollection if this call is done after modifications to entities, so it 
-	// would be beneficial to ask the system "if there's anything observing the change" before actually reconstructing 
-	// the chunk collection. Unless the new ChunkCollection gets created as a side-product of the batched modification 
-	// operation
-	bool OnCompositionChanged(const FMassArchetypeSubChunks& ChunkCollection, const FMassArchetypeCompositionDescriptor& Composition, const FMassObservedOperation Operation, FMassProcessingContext* InProcessingContext = nullptr);
+	bool OnCompositionChanged(const FMassArchetypeSubChunks& ChunkCollection, const FMassArchetypeCompositionDescriptor& Composition, const EMassObservedOperation Operation, FMassProcessingContext* InProcessingContext = nullptr);
 
-	void OnPostItemAdded(const UScriptStruct& ItemType, const FMassArchetypeSubChunks& ChunkCollection)
+	void OnPostFragmentOrTagAdded(const UScriptStruct& FragmentOrTagType, const FMassArchetypeSubChunks& ChunkCollection)
 	{
-		OnSingleItemOperation(ItemType, ChunkCollection, FMassObservedOperation::Add);
+		OnFragmentOrTagOperation(FragmentOrTagType, ChunkCollection, EMassObservedOperation::Add);
 	}
-	void OnPreItemRemoved(const UScriptStruct& ItemType, const FMassArchetypeSubChunks& ChunkCollection)
+	void OnPreFragmentOrTagRemoved(const UScriptStruct& FragmentOrTagType, const FMassArchetypeSubChunks& ChunkCollection)
 	{
-		OnSingleItemOperation(ItemType, ChunkCollection, FMassObservedOperation::Remove);
+		OnFragmentOrTagOperation(FragmentOrTagType, ChunkCollection, EMassObservedOperation::Remove);
 	}
-	void OnSingleItemOperation(const UScriptStruct& ItemType, const FMassArchetypeSubChunks& ChunkCollection, const FMassObservedOperation Operation);
+	// @todo I don't love this name. Alternatively could be OnSingleFragmentOrTagOperation. Long term we'll switch to
+	// OnSingleFragmentOperation and advertise that Tags are a type of Fragment as well (conceptually).
+	void OnFragmentOrTagOperation(const UScriptStruct& FragmentOrTagType, const FMassArchetypeSubChunks& ChunkCollection, const EMassObservedOperation Operation);
 
-	void AddObserverInstance(const UScriptStruct& ItemType, const FMassObservedOperation Operation, UMassProcessor& ObserverProcessor);
+	void AddObserverInstance(const UScriptStruct& FragmentOrTagType, const EMassObservedOperation Operation, UMassProcessor& ObserverProcessor);
 
 protected:
 	friend UMassEntitySubsystem;
@@ -118,17 +108,17 @@ protected:
 
 	void Initialize();
 	void HandleFragmentsImpl(FMassProcessingContext& ProcessingContext, const FMassArchetypeSubChunks& ChunkCollection
-		, TArrayView<const UScriptStruct*> ObservedTypes, FMassItemObserversMap& HandlersContainer);
-	void HandleSingleItemImpl(const UScriptStruct& FragmentType, const FMassArchetypeSubChunks& ChunkCollection, FMassItemObserversMap& HandlersContainer);
+		, TArrayView<const UScriptStruct*> ObservedTypes, FMassObserversMap& HandlersContainer);
+	void HandleSingleEntityImpl(const UScriptStruct& FragmentType, const FMassArchetypeSubChunks& ChunkCollection, FMassObserversMap& HandlersContainer);
 
-	FMassFragmentBitSet ObservedFragments[(uint8)FMassObservedOperation::MAX];
-	FMassTagBitSet ObservedTags[(uint8)FMassObservedOperation::MAX];
+	FMassFragmentBitSet ObservedFragments[(uint8)EMassObservedOperation::MAX];
+	FMassTagBitSet ObservedTags[(uint8)EMassObservedOperation::MAX];
 
 	UPROPERTY()
-	FMassItemObserversMap FragmentObservers[(uint8)FMassObservedOperation::MAX];
+	FMassObserversMap FragmentObservers[(uint8)EMassObservedOperation::MAX];
 
 	UPROPERTY()
-	FMassItemObserversMap TagObservers[(uint8)FMassObservedOperation::MAX];
+	FMassObserversMap TagObservers[(uint8)EMassObservedOperation::MAX];
 
 	/** 
 	 * The owning EntitySubsystem. No need for it to be a UPROPERTY since by design we don't support creation of 
diff --git a/Engine/Plugins/Runtime/MassEntity/Source/MassEntity/Public/MassObserverRegistry.h b/Engine/Plugins/Runtime/MassEntity/Source/MassEntity/Public/MassObserverRegistry.h
--- a/Engine/Plugins/Runtime/MassEntity/Source/MassEntity/Public/MassObserverRegistry.h
+++ b/Engine/Plugins/Runtime/MassEntity/Source/MassEntity/Public/MassObserverRegistry.h
@@ -9,6 +9,9 @@
 
 struct FMassObserverManager;
 
+/**
+ * A wrapper type for a TArray to support having map-of-arrays UPROPERTY members in FMassEntityObserverClassesMap
+ */
 USTRUCT()
 struct FMassProcessorClassCollection
 {
@@ -18,6 +21,29 @@ struct FMassProcessorClassCollection
 	TArray<TSubclassOf<UMassProcessor>> ClassCollection;
 };
 
+/**
+ * A wrapper type for a TMap to support having array-of-maps UPROPERTY members in UMassObserverRegistry
+ */
+USTRUCT()
+struct FMassEntityObserverClassesMap
+{
+	GENERATED_BODY()
+
+	/** a helper accessor simplifying access while still keeping Container private */
+	const TMap<const UScriptStruct*, FMassProcessorClassCollection>& operator*() const
+	{
+		return Container;
+	}
+
+	TMap<const UScriptStruct*, FMassProcessorClassCollection>& operator*()
+	{
+		return Container;
+	}
+
+private:
+	UPROPERTY()
+	TMap<const UScriptStruct*, FMassProcessorClassCollection> Container;
+};
 
 UCLASS()
 class MASSENTITY_API UMassObserverRegistry : public UObject
@@ -30,42 +56,14 @@ public:
 	static UMassObserverRegistry& GetMutable() { return *GetMutableDefault<UMassObserverRegistry>(); }
 	static const UMassObserverRegistry& Get() { return *GetDefault<UMassObserverRegistry>(); }
 
-	template<typename TFragment>
-	static void RegisterInitializer(const UMassProcessor& CDO)
-	{
-		ensureMsgf(CDO.HasAnyFlags(RF_ClassDefaultObject) && CDO.GetClass()->HasAnyClassFlags(CLASS_Abstract) == false
-			, TEXT("Only CDOs are expected to be registered with UMassObserverRegistry::RegisterInitializer"));
-		GetMutable().RegisterFragmentAddedObserver(*TFragment::StaticStruct(), CDO.GetClass());
-	}
-
-	template<typename TFragment>
-	static void RegisterDeinitializer(const UMassProcessor& CDO)
-	{
-		ensureMsgf(CDO.HasAnyFlags(RF_ClassDefaultObject) && CDO.GetClass()->HasAnyClassFlags(CLASS_Abstract) == false
-			, TEXT("Only CDOs are expected to be registered with UMassObserverRegistry::RegisterDeinitializer"));
-		GetMutable().RegisterFragmentRemovedObserver(*TFragment::StaticStruct(), CDO.GetClass());
-	}
-
-	void RegisterFragmentAddedObserver(const UScriptStruct& FragmentType, TSubclassOf<UMassProcessor> FragmentInitializerClass);
-	void RegisterFragmentRemovedObserver(const UScriptStruct& FragmentType, TSubclassOf<UMassProcessor> FragmentDeinitializerClass);
-	void RegisterTagAddedObserver(const UScriptStruct& FragmentType, TSubclassOf<UMassProcessor> FragmentInitializerClass);
-	void RegisterTagRemovedObserver(const UScriptStruct& FragmentType, TSubclassOf<UMassProcessor> FragmentDeinitializerClass);
+	void RegisterObserver(const UScriptStruct& ObservedType, const EMassObservedOperation Operation, TSubclassOf<UMassProcessor> ObserverClass);
 
 protected:
 	friend FMassObserverManager;
 
-	TConstArrayView<TSubclassOf<UMassProcessor>> GetFragmentInitializers(const UScriptStruct& FragmentType) const;
-	TConstArrayView<TSubclassOf<UMassProcessor>> GetFragmentDeinitializers(const UScriptStruct& FragmentType) const;
-
 	UPROPERTY()
-	TMap<const UScriptStruct*, FMassProcessorClassCollection> FragmentInitializersMap;
+	FMassEntityObserverClassesMap FragmentObservers[(uint8)EMassObservedOperation::MAX];
 
 	UPROPERTY()
-	TMap<const UScriptStruct*, FMassProcessorClassCollection> FragmentDeinitializersMap;
-
-	UPROPERTY()
-	TMap<const UScriptStruct*, FMassProcessorClassCollection> TagAddedObserversMap;
-	
-	UPROPERTY()
-	TMap<const UScriptStruct*, FMassProcessorClassCollection> TagRemovedObserversMap;
+	FMassEntityObserverClassesMap TagObservers[(uint8)EMassObservedOperation::MAX];
 };
diff --git a/Engine/Plugins/Runtime/MassEntity/Source/MassEntityTestSuite/Private/MassObserverTest.cpp b/Engine/Plugins/Runtime/MassEntity/Source/MassEntityTestSuite/Private/MassObserverTest.cpp
--- a/Engine/Plugins/Runtime/MassEntity/Source/MassEntityTestSuite/Private/MassObserverTest.cpp
+++ b/Engine/Plugins/Runtime/MassEntity/Source/MassEntityTestSuite/Private/MassObserverTest.cpp
@@ -26,7 +26,7 @@ struct FTagBaseOperation : FEntityTestBase
 {
 	TArray<FMassEntityHandle> AffectedEntities;
 	UMassTestProcessorBase* TagObserver = nullptr;
-	FMassObservedOperation OperationObserved = FMassObservedOperation::MAX;
+	EMassObservedOperation OperationObserved = EMassObservedOperation::MAX;
 	TArray<FMassEntityHandle> EntitiesInt;
 	TArray<FMassEntityHandle> EntitiesIntsFloat;
 	TArray<FMassEntityHandle> ExpectedEntities;
@@ -85,7 +85,7 @@ struct FTagBaseOperation : FEntityTestBase
 
 struct FTag_SingleEntitySingleArchetypeAdd : FTagBaseOperation
 {
-	FTag_SingleEntitySingleArchetypeAdd() { OperationObserved = FMassObservedOperation::Add; }
+	FTag_SingleEntitySingleArchetypeAdd() { OperationObserved = EMassObservedOperation::Add; }
 	virtual bool PerformOperation() override 
 	{
 		ExpectedEntities = { EntitiesInt[1] };
@@ -97,7 +97,7 @@ IMPLEMENT_AI_INSTANT_TEST(FTag_SingleEntitySingleArchetypeAdd, "System.Mass.Obse
 
 struct FTag_SingleEntitySingleArchetypeRemove : FTagBaseOperation
 {
-	FTag_SingleEntitySingleArchetypeRemove() { OperationObserved = FMassObservedOperation::Remove; }
+	FTag_SingleEntitySingleArchetypeRemove() { OperationObserved = EMassObservedOperation::Remove; }
 	virtual bool PerformOperation() override
 	{
 		ExpectedEntities = { EntitiesInt[1] };
@@ -114,7 +114,7 @@ IMPLEMENT_AI_INSTANT_TEST(FTag_SingleEntitySingleArchetypeRemove, "System.Mass.O
 
 struct FTag_SingleEntitySingleArchetypeDestroy : FTagBaseOperation
 {
-	FTag_SingleEntitySingleArchetypeDestroy() { OperationObserved = FMassObservedOperation::Remove; }
+	FTag_SingleEntitySingleArchetypeDestroy() { OperationObserved = EMassObservedOperation::Remove; }
 	virtual bool PerformOperation() override
 	{
 		ExpectedEntities = { EntitiesInt[1] };
@@ -130,7 +130,7 @@ IMPLEMENT_AI_INSTANT_TEST(FTag_SingleEntitySingleArchetypeDestroy, "System.Mass.
 
 struct FTag_MultipleArchetypeAdd : FTagBaseOperation
 {
-	FTag_MultipleArchetypeAdd() { OperationObserved = FMassObservedOperation::Add; }
+	FTag_MultipleArchetypeAdd() { OperationObserved = EMassObservedOperation::Add; }
 
 	virtual bool PerformOperation() override
 	{
@@ -146,7 +146,7 @@ IMPLEMENT_AI_INSTANT_TEST(FTag_MultipleArchetypeAdd, "System.Mass.Observer.Tag.M
 
 struct FTag_MultipleArchetypeRemove : FTagBaseOperation
 {
-	FTag_MultipleArchetypeRemove() { OperationObserved = FMassObservedOperation::Remove; }
+	FTag_MultipleArchetypeRemove() { OperationObserved = EMassObservedOperation::Remove; }
 
 	virtual bool PerformOperation() override
 	{
@@ -169,7 +169,7 @@ IMPLEMENT_AI_INSTANT_TEST(FTag_MultipleArchetypeRemove, "System.Mass.Observer.Ta
 
 struct FTag_MultipleArchetypeDestroy : FTagBaseOperation
 {
-	FTag_MultipleArchetypeDestroy() { OperationObserved = FMassObservedOperation::Remove; }
+	FTag_MultipleArchetypeDestroy() { OperationObserved = EMassObservedOperation::Remove; }
 
 	virtual bool PerformOperation() override
 	{
@@ -190,6 +190,29 @@ struct FTag_MultipleArchetypeDestroy : FTagBaseOperation
 };
 IMPLEMENT_AI_INSTANT_TEST(FTag_MultipleArchetypeDestroy, "System.Mass.Observer.Tag.MultipleArchetypesDestroy");
 
+struct FTag_MultipleArchetypeSwap : FTagBaseOperation
+{
+	FTag_MultipleArchetypeSwap() { OperationObserved = EMassObservedOperation::Remove; }
+
+	virtual bool PerformOperation() override
+	{
+		ExpectedEntities = { EntitiesIntsFloat[1], EntitiesInt[0], EntitiesInt[2] };
+		for (const FMassEntityHandle& ModifiedEntity : ExpectedEntities)
+		{
+			EntitySubsystem->Defer().AddTag<FTestTag_A>(ModifiedEntity);
+		}
+		EntitySubsystem->Defer().ReplayBufferAgainstSystem(EntitySubsystem);
+		// since we're only observing tag removal we don't expect AffectedEntities to contain any data at this point
+		AITEST_EQUAL(TEXT("Tag addition is not being observed and is not expected to produce results yet"), AffectedEntities.Num(), 0);
+		for (const FMassEntityHandle& ModifiedEntity : ExpectedEntities)
+		{
+			EntitySubsystem->Defer().PushCommand(FCommandSwapTags(ModifiedEntity, FTestTag_A::StaticStruct(), FTestTag_B::StaticStruct()));
+		}
+		return true;
+	}
+};
+IMPLEMENT_AI_INSTANT_TEST(FTag_MultipleArchetypeSwap, "System.Mass.Observer.Tag.MultipleArchetypesSwap");
+
 } // FMassObserverTest
 
 PRAGMA_ENABLE_OPTIMIZATION
