Removed UMassFragmentInitializer and UMassFragmentDeinitializer in favor of their more configurable parrent, UMassObserverProcessor

diff --git a/Engine/Plugins/Runtime/MassEntity/Source/MassEntity/Private/MassObserverProcessor.cpp b/Engine/Plugins/Runtime/MassEntity/Source/MassEntity/Private/MassObserverProcessor.cpp
--- a/Engine/Plugins/Runtime/MassEntity/Source/MassEntity/Private/MassObserverProcessor.cpp
+++ b/Engine/Plugins/Runtime/MassEntity/Source/MassEntity/Private/MassObserverProcessor.cpp
@@ -2,6 +2,7 @@
 
 #include "MassObserverProcessor.h"
 #include "MassObserverRegistry.h"
+#include "MassEntityDebug.h"
 
 //----------------------------------------------------------------------//
 // UMassObserverProcessor
@@ -18,26 +19,26 @@ void UMassObserverProcessor::PostInitProperties()
 {
 	Super::PostInitProperties();
 
-	if (HasAnyFlags(RF_ClassDefaultObject) && GetClass()->HasAnyClassFlags(CLASS_Abstract) == false && FragmentType != nullptr)
+	UClass* MyClass = GetClass();
+	CA_ASSUME(MyClass);
+
+	if (HasAnyFlags(RF_ClassDefaultObject) && MyClass->HasAnyClassFlags(CLASS_Abstract) == false)
 	{
-		Register();
+		if (ensure(ObservedType != nullptr && Operation != EMassObservedOperation::MAX))
+		{
+			Register();
+		}
+		else
+		{
+			UE_LOG(LogMass, Error, TEXT("%s attempting to register %s while it\'s misconfigured, Type: %s, Operation: %s")
+				, ANSI_TO_TCHAR(__FUNCTION__), *MyClass->GetName(), *GetNameSafe(ObservedType), *UEnum::GetValueAsString(Operation));
+		}
 	}
 }
 
-//----------------------------------------------------------------------//
-//  UMassFragmentInitializer
-//----------------------------------------------------------------------//
-void UMassFragmentInitializer::Register()
+void UMassObserverProcessor::Register()
 {
-	check(FragmentType);
-	UMassObserverRegistry::GetMutable().RegisterFragmentAddedObserver(*FragmentType, GetClass());
+	check(ObservedType);
+	UMassObserverRegistry::GetMutable().RegisterObserver(*ObservedType, Operation, GetClass());
 }
 
-//----------------------------------------------------------------------//
-//  UMassFragmentDeinitializer
-//----------------------------------------------------------------------//
-void UMassFragmentDeinitializer::Register()
-{
-	check(FragmentType);
-	UMassObserverRegistry::GetMutable().RegisterFragmentRemovedObserver(*FragmentType, GetClass());
-}
diff --git a/Engine/Plugins/Runtime/MassEntity/Source/MassEntity/Public/MassEntityTypes.h b/Engine/Plugins/Runtime/MassEntity/Source/MassEntity/Public/MassEntityTypes.h
--- a/Engine/Plugins/Runtime/MassEntity/Source/MassEntity/Public/MassEntityTypes.h
+++ b/Engine/Plugins/Runtime/MassEntity/Source/MassEntity/Public/MassEntityTypes.h
@@ -316,4 +316,17 @@ protected:
 	
 	TArray<FConstSharedStruct> ConstSharedFragments;
 	TArray<FSharedStruct> SharedFragments;
-};
\ No newline at end of file
+};
+
+UENUM()
+enum class EMassObservedOperation : uint8
+{
+	Add,
+	Remove,
+	// @todo Keeping this here as a indication of design intent. For now we handle entity destruction like removal, but 
+	// there might be computationally expensive cases where we might want to avoid for soon-to-be-dead entities. 
+	// Destroy,
+	// @todo another planned supported operation type
+	// Touch,
+	MAX
+};
diff --git a/Engine/Plugins/Runtime/MassEntity/Source/MassEntity/Public/MassObserverProcessor.h b/Engine/Plugins/Runtime/MassEntity/Source/MassEntity/Public/MassObserverProcessor.h
--- a/Engine/Plugins/Runtime/MassEntity/Source/MassEntity/Public/MassObserverProcessor.h
+++ b/Engine/Plugins/Runtime/MassEntity/Source/MassEntity/Public/MassObserverProcessor.h
@@ -19,33 +19,16 @@ public:
 protected:
 	virtual void PostInitProperties() override;
 
-	/** Override to register this processor with the observer registry */
-	virtual void Register() PURE_VIRTUAL(UMassObserverProcessor::Register, );
+	/** 
+	 * By default registers this class as Operation observer of ObservedType. Override to register for multiple 
+	 * operations and/or types 
+	 */
+	virtual void Register();
 
 protected:
-	/** Set in class' constructor and determines for which Fragment type this given UMassObserverProcessor will be used 
-	 *  as the default initializer/deinitializer (as returned by the MassInitializersRegistry). If set to null will
-	 *  require the user to manually add this UMassObserverProcessor as initializer/deinitializer */
+	/** Determines which Fragment or Tag type this given UMassObserverProcessor will be observing */
 	UPROPERTY()
-	UScriptStruct* FragmentType;
-};
+	UScriptStruct* ObservedType = nullptr;
 
-UCLASS(abstract)
-class MASSENTITY_API UMassFragmentInitializer : public UMassObserverProcessor
-{
-	GENERATED_BODY()
-
-protected:
-	/** Adds this FragmentInitializer to the MassInitializersRegistry to be used to initialize fragments of FragmentType */
-	virtual void Register() override;
-};
-
-UCLASS(abstract)
-class MASSENTITY_API UMassFragmentDeinitializer : public UMassObserverProcessor
-{
-	GENERATED_BODY()
-
-protected:
-	/** Adds this FragmentDestructor to the MassInitializersRegistry to be used to deinitialize fragments of FragmentType */
-	virtual void Register() override;
+	EMassObservedOperation Operation = EMassObservedOperation::MAX;
 };
