Mitigated a chance of Mass entities misconfiguration causing a check-fail in command buffer's guts

#rnx

diff --git a/Engine/Plugins/Runtime/MassEntity/Source/MassEntity/Public/MassCommandBuffer.h b/Engine/Plugins/Runtime/MassEntity/Source/MassEntity/Public/MassCommandBuffer.h
--- a/Engine/Plugins/Runtime/MassEntity/Source/MassEntity/Public/MassCommandBuffer.h
+++ b/Engine/Plugins/Runtime/MassEntity/Source/MassEntity/Public/MassCommandBuffer.h
@@ -486,8 +486,14 @@ struct MASSENTITY_API FCommandSwapTags : public FCommandBufferEntryBase
 
 	void AppendAffectedEntitiesPerType(FMassCommandsObservedTypes& ObservedTypes)
 	{
-		ObservedTypes.TagRemoved(OldTagType, TargetEntity);
-		ObservedTypes.TagAdded(NewTagType, TargetEntity);
+		if (OldTagType)
+		{
+			ObservedTypes.TagRemoved(OldTagType, TargetEntity);
+		}
+		if (NewTagType)
+		{
+			ObservedTypes.TagAdded(NewTagType, TargetEntity);
+		}
 	}
 
 protected:
