NoPCH compilation fix

diff --git a/Engine/Plugins/Runtime/MassEntity/Source/MassEntity/Public/MassEntitySubsystem.h b/Engine/Plugins/Runtime/MassEntity/Source/MassEntity/Public/MassEntitySubsystem.h
--- a/Engine/Plugins/Runtime/MassEntity/Source/MassEntity/Public/MassEntitySubsystem.h
+++ b/Engine/Plugins/Runtime/MassEntity/Source/MassEntity/Public/MassEntitySubsystem.h
@@ -9,6 +9,7 @@
 #include "MassEntityQuery.h"
 #include "StructUtilsTypes.h"
 #include "MassObserverManager.h"
+#include "Containers/Queue.h"
 #include "MassEntitySubsystem.generated.h"
 
 
