Non-Unity compilation fix

diff --git a/Engine/Plugins/Runtime/MassEntity/Source/MassEntity/Private/MassEntitySubsystem.cpp b/Engine/Plugins/Runtime/MassEntity/Source/MassEntity/Private/MassEntitySubsystem.cpp
--- a/Engine/Plugins/Runtime/MassEntity/Source/MassEntity/Private/MassEntitySubsystem.cpp
+++ b/Engine/Plugins/Runtime/MassEntity/Source/MassEntity/Private/MassEntitySubsystem.cpp
@@ -6,6 +6,7 @@
 #include "HAL/IConsoleManager.h"
 #include "Engine/World.h"
 #include "UObject/UObjectIterator.h"
+#include "VisualLogger/VisualLogger.h"
 
 const FMassEntityHandle UMassEntitySubsystem::InvalidEntity;
 
