Change to FMassArchetypeSubChunks to support archetype-less entities (i.e. built but not created).

diff --git a/Engine/Plugins/Runtime/MassEntity/Source/MassEntity/Private/MassArchetypeTypes.cpp b/Engine/Plugins/Runtime/MassEntity/Source/MassEntity/Private/MassArchetypeTypes.cpp
--- a/Engine/Plugins/Runtime/MassEntity/Source/MassEntity/Private/MassArchetypeTypes.cpp
+++ b/Engine/Plugins/Runtime/MassEntity/Source/MassEntity/Private/MassArchetypeTypes.cpp
@@ -20,25 +20,36 @@ uint32 GetTypeHash(const FMassArchetypeHandle& Instance)
 FMassArchetypeSubChunks::FMassArchetypeSubChunks(const FMassArchetypeHandle& InArchetype, TConstArrayView<FMassEntityHandle> InEntities, EDuplicatesHandling DuplicatesHandling)
 	: Archetype(InArchetype)
 {
-	check(InArchetype.IsValid());
-
 	if (InEntities.Num() <= 0)
 	{
 		return;
 	}
 
-	const FMassArchetypeData& ArchetypeData = *InArchetype.DataPtr.Get();
-	const int32 NumEntitiesPerChunk = ArchetypeData.GetNumEntitiesPerChunk();
+	const FMassArchetypeData* ArchetypeData = InArchetype.DataPtr.Get();
+	const int32 NumEntitiesPerChunk = ArchetypeData ? ArchetypeData->GetNumEntitiesPerChunk() : MAX_int32;
 
 	// InEntities has a real chance of not being sorted by AbsoluteIndex. We gotta fix that to optimize how we process the data 
 	TArray<int32> TrueIndices;
 	TrueIndices.AddUninitialized(InEntities.Num());
 	int32 i = 0;
-	for (const FMassEntityHandle& Entity : InEntities)
+	if (ArchetypeData)
 	{
-		TrueIndices[i] = ArchetypeData.GetInternalIndexForEntity(Entity.Index);
-		++i;
+		for (const FMassEntityHandle& Entity : InEntities)
+		{
+			TrueIndices[i++] = ArchetypeData->GetInternalIndexForEntity(Entity.Index);
+		}
 	}
+	else
+	{
+		// special case, where we have a bunch of entities that have been built but not assigned an archetype yet.
+		// we use their base index for the sake of sorting here. Will still get some perf benefits and we can keep using
+		// FMassArchetypeSubChunks as the generic batched API wrapper for entities
+		for (const FMassEntityHandle& Entity : InEntities)
+		{
+			TrueIndices[i++] = Entity.Index;
+		}
+	}
+
 	TrueIndices.Sort();
 
 #if DO_GUARD_SLOW
diff --git a/Engine/Plugins/Runtime/MassEntity/Source/MassEntity/Private/MassEntityUtils.cpp b/Engine/Plugins/Runtime/MassEntity/Source/MassEntity/Private/MassEntityUtils.cpp
--- a/Engine/Plugins/Runtime/MassEntity/Source/MassEntity/Private/MassEntityUtils.cpp
+++ b/Engine/Plugins/Runtime/MassEntity/Source/MassEntity/Private/MassEntityUtils.cpp
@@ -50,12 +50,7 @@ void CreateSparseChunks(const UMassEntitySubsystem& EntitySystem, const TConstAr
 
 	for (auto& Pair : ArchetypeToEntities)
 	{
-		// @todo this is a temporary measure to skip entities we've destroyed without the
-		// UMassSpawnerSubsystem::DestroyEntities' caller knowledge. Should be removed once that's addressed.
-		if (Pair.Key.IsValid())
-		{
-			OutChunkCollections.Add(FMassArchetypeSubChunks(Pair.Key, Pair.Value, DuplicatesHandling));
-		}
+		OutChunkCollections.Add(FMassArchetypeSubChunks(Pair.Key, Pair.Value, DuplicatesHandling));
 	}
 }
 
