Fix shared data implementation when specifying sub chunk info to queries.

#rnx

diff --git a/Engine/Plugins/Runtime/MassEntity/Source/MassEntity/Private/MassArchetypeData.cpp b/Engine/Plugins/Runtime/MassEntity/Source/MassEntity/Private/MassArchetypeData.cpp
--- a/Engine/Plugins/Runtime/MassEntity/Source/MassEntity/Private/MassArchetypeData.cpp
+++ b/Engine/Plugins/Runtime/MassEntity/Source/MassEntity/Private/MassArchetypeData.cpp
@@ -475,8 +475,8 @@ void FMassArchetypeData::ExecuteFunction(FMassExecutionContext& RunContext, cons
 	// mz@todo to be removed
 	RunContext.SetCurrentArchetypesTagBitSet(GetTagBitSet());
 
-	BindConstSharedFragmentRequirements(RunContext, RequirementMapping.ChunkFragments);
-	BindSharedFragmentRequirements(RunContext, RequirementMapping.ChunkFragments);
+	BindConstSharedFragmentRequirements(RunContext, RequirementMapping.ConstSharedFragments);
+	BindSharedFragmentRequirements(RunContext, RequirementMapping.SharedFragments);
 
 	for (FMassArchetypeChunkIterator ChunkIterator(SubChunkContainer); ChunkIterator; ++ChunkIterator)
 	{
diff --git a/Engine/Plugins/Runtime/MassEntity/Source/MassEntity/Private/MassEntityQuery.cpp b/Engine/Plugins/Runtime/MassEntity/Source/MassEntity/Private/MassEntityQuery.cpp
--- a/Engine/Plugins/Runtime/MassEntity/Source/MassEntity/Private/MassEntityQuery.cpp
+++ b/Engine/Plugins/Runtime/MassEntity/Source/MassEntity/Private/MassEntityQuery.cpp
@@ -119,7 +119,9 @@ bool FMassEntityQuery::DoesArchetypeMatchRequirements(const FMassArchetypeHandle
 		&& (RequiredAnyTags.IsEmpty() || ArchetypeComposition.Tags.HasAny(RequiredAnyTags))
 		&& ArchetypeComposition.Tags.HasNone(RequiredNoneTags)
 		&& ArchetypeComposition.ChunkFragments.HasAll(RequiredAllChunkFragments)
-		&& ArchetypeComposition.ChunkFragments.HasNone(RequiredNoneChunkFragments);
+		&& ArchetypeComposition.ChunkFragments.HasNone(RequiredNoneChunkFragments)
+		&& ArchetypeComposition.SharedFragments.HasAll(RequiredAllSharedFragments)
+		&& ArchetypeComposition.SharedFragments.HasNone(RequiredNoneSharedFragments);
 }
 
 void FMassEntityQuery::ForEachEntityChunk(const FMassArchetypeSubChunks& Chunks, UMassEntitySubsystem& EntitySubsystem, FMassExecutionContext& ExecutionContext, const FMassExecuteFunction& ExecuteFunction)
