UE::Mass::Core:bBitwiseRelocateFragments got removed and replaced with "true".
Also:
* AddEntityInternal and RemoveEntityInternal's parametrized behavior got extraced out and put at callsites.
* FMemory::Memcopy got replaced with FMemory::Memmove

diff --git a/Engine/Plugins/Runtime/MassEntity/Source/MassEntity/Private/MassArchetypeData.cpp b/Engine/Plugins/Runtime/MassEntity/Source/MassEntity/Private/MassArchetypeData.cpp
--- a/Engine/Plugins/Runtime/MassEntity/Source/MassEntity/Private/MassArchetypeData.cpp
+++ b/Engine/Plugins/Runtime/MassEntity/Source/MassEntity/Private/MassArchetypeData.cpp
@@ -7,10 +7,6 @@
 
 //////////////////////////////////////////////////////////////////////
 // FMassArchetypeData
-namespace UE::Mass::Core 
-{
-	constexpr static bool bBitwiseRelocateFragments = true;
-}// UE::Mass::Core
 
 void FMassArchetypeData::ForEachFragmentType(TFunction< void(const UScriptStruct* /*Fragment*/)> Function) const
 {
@@ -117,10 +113,20 @@ void FMassArchetypeData::InitializeWithSibling(const FMassArchetypeData& Sibling
 
 void FMassArchetypeData::AddEntity(FMassEntityHandle Entity)
 {
-	AddEntityInternal(Entity, true/*bInitializeFragments*/);
+	const int32 AbsoluteIndex = AddEntityInternal(Entity);
+
+	// Initialize fragments
+	const int32 ChunkIndex = AbsoluteIndex / NumEntitiesPerChunk;
+	const int32 IndexWithinChunk = AbsoluteIndex % NumEntitiesPerChunk;
+	FMassArchetypeChunk& Chunk = Chunks[ChunkIndex];
+	for (const FMassArchetypeFragmentConfig& FragmentConfig : FragmentConfigs)
+	{
+		void* FragmentPtr = FragmentConfig.GetFragmentData(Chunk.GetRawMemory(), IndexWithinChunk);
+		FragmentConfig.FragmentType->InitializeStruct(FragmentPtr);
+	}
 }
 
-int32 FMassArchetypeData::AddEntityInternal(FMassEntityHandle Entity, const bool bInitializeFragments)
+int32 FMassArchetypeData::AddEntityInternal(FMassEntityHandle Entity)
 {
 	int32 IndexWithinChunk = 0;
 	int32 AbsoluteIndex = 0;
@@ -174,16 +180,6 @@ int32 FMassArchetypeData::AddEntityInternal(FMassEntityHandle Entity, const bool
 		DestinationChunk->AddInstance();
 	}
 
-	// Initialize the fragment memory
-	if (bInitializeFragments)
-	{
-		for (const FMassArchetypeFragmentConfig& FragmentConfig : FragmentConfigs)
-		{
-			void* FragmentPtr = FragmentConfig.GetFragmentData(DestinationChunk->GetRawMemory(), IndexWithinChunk);
-			FragmentConfig.FragmentType->InitializeStruct(FragmentPtr);
-		}
-	}
-
 	// Add to the table and map
 	EntityMap.Add(Entity.Index, AbsoluteIndex);
 	DestinationChunk->GetEntityArrayElementRef(EntityListOffsetWithinChunk, IndexWithinChunk) = Entity;
@@ -194,10 +190,23 @@ int32 FMassArchetypeData::AddEntityInternal(FMassEntityHandle Entity, const bool
 void FMassArchetypeData::RemoveEntity(FMassEntityHandle Entity)
 {
 	const int32 AbsoluteIndex = EntityMap.FindAndRemoveChecked(Entity.Index);
-	RemoveEntityInternal(AbsoluteIndex, true/*bDestroyFragments*/);
+
+	// Destroy fragments
+	const int32 ChunkIndex = AbsoluteIndex / NumEntitiesPerChunk;
+	const int32 IndexWithinChunk = AbsoluteIndex % NumEntitiesPerChunk;
+	FMassArchetypeChunk& Chunk = Chunks[ChunkIndex];
+
+	for (const FMassArchetypeFragmentConfig& FragmentConfig : FragmentConfigs)
+	{
+		// Destroy the fragment data
+		void* DyingFragmentPtr = FragmentConfig.GetFragmentData(Chunk.GetRawMemory(), IndexWithinChunk);
+		FragmentConfig.FragmentType->DestroyStruct(DyingFragmentPtr);
+	}
+
+	RemoveEntityInternal(AbsoluteIndex);
 }
 
-void FMassArchetypeData::RemoveEntityInternal(const int32 AbsoluteIndex, const bool bDestroyFragments)
+void FMassArchetypeData::RemoveEntityInternal(const int32 AbsoluteIndex)
 {
 	const int32 ChunkIndex = AbsoluteIndex / NumEntitiesPerChunk;
 	const int32 IndexWithinChunk = AbsoluteIndex % NumEntitiesPerChunk;
@@ -206,8 +215,6 @@ void FMassArchetypeData::RemoveEntityInternal(const int32 AbsoluteIndex, const b
 
 	const int32 IndexToSwapFrom = Chunk.GetNumInstances() - 1;
 
-	checkf(bDestroyFragments || UE::Mass::Core::bBitwiseRelocateFragments, TEXT("We allow not to destroy fragments only in bit wise relocation mode."));
-
 	// Remove and swap the last entry in the chunk to the location of the removed item (if it's not the same as the dying entry)
 	if (IndexToSwapFrom != IndexWithinChunk)
 	{
@@ -216,28 +223,8 @@ void FMassArchetypeData::RemoveEntityInternal(const int32 AbsoluteIndex, const b
 			void* DyingFragmentPtr = FragmentConfig.GetFragmentData(Chunk.GetRawMemory(), IndexWithinChunk);
 			void* MovingFragmentPtr = FragmentConfig.GetFragmentData(Chunk.GetRawMemory(), IndexToSwapFrom);
 
-			if (UE::Mass::Core::bBitwiseRelocateFragments)
-			{
-				// Destroy fragment data
-				if (bDestroyFragments)
-				{
-					FragmentConfig.FragmentType->DestroyStruct(DyingFragmentPtr);
-				}
-
-				// Move last entry
-				FMemory::Memcpy(DyingFragmentPtr, MovingFragmentPtr, FragmentConfig.FragmentType->GetStructureSize());
-			}
-			else
-			{
-				// Destroy & initialize the fragment data
-				FragmentConfig.FragmentType->ClearScriptStruct(DyingFragmentPtr);
-
-				// Copy last entry
-				FragmentConfig.FragmentType->CopyScriptStruct(DyingFragmentPtr, MovingFragmentPtr);
-
-				// Destroy last entry
-				FragmentConfig.FragmentType->DestroyStruct(MovingFragmentPtr);
-			}
+			// Move last entry
+			FMemory::Memmove(DyingFragmentPtr, MovingFragmentPtr, FragmentConfig.FragmentType->GetStructureSize());
 		}
 
 		// Update the entity table and map
@@ -245,15 +232,6 @@ void FMassArchetypeData::RemoveEntityInternal(const int32 AbsoluteIndex, const b
 		Chunk.GetEntityArrayElementRef(EntityListOffsetWithinChunk, IndexWithinChunk) = EntityBeingSwapped;
 		EntityMap.FindChecked(EntityBeingSwapped.Index) = AbsoluteIndex;
 	}
-	else if (!UE::Mass::Core::bBitwiseRelocateFragments || bDestroyFragments)
-	{
-		for (const FMassArchetypeFragmentConfig& FragmentConfig : FragmentConfigs)
-		{
-			// Destroy the fragment data
-			void* DyingFragmentPtr = FragmentConfig.GetFragmentData(Chunk.GetRawMemory(), IndexWithinChunk);
-			FragmentConfig.FragmentType->DestroyStruct(DyingFragmentPtr);
-		}
-	}
 	
 	Chunk.RemoveInstance();
 
@@ -299,25 +277,11 @@ void FMassArchetypeData::BatchDestroyEntityChunks(FMassArchetypeSubChunks::FCons
 				void* DyingFragmentPtr = FragmentConfig.GetFragmentData(Chunk.GetRawMemory(), SubchunkInfo.SubchunkStart);
 				void* MovingFragmentPtr = FragmentConfig.GetFragmentData(Chunk.GetRawMemory(), SwapStartIndex);
 
-				if (UE::Mass::Core::bBitwiseRelocateFragments)
-				{
-					// Destroy the fragments we'll replace by the following copy
-					FragmentConfig.FragmentType->DestroyStruct(DyingFragmentPtr, NumberToMove);
+				// Destroy the fragments we'll replace by the following copy
+				FragmentConfig.FragmentType->DestroyStruct(DyingFragmentPtr, NumberToMove);
 
-					// Swap fragments to the empty space just created.
-					FMemory::Memcpy(DyingFragmentPtr, MovingFragmentPtr, FragmentConfig.FragmentType->GetStructureSize() * NumberToMove);
-				}
-				else
-				{
-					// Clear fragments that we will copy over. Clear destroys and initializes the fragments, which is needed for CopyScriptStruct().
-					FragmentConfig.FragmentType->ClearScriptStruct(DyingFragmentPtr, NumberToMove);
-
-					// Swap fragments into the empty space just created.
-					FragmentConfig.FragmentType->CopyScriptStruct(DyingFragmentPtr, MovingFragmentPtr, NumberToMove);
-
-					// Destroy the fragments that were moved.
-					FragmentConfig.FragmentType->DestroyStruct(MovingFragmentPtr, NumberToMove);
-				}
+				// Swap fragments to the empty space just created.
+				FMemory::Memmove(DyingFragmentPtr, MovingFragmentPtr, FragmentConfig.FragmentType->GetStructureSize() * NumberToMove);
 			}
 
 			// Update the entity table and map
@@ -395,7 +359,6 @@ void FMassArchetypeData::SetFragmentsData(const FMassEntityHandle Entity, TArray
 		check(FragmentType);
 		const int32 FragmentIndex = FragmentIndexMap.FindChecked(FragmentType);
 		void* FragmentMemory = GetFragmentData(FragmentIndex, InternalIndex);
-		// No UE::Mass::Core::bBitwiseRelocateFragments, this isn't a move fragment
 		FragmentType->CopyScriptStruct(FragmentMemory, Instance.GetMemory());
 	}
 }
@@ -415,7 +378,6 @@ void FMassArchetypeData::SetFragmentData(FMassArchetypeSubChunks::FConstSubChunk
 		uint8* FragmentMemory = (uint8*)FragmentConfigs[FragmentIndex].GetFragmentData(Chunks[ChunkIterator->ChunkIndex].GetRawMemory(), ChunkIterator->SubchunkStart);
 		for (int i = ChunkIterator->Length; i; --i, FragmentMemory += FragmentTypeSize)
 		{
-			// No UE::Mass::Core::bBitwiseRelocateFragments, this isn't a move of a fragment
 			FragmentType->CopyScriptStruct(FragmentMemory, FragmentSourceMemory);
 		}
 	}
@@ -430,44 +392,43 @@ void FMassArchetypeData::MoveEntityToAnotherArchetype(const FMassEntityHandle En
 	const int32 IndexWithinChunk = AbsoluteIndex % NumEntitiesPerChunk;
 	FMassArchetypeChunk& Chunk = Chunks[ChunkIndex];
 
-	constexpr bool bInitializeFragmentsDuringCreation = !UE::Mass::Core::bBitwiseRelocateFragments;
-	const int32 NewAbsoluteIndex = NewArchetype.AddEntityInternal(Entity, bInitializeFragmentsDuringCreation);
+	const int32 NewAbsoluteIndex = NewArchetype.AddEntityInternal(Entity);
 	const int32 NewChunkIndex = NewAbsoluteIndex / NewArchetype.NumEntitiesPerChunk;
 	const int32 NewIndexWithinChunk = NewAbsoluteIndex % NewArchetype.NumEntitiesPerChunk;
 	FMassArchetypeChunk& NewChunk = NewArchetype.Chunks[NewChunkIndex];
 
-	// for every NewArchetype's fragment see if it was in the old fragment as well and if so copy it's value. 
-	// If not then initialize the fragment (unless it has already been initialized based on bInitializeFragmentsDuringCreation
-	// value
+	// Relocate existing fragments and intialize new ones.
 	for (const FMassArchetypeFragmentConfig& NewFragmentConfig : NewArchetype.FragmentConfigs)
 	{
 		const int32* OldFragmentIndex = FragmentIndexMap.Find(NewFragmentConfig.FragmentType);
 		void* Dst = NewFragmentConfig.GetFragmentData(NewChunk.GetRawMemory(), NewIndexWithinChunk);
 
-		// Only copy if the fragment type exists in both archetypes
+		// Only relocate if the fragment type exists in both archetypes
 		if (OldFragmentIndex)
 		{
 			const void* Src = FragmentConfigs[*OldFragmentIndex].GetFragmentData(Chunk.GetRawMemory(), IndexWithinChunk);
-			if (UE::Mass::Core::bBitwiseRelocateFragments)
-			{
-				FMemory::Memcpy(Dst, Src, NewFragmentConfig.FragmentType->GetStructureSize());
-			}
-			else
-			{
-				NewFragmentConfig.FragmentType->CopyScriptStruct(Dst, Src);
-			}
+			FMemory::Memmove(Dst, Src, NewFragmentConfig.FragmentType->GetStructureSize());
 		}
-		else if (bInitializeFragmentsDuringCreation == false)
+		else
 		{
 			// the fragment's unique to the NewArchetype need to be initialized
-			// @todo we're doing it for tags here as well. A tiny bit of perf lost. Probably not worth adding a check
-			// but something to keep in mind. Will go away once tags are more of an archetype fragment than entity's
 			NewFragmentConfig.FragmentType->InitializeStruct(Dst);
 		}
 	}
 
-	constexpr bool bDestroyFragments = !UE::Mass::Core::bBitwiseRelocateFragments;
-	RemoveEntityInternal(AbsoluteIndex, bDestroyFragments);
+	// Delete fragments that were left behind
+	for (const FMassArchetypeFragmentConfig& FragmentConfig : FragmentConfigs)
+	{
+		// If the fragment is not in the new archetype, destroy it.
+		const int32* NewFragmentIndex = NewArchetype.FragmentIndexMap.Find(FragmentConfig.FragmentType);
+		if (NewFragmentIndex == nullptr)
+		{
+			void* DyingFragmentPtr = FragmentConfig.GetFragmentData(Chunk.GetRawMemory(), IndexWithinChunk);
+			FragmentConfig.FragmentType->DestroyStruct(DyingFragmentPtr);
+		}
+	}
+	
+	RemoveEntityInternal(AbsoluteIndex);
 }
 
 void FMassArchetypeData::ExecuteFunction(FMassExecutionContext& RunContext, const FMassExecuteFunction& Function, const FMassQueryRequirementIndicesMapping& RequirementMapping, FMassArchetypeSubChunks::FConstSubChunkArrayView SubChunkContainer)
@@ -599,28 +560,13 @@ void FMassArchetypeData::CompactEntities(const double TimeAllowed)
 		{
 			void* FromFragmentPtr = FragmentConfig.GetFragmentData(ChunkToEmpty->GetRawMemory(), FromIndex);
 			void* ToFragmentPtr = FragmentConfig.GetFragmentData(ChunkToFill->GetRawMemory(), ToIndex);
-
-			if (UE::Mass::Core::bBitwiseRelocateFragments)
-			{
-				// Move all entries
-				FMemory::Memcpy(ToFragmentPtr, FromFragmentPtr, FragmentConfig.FragmentType->GetStructureSize() * NumberOfEntitiesToMove);
-			}
-			else
-			{
-				// Destroy & initialize the fragment data
-				FragmentConfig.FragmentType->ClearScriptStruct(ToFragmentPtr, NumberOfEntitiesToMove);
-
-				// Copy all entries
-				FragmentConfig.FragmentType->CopyScriptStruct(ToFragmentPtr, FromFragmentPtr, NumberOfEntitiesToMove);
-
-				// Destroy all entries
-				FragmentConfig.FragmentType->DestroyStruct(FromFragmentPtr, NumberOfEntitiesToMove);
-			}
+			// Move all entries
+			FMemory::Memmove(ToFragmentPtr, FromFragmentPtr, FragmentConfig.FragmentType->GetStructureSize() * NumberOfEntitiesToMove);
 		}
 
 		FMassEntityHandle* FromEntity = &ChunkToEmpty->GetEntityArrayElementRef(EntityListOffsetWithinChunk, FromIndex);
 		FMassEntityHandle* ToEntity = &ChunkToFill->GetEntityArrayElementRef(EntityListOffsetWithinChunk, ToIndex);
-		FMemory::Memcpy(ToEntity, FromEntity, NumberOfEntitiesToMove * sizeof(FMassEntityHandle));
+		FMemory::Memmove(ToEntity, FromEntity, NumberOfEntitiesToMove * sizeof(FMassEntityHandle));
 		ChunkToFill->AddMultipleInstances(NumberOfEntitiesToMove);
 		ChunkToEmpty->RemoveMultipleInstances(NumberOfEntitiesToMove);
 
diff --git a/Engine/Plugins/Runtime/MassEntity/Source/MassEntity/Private/MassArchetypeData.h b/Engine/Plugins/Runtime/MassEntity/Source/MassEntity/Private/MassArchetypeData.h
--- a/Engine/Plugins/Runtime/MassEntity/Source/MassEntity/Private/MassArchetypeData.h
+++ b/Engine/Plugins/Runtime/MassEntity/Source/MassEntity/Private/MassArchetypeData.h
@@ -304,6 +304,6 @@ protected:
 	void BindSharedFragmentRequirements(FMassExecutionContext& RunContext, const FMassFragmentIndicesMapping& ChunkFragmentsMapping);
 
 private:
-	int32 AddEntityInternal(FMassEntityHandle Entity, const bool bInitializeFragments);
-	void RemoveEntityInternal(const int32 AbsoluteIndex, const bool bDestroyFragments);
+	int32 AddEntityInternal(FMassEntityHandle Entity);
+	void RemoveEntityInternal(const int32 AbsoluteIndex);
 };
