Reverted the change to use memmove over memcpy in MassArchetypeData.cpp

diff --git a/Engine/Plugins/Runtime/MassEntity/Source/MassEntity/Private/MassArchetypeData.cpp b/Engine/Plugins/Runtime/MassEntity/Source/MassEntity/Private/MassArchetypeData.cpp
--- a/Engine/Plugins/Runtime/MassEntity/Source/MassEntity/Private/MassArchetypeData.cpp
+++ b/Engine/Plugins/Runtime/MassEntity/Source/MassEntity/Private/MassArchetypeData.cpp
@@ -224,7 +224,7 @@ void FMassArchetypeData::RemoveEntityInternal(const int32 AbsoluteIndex)
 			void* MovingFragmentPtr = FragmentConfig.GetFragmentData(Chunk.GetRawMemory(), IndexToSwapFrom);
 
 			// Move last entry
-			FMemory::Memmove(DyingFragmentPtr, MovingFragmentPtr, FragmentConfig.FragmentType->GetStructureSize());
+			FMemory::Memcpy(DyingFragmentPtr, MovingFragmentPtr, FragmentConfig.FragmentType->GetStructureSize());
 		}
 
 		// Update the entity table and map
@@ -281,7 +281,7 @@ void FMassArchetypeData::BatchDestroyEntityChunks(FMassArchetypeSubChunks::FCons
 				FragmentConfig.FragmentType->DestroyStruct(DyingFragmentPtr, NumberToMove);
 
 				// Swap fragments to the empty space just created.
-				FMemory::Memmove(DyingFragmentPtr, MovingFragmentPtr, FragmentConfig.FragmentType->GetStructureSize() * NumberToMove);
+				FMemory::Memcpy(DyingFragmentPtr, MovingFragmentPtr, FragmentConfig.FragmentType->GetStructureSize() * NumberToMove);
 			}
 
 			// Update the entity table and map
@@ -407,7 +407,7 @@ void FMassArchetypeData::MoveEntityToAnotherArchetype(const FMassEntityHandle En
 		if (OldFragmentIndex)
 		{
 			const void* Src = FragmentConfigs[*OldFragmentIndex].GetFragmentData(Chunk.GetRawMemory(), IndexWithinChunk);
-			FMemory::Memmove(Dst, Src, NewFragmentConfig.FragmentType->GetStructureSize());
+			FMemory::Memcpy(Dst, Src, NewFragmentConfig.FragmentType->GetStructureSize());
 		}
 		else
 		{
@@ -561,12 +561,12 @@ void FMassArchetypeData::CompactEntities(const double TimeAllowed)
 			void* FromFragmentPtr = FragmentConfig.GetFragmentData(ChunkToEmpty->GetRawMemory(), FromIndex);
 			void* ToFragmentPtr = FragmentConfig.GetFragmentData(ChunkToFill->GetRawMemory(), ToIndex);
 			// Move all entries
-			FMemory::Memmove(ToFragmentPtr, FromFragmentPtr, FragmentConfig.FragmentType->GetStructureSize() * NumberOfEntitiesToMove);
+			FMemory::Memcpy(ToFragmentPtr, FromFragmentPtr, FragmentConfig.FragmentType->GetStructureSize() * NumberOfEntitiesToMove);
 		}
 
 		FMassEntityHandle* FromEntity = &ChunkToEmpty->GetEntityArrayElementRef(EntityListOffsetWithinChunk, FromIndex);
 		FMassEntityHandle* ToEntity = &ChunkToFill->GetEntityArrayElementRef(EntityListOffsetWithinChunk, ToIndex);
-		FMemory::Memmove(ToEntity, FromEntity, NumberOfEntitiesToMove * sizeof(FMassEntityHandle));
+		FMemory::Memcpy(ToEntity, FromEntity, NumberOfEntitiesToMove * sizeof(FMassEntityHandle));
 		ChunkToFill->AddMultipleInstances(NumberOfEntitiesToMove);
 		ChunkToEmpty->RemoveMultipleInstances(NumberOfEntitiesToMove);
 
