Added a batched way to change entities set tags. A part of this change is adding batch Add and Move operations to MassArchetypeData

diff --git a/Engine/Plugins/Runtime/MassEntity/Source/MassEntity/Private/MassArchetypeData.cpp b/Engine/Plugins/Runtime/MassEntity/Source/MassEntity/Private/MassArchetypeData.cpp
--- a/Engine/Plugins/Runtime/MassEntity/Source/MassEntity/Private/MassArchetypeData.cpp
+++ b/Engine/Plugins/Runtime/MassEntity/Source/MassEntity/Private/MassArchetypeData.cpp
@@ -263,52 +263,14 @@ void FMassArchetypeData::BatchDestroyEntityChunks(FMassArchetypeSubChunks::FCons
 		FMassEntityHandle* DyingEntityPtr = &Chunk.GetEntityArrayElementRef(EntityListOffsetWithinChunk, SubchunkInfo.SubchunkStart);
 		OutEntitiesRemoved.Append(DyingEntityPtr, SubchunkInfo.Length);
 
-		const int32 NumberToMove = FMath::Min(Chunk.GetNumInstances() - (SubchunkInfo.SubchunkStart + SubchunkInfo.Length), SubchunkInfo.Length);
-		checkf(NumberToMove >= 0, TEXT("Trying to move a negative number of elements indicates a problem with SubchunkInfo, it's possibly out of date."));
-		const int32 NumberToCut = FMath::Max(SubchunkInfo.Length - NumberToMove, 0);
-		
-		if (NumberToMove > 0)
-		{
-			const int32 SwapStartIndex = Chunk.GetNumInstances() - NumberToMove;
-			checkf((SubchunkInfo.SubchunkStart + NumberToMove - 1) < SwapStartIndex, TEXT("Remove and Move ranges overlap"));
-
-			for (const FMassArchetypeFragmentConfig& FragmentConfig : FragmentConfigs)
-			{
-				void* DyingFragmentPtr = FragmentConfig.GetFragmentData(Chunk.GetRawMemory(), SubchunkInfo.SubchunkStart);
-				void* MovingFragmentPtr = FragmentConfig.GetFragmentData(Chunk.GetRawMemory(), SwapStartIndex);
-
-				// Destroy the fragments we'll replace by the following copy
-				FragmentConfig.FragmentType->DestroyStruct(DyingFragmentPtr, NumberToMove);
-
-				// Swap fragments to the empty space just created.
-				FMemory::Memcpy(DyingFragmentPtr, MovingFragmentPtr, FragmentConfig.FragmentType->GetStructureSize() * NumberToMove);
-			}
-
-			// Update the entity table and map
-			const FMassEntityHandle* MovingEntityPtr = &Chunk.GetEntityArrayElementRef(EntityListOffsetWithinChunk, SwapStartIndex);
-			int32 AbsoluteIndex = SubchunkInfo.ChunkIndex * NumEntitiesPerChunk + SubchunkInfo.SubchunkStart;
-
-			for (int i = 0; i < NumberToMove; ++i)
-			{
-				DyingEntityPtr[i] = MovingEntityPtr[i];
-				EntityMap.FindChecked(MovingEntityPtr[i].Index) = AbsoluteIndex++;
-			}
-		}
-
-		if (NumberToCut > 0)
+		for (const FMassArchetypeFragmentConfig& FragmentConfig : FragmentConfigs)
 		{
-			// just clean up the rest. Note that we explicitly do not clean the spots vacated by entities moved from 
-			// the back of the chunk - if we did the risk calling DestroyStruct on them multiple times
-			const int32 CutStartIndex = SubchunkInfo.SubchunkStart + NumberToMove;
-			for (const FMassArchetypeFragmentConfig& FragmentConfig : FragmentConfigs)
-			{
-				// Destroy the fragment data
-				void* DyingFragmentPtr = FragmentConfig.GetFragmentData(Chunk.GetRawMemory(), CutStartIndex);
-				FragmentConfig.FragmentType->DestroyStruct(DyingFragmentPtr, NumberToCut);
-			}
+			// Destroy the fragment data
+			void* DyingFragmentPtr = FragmentConfig.GetFragmentData(Chunk.GetRawMemory(), SubchunkInfo.SubchunkStart);
+			FragmentConfig.FragmentType->DestroyStruct(DyingFragmentPtr, SubchunkInfo.Length);
 		}
 
-		Chunk.RemoveMultipleInstances(SubchunkInfo.Length);
+		BatchRemoveEntitiesInternal(SubchunkInfo.ChunkIndex, SubchunkInfo.SubchunkStart, SubchunkInfo.Length);
 	}
 
 	for (int i = InitialOutEntitiesCount; i < OutEntitiesRemoved.Num(); ++i)
@@ -397,37 +359,8 @@ void FMassArchetypeData::MoveEntityToAnotherArchetype(const FMassEntityHandle En
 	const int32 NewIndexWithinChunk = NewAbsoluteIndex % NewArchetype.NumEntitiesPerChunk;
 	FMassArchetypeChunk& NewChunk = NewArchetype.Chunks[NewChunkIndex];
 
-	// Relocate existing fragments and intialize new ones.
-	for (const FMassArchetypeFragmentConfig& NewFragmentConfig : NewArchetype.FragmentConfigs)
-	{
-		const int32* OldFragmentIndex = FragmentIndexMap.Find(NewFragmentConfig.FragmentType);
-		void* Dst = NewFragmentConfig.GetFragmentData(NewChunk.GetRawMemory(), NewIndexWithinChunk);
+	MoveFragmentsToAnotherArchetypeInternal(NewArchetype, { NewChunk.GetRawMemory(), NewIndexWithinChunk }, { Chunk.GetRawMemory(), IndexWithinChunk }, /*Count=*/1);
 
-		// Only relocate if the fragment type exists in both archetypes
-		if (OldFragmentIndex)
-		{
-			const void* Src = FragmentConfigs[*OldFragmentIndex].GetFragmentData(Chunk.GetRawMemory(), IndexWithinChunk);
-			FMemory::Memcpy(Dst, Src, NewFragmentConfig.FragmentType->GetStructureSize());
-		}
-		else
-		{
-			// the fragment's unique to the NewArchetype need to be initialized
-			NewFragmentConfig.FragmentType->InitializeStruct(Dst);
-		}
-	}
-
-	// Delete fragments that were left behind
-	for (const FMassArchetypeFragmentConfig& FragmentConfig : FragmentConfigs)
-	{
-		// If the fragment is not in the new archetype, destroy it.
-		const int32* NewFragmentIndex = NewArchetype.FragmentIndexMap.Find(FragmentConfig.FragmentType);
-		if (NewFragmentIndex == nullptr)
-		{
-			void* DyingFragmentPtr = FragmentConfig.GetFragmentData(Chunk.GetRawMemory(), IndexWithinChunk);
-			FragmentConfig.FragmentType->DestroyStruct(DyingFragmentPtr);
-		}
-	}
-	
 	RemoveEntityInternal(AbsoluteIndex);
 }
 
@@ -515,7 +448,7 @@ void FMassArchetypeData::CompactEntities(const double TimeAllowed)
 	{
 		// Skip already full chunks
 		const int32 NumInstances = Chunk.GetNumInstances();
-		if(NumInstances > 0 && NumInstances < NumEntitiesPerChunk)
+		if (NumInstances > 0 && NumInstances < NumEntitiesPerChunk)
 		{
 			SortedChunks.Add(&Chunk);
 		}
@@ -533,7 +466,7 @@ void FMassArchetypeData::CompactEntities(const double TimeAllowed)
 	});
 
 	int32 ChunkToFillSortedIdx = 0;
-	int32 ChunkToEmptySortedIdx = SortedChunks.Num() -1;
+	int32 ChunkToEmptySortedIdx = SortedChunks.Num() - 1;
 	while (ChunkToFillSortedIdx < ChunkToEmptySortedIdx && FPlatformTime::Seconds() < TimeAllowedEnd)
 	{
 		while (ChunkToFillSortedIdx < SortedChunks.Num() && SortedChunks[ChunkToFillSortedIdx]->GetNumInstances() == NumEntitiesPerChunk)
@@ -551,7 +484,7 @@ void FMassArchetypeData::CompactEntities(const double TimeAllowed)
 
 		FMassArchetypeChunk* ChunkToFill = SortedChunks[ChunkToFillSortedIdx];
 		FMassArchetypeChunk* ChunkToEmpty = SortedChunks[ChunkToEmptySortedIdx];
-		const int32 NumberOfEntitiesToMove =  FMath::Min(NumEntitiesPerChunk-ChunkToFill->GetNumInstances(), ChunkToEmpty->GetNumInstances());
+		const int32 NumberOfEntitiesToMove = FMath::Min(NumEntitiesPerChunk - ChunkToFill->GetNumInstances(), ChunkToEmpty->GetNumInstances());
 		const int32 FromIndex = ChunkToEmpty->GetNumInstances() - NumberOfEntitiesToMove;
 		const int32 ToIndex = ChunkToFill->GetNumInstances();
 		check(NumberOfEntitiesToMove > 0);
@@ -571,12 +504,12 @@ void FMassArchetypeData::CompactEntities(const double TimeAllowed)
 		ChunkToEmpty->RemoveMultipleInstances(NumberOfEntitiesToMove);
 
 		const int32 ChunkToFillIdx = UE_PTRDIFF_TO_INT32(ChunkToFill - &Chunks[0]);
-		check(ChunkToFillIdx >=0 && ChunkToFillIdx < Chunks.Num());
+		check(ChunkToFillIdx >= 0 && ChunkToFillIdx < Chunks.Num());
 		const int32 AbsoluteIndex = ChunkToFillIdx * NumEntitiesPerChunk + ToIndex;
 
-		for (int32 i =0; i < NumberOfEntitiesToMove; i++, ++ToEntity)
+		for (int32 i = 0; i < NumberOfEntitiesToMove; i++, ++ToEntity)
 		{
-			EntityMap.FindChecked(ToEntity->Index) = AbsoluteIndex+i;
+			EntityMap.FindChecked(ToEntity->Index) = AbsoluteIndex + i;
 		}
 	}
 }
@@ -907,3 +840,235 @@ void FMassArchetypeData::REMOVEME_GetArrayViewForFragmentInChunk(int32 ChunkInde
 	OutNumEntities = Chunk.GetNumInstances();
 }
 
+//////////////////////////////////////////////////////////////////////
+// FMassArchetypeData batched api
+
+void FMassArchetypeData::BatchAddEntities(TConstArrayView<FMassEntityHandle> Entities, TArray<FMassArchetypeSubChunks::FSubChunkInfo>& OutNewChunks)
+{
+	FMassArchetypeSubChunks::FSubChunkInfo ResultSubchunk;
+	ResultSubchunk.ChunkIndex = 0;
+	int32 NumberMoved = 0;
+	do 
+	{
+		ResultSubchunk = PrepareNextEntitiesSpanInternal(MakeArrayView(Entities.GetData() + NumberMoved, Entities.Num() - NumberMoved), ResultSubchunk.ChunkIndex);
+		check(Chunks.IsValidIndex(ResultSubchunk.ChunkIndex) && Chunks[ResultSubchunk.ChunkIndex].IsValidSubChunk(ResultSubchunk.SubchunkStart, ResultSubchunk.Length));
+		
+		for (const FMassArchetypeFragmentConfig& FragmentConfig : FragmentConfigs)
+		{
+			void* FragmentPtr = FragmentConfig.GetFragmentData(Chunks[ResultSubchunk.ChunkIndex].GetRawMemory(), ResultSubchunk.SubchunkStart);
+			FragmentConfig.FragmentType->InitializeStruct(FragmentPtr, ResultSubchunk.Length);
+		}
+
+		NumberMoved += ResultSubchunk.Length;
+
+		OutNewChunks.Add(ResultSubchunk);
+
+	} while (NumberMoved < Entities.Num());
+}
+
+void FMassArchetypeData::BatchMoveEntitiesToAnotherArchetype(const FMassArchetypeSubChunks& ChunkCollection, FMassArchetypeData& NewArchetype, TArray<FMassEntityHandle>& OutEntitesBeingMoved, TArray<FMassArchetypeSubChunks::FSubChunkInfo>* OutNewChunks)
+{
+	check(&NewArchetype != this);
+
+	// Sorting the subchunks info so that subchunks of a given chunk are processed "from the back". Otherwise removing 
+	// a subchunk from the front of the chunk would inevitably invalidate following subchunks' information.
+	TArray<FMassArchetypeSubChunks::FSubChunkInfo> Subchunks(ChunkCollection.GetChunks());
+	Subchunks.Sort([](const FMassArchetypeSubChunks::FSubChunkInfo& A, const FMassArchetypeSubChunks::FSubChunkInfo& B)
+		{
+			return A.ChunkIndex < B.ChunkIndex || (A.ChunkIndex == B.ChunkIndex && A.SubchunkStart > B.SubchunkStart);
+		});
+
+	const int32 InitialOutEntitiesCount = OutEntitesBeingMoved.Num();
+
+	for (const FMassArchetypeSubChunks::FSubChunkInfo SubchunkInfo : Subchunks)
+	{
+		FMassArchetypeChunk& Chunk = Chunks[SubchunkInfo.ChunkIndex];
+
+		// 0 - consider compacting new archetype to ensure larger empty spaces
+		// 1. find next free spot in the destination archetype
+		// 2. min(amount of elements) to move
+
+		// gather entities we're about to remove
+		FMassEntityHandle* DyingEntityPtr = &Chunk.GetEntityArrayElementRef(EntityListOffsetWithinChunk, SubchunkInfo.SubchunkStart);
+		const int32 EntitesBeingMovedStartIndex = OutEntitesBeingMoved.Num();
+		OutEntitesBeingMoved.Append(DyingEntityPtr, SubchunkInfo.Length);
+
+		FMassArchetypeSubChunks::FSubChunkInfo ResultSubChunk;
+		ResultSubChunk.ChunkIndex = 0;
+		ResultSubChunk.Length = 0;
+		int32 NumberMoved = 0;
+
+		do
+		{
+			const int32 IndexWithinChunk = SubchunkInfo.SubchunkStart + NumberMoved;
+
+			ResultSubChunk = NewArchetype.PrepareNextEntitiesSpanInternal(MakeArrayView(DyingEntityPtr + NumberMoved, SubchunkInfo.Length - NumberMoved)
+				, ResultSubChunk.ChunkIndex);
+
+			FMassArchetypeChunk& NewChunk = NewArchetype.Chunks[ResultSubChunk.ChunkIndex];
+			MoveFragmentsToAnotherArchetypeInternal(NewArchetype, {NewChunk.GetRawMemory(), ResultSubChunk.SubchunkStart}, {Chunk.GetRawMemory(), IndexWithinChunk}, ResultSubChunk.Length);
+
+			NumberMoved += ResultSubChunk.Length;
+
+			// @todo consider adding a 'merge sequences' pass at the end of the function since we can end up with 
+			// subchunks being right next to each other
+			if (OutNewChunks)
+			{
+				OutNewChunks->Add(ResultSubChunk);
+			}
+
+		} while (NumberMoved < SubchunkInfo.Length);
+
+		BatchRemoveEntitiesInternal(SubchunkInfo.ChunkIndex, SubchunkInfo.SubchunkStart, SubchunkInfo.Length);
+	}
+
+	for (int i = InitialOutEntitiesCount; i < OutEntitesBeingMoved.Num(); ++i)
+	{
+		EntityMap.FindAndRemoveChecked(OutEntitesBeingMoved[i].Index);
+	}
+}
+
+FMassArchetypeSubChunks::FSubChunkInfo FMassArchetypeData::PrepareNextEntitiesSpanInternal(TConstArrayView<FMassEntityHandle> Entities, const int32 StartingChunk)
+{
+	int32 StartIndexWithinChunk = INDEX_NONE;
+	int32 AbsoluteStartIndex = 0;
+
+	FMassArchetypeChunk* DestinationChunk = nullptr;
+	
+	int32 ChunkIndex = StartingChunk;
+	// find a chunk with any room left
+	for (; ChunkIndex < Chunks.Num(); ++ChunkIndex)
+	{
+		FMassArchetypeChunk& Chunk = Chunks[ChunkIndex];
+		if (Chunk.GetNumInstances() < NumEntitiesPerChunk)
+		{
+			StartIndexWithinChunk = Chunk.GetNumInstances();
+			AbsoluteStartIndex = ChunkIndex * NumEntitiesPerChunk + StartIndexWithinChunk;
+
+			DestinationChunk = &Chunk;
+
+			if (StartIndexWithinChunk == 0)
+			{
+				Chunk.Recycle(ChunkFragmentsTemplate);
+			}
+			break;
+		}
+	}
+
+	// if no chunk found create one
+	if (DestinationChunk == nullptr)
+	{
+		ChunkIndex = Chunks.Num();
+		AbsoluteStartIndex = Chunks.Num() * NumEntitiesPerChunk;
+		StartIndexWithinChunk = 0;
+
+		DestinationChunk = &Chunks.Emplace_GetRef(GetChunkAllocSize(), ChunkFragmentsTemplate);
+	}
+
+	check(DestinationChunk);
+
+	// we might be able to fit in less entitites than requested
+	const int32 NumToAdd = FMath::Min(NumEntitiesPerChunk - StartIndexWithinChunk, Entities.Num());
+	check(NumToAdd);
+	DestinationChunk->AddMultipleInstances(NumToAdd);
+
+	// Add to the table and map
+	int32 AbsoluteIndex = AbsoluteStartIndex;
+	for (int32 i = 0; i < NumToAdd; ++i)
+	{
+		EntityMap.Add(Entities[i].Index, AbsoluteIndex++);
+	}
+
+	FMassEntityHandle* FirstAddedEntity = &DestinationChunk->GetEntityArrayElementRef(EntityListOffsetWithinChunk, StartIndexWithinChunk);
+	FMemory::Memcpy(FirstAddedEntity, Entities.GetData(), sizeof(FMassEntityHandle) * NumToAdd);
+
+	return FMassArchetypeSubChunks::FSubChunkInfo(ChunkIndex, StartIndexWithinChunk, NumToAdd);
+}
+
+void FMassArchetypeData::BatchRemoveEntitiesInternal(const int32 ChunkIndex, const int32 StartIndexWithinChunk, const int32 NumberToRemove)
+{
+	FMassArchetypeChunk& Chunk = Chunks[ChunkIndex];
+	
+	const int32 NumberToMove = FMath::Min(Chunk.GetNumInstances() - (StartIndexWithinChunk + NumberToRemove), NumberToRemove);
+	checkf(NumberToMove >= 0, TEXT("Trying to move a negative number of elements indicates a problem with SubchunkInfo, it's possibly out of date."));
+	const int32 NumberToCut = FMath::Max(NumberToRemove - NumberToMove, 0);
+
+	if (NumberToMove > 0)
+	{
+		FMassEntityHandle* DyingEntityPtr = &Chunk.GetEntityArrayElementRef(EntityListOffsetWithinChunk, StartIndexWithinChunk);
+
+		const int32 SwapStartIndex = Chunk.GetNumInstances() - NumberToMove;
+		checkf((StartIndexWithinChunk + NumberToMove - 1) < SwapStartIndex, TEXT("Remove and Move ranges overlap"));
+
+		MoveFragmentsToNewLocationInternal({ Chunk.GetRawMemory(), StartIndexWithinChunk }, { Chunk.GetRawMemory(), SwapStartIndex }, NumberToMove);
+		
+		// Update the entity table and map
+		const FMassEntityHandle* MovingEntityPtr = &Chunk.GetEntityArrayElementRef(EntityListOffsetWithinChunk, SwapStartIndex);
+		int32 AbsoluteIndex = ChunkIndex * NumEntitiesPerChunk + StartIndexWithinChunk;
+
+		for (int i = 0; i < NumberToMove; ++i)
+		{
+			DyingEntityPtr[i] = MovingEntityPtr[i];
+			EntityMap.FindChecked(MovingEntityPtr[i].Index) = AbsoluteIndex++;
+		}
+	}
+
+	Chunk.RemoveMultipleInstances(NumberToRemove);
+
+	// If the chunk itself is empty now, see if we can remove it entirely
+	// Note: This is only possible for trailing chunks, to avoid messing up the absolute indices in the entities map
+	while ((Chunks.Num() > 0) && (Chunks.Last().GetNumInstances() == 0))
+	{
+		Chunks.RemoveAt(Chunks.Num() - 1, 1, /*bAllowShrinking=*/ false);
+	}
+}
+
+void FMassArchetypeData::MoveFragmentsToAnotherArchetypeInternal(FMassArchetypeData& TargetArchetype, FMassArchetypeData::FTransientChunkLocation Target
+	, const FMassArchetypeData::FTransientChunkLocation Source, const int32 ElementsNum)
+{
+	// for every TargetArchetype's fragment see if it was in the old archetype as well and if so copy it's value. 
+	// If not then initialize the fragment.
+	for (const FMassArchetypeFragmentConfig& TargetFragmentConfig : TargetArchetype.FragmentConfigs)
+	{
+		const int32* OldFragmentIndex = FragmentIndexMap.Find(TargetFragmentConfig.FragmentType);
+		void* Dst = TargetFragmentConfig.GetFragmentData(Target.RawChunkMemory, Target.IndexWithinChunk);
+
+		// Only copy if the fragment type exists in both archetypes
+		if (OldFragmentIndex)
+		{
+			const void* Src = FragmentConfigs[*OldFragmentIndex].GetFragmentData(Source.RawChunkMemory, Source.IndexWithinChunk);
+			FMemory::Memcpy(Dst, Src, TargetFragmentConfig.FragmentType->GetStructureSize() * ElementsNum);
+		}
+		else
+		{
+			// the fragment's unique to the TargetArchetype need to be initialized
+			// @todo we're doing it for tags here as well. A tiny bit of perf lost. Probably not worth adding a check
+			// but something to keep in mind. Will go away once tags are more of an archetype fragment than entity's
+			TargetFragmentConfig.FragmentType->InitializeStruct(Dst, ElementsNum);
+		}
+	}
+
+	// Delete fragments that were left behind
+	for (const FMassArchetypeFragmentConfig& FragmentConfig : FragmentConfigs)
+	{
+		// If the fragment is not in the new archetype, destroy it.
+		const int32* NewFragmentIndex = TargetArchetype.FragmentIndexMap.Find(FragmentConfig.FragmentType);
+		if (NewFragmentIndex == nullptr)
+		{
+			void* DyingFragmentPtr = FragmentConfig.GetFragmentData(Source.RawChunkMemory, Source.IndexWithinChunk);
+			FragmentConfig.FragmentType->DestroyStruct(DyingFragmentPtr, ElementsNum);
+		}
+	}
+}
+
+FORCEINLINE void FMassArchetypeData::MoveFragmentsToNewLocationInternal(FMassArchetypeData::FTransientChunkLocation Target, const FMassArchetypeData::FTransientChunkLocation Source, const int32 NumberToMove)
+{
+	for (const FMassArchetypeFragmentConfig& FragmentConfig : FragmentConfigs)
+	{
+		void* DyingFragmentPtr = FragmentConfig.GetFragmentData(Target.RawChunkMemory, Target.IndexWithinChunk);
+		void* MovingFragmentPtr = FragmentConfig.GetFragmentData(Source.RawChunkMemory, Source.IndexWithinChunk); 
+
+		// Swap fragments to the empty space just created.
+		FMemory::Memcpy(DyingFragmentPtr, MovingFragmentPtr, FragmentConfig.FragmentType->GetStructureSize() * NumberToMove);
+	}
+}
diff --git a/Engine/Plugins/Runtime/MassEntity/Source/MassEntity/Private/MassArchetypeData.h b/Engine/Plugins/Runtime/MassEntity/Source/MassEntity/Private/MassArchetypeData.h
--- a/Engine/Plugins/Runtime/MassEntity/Source/MassEntity/Private/MassArchetypeData.h
+++ b/Engine/Plugins/Runtime/MassEntity/Source/MassEntity/Private/MassArchetypeData.h
@@ -113,6 +113,11 @@ public:
 		}
 	}
 
+	bool IsValidSubChunk(const int32 StartIndex, const int32 Length) const
+	{
+		return StartIndex >= 0 && StartIndex < NumInstances && (StartIndex + Length) <= NumInstances;
+	}
+
 #if WITH_MASSENTITY_DEBUG
 	int32 DebugGetChunkFragmentCount() const { return ChunkFragmentData.Num(); }
 #endif // WITH_MASSENTITY_DEBUG
@@ -190,7 +195,6 @@ public:
 
 	void AddEntity(FMassEntityHandle Entity);
 	void RemoveEntity(FMassEntityHandle Entity);
-	void BatchDestroyEntityChunks(FMassArchetypeSubChunks::FConstSubChunkArrayView SubChunkContainer, TArray<FMassEntityHandle>& OutEntitiesRemoved);
 
 	bool HasFragmentDataForEntity(const UScriptStruct* FragmentType, int32 EntityIndex) const;
 	void* GetFragmentDataForEntityChecked(const UScriptStruct* FragmentType, int32 EntityIndex) const;
@@ -292,7 +296,24 @@ public:
 
 	bool IsInitialized() const { return TotalBytesPerEntity > 0 && FragmentConfigs.IsEmpty() == false; }
 
+	//////////////////////////////////////////////////////////////////////
+	// batched api
+	void BatchDestroyEntityChunks(FMassArchetypeSubChunks::FConstSubChunkArrayView SubChunkContainer, TArray<FMassEntityHandle>& OutEntitiesRemoved);
+	void BatchAddEntities(TConstArrayView<FMassEntityHandle> Entities, TArray<FMassArchetypeSubChunks::FSubChunkInfo>& OutNewChunks);
+	void BatchMoveEntitiesToAnotherArchetype(const FMassArchetypeSubChunks& ChunkCollection, FMassArchetypeData& NewArchetype, TArray<FMassEntityHandle>& OutEntitesBeingMoved, TArray<FMassArchetypeSubChunks::FSubChunkInfo>* OutNewChunks = nullptr);
+
 protected:
+	FMassArchetypeSubChunks::FSubChunkInfo PrepareNextEntitiesSpanInternal(TConstArrayView<FMassEntityHandle> Entities, const int32 StartingChunk = 0);
+	void BatchRemoveEntitiesInternal(const int32 ChunkIndex, const int32 StartIndexWithinChunk, const int32 NumberToRemove);
+
+	struct FTransientChunkLocation
+	{
+		uint8* RawChunkMemory;
+		int32 IndexWithinChunk;
+	};
+	void MoveFragmentsToAnotherArchetypeInternal(FMassArchetypeData& TargetArchetype, FTransientChunkLocation Target, const FTransientChunkLocation Source, const int32 ElementsNum);
+	void MoveFragmentsToNewLocationInternal(FTransientChunkLocation Target, const FTransientChunkLocation Source, const int32 NumberToMove);
+
 	FORCEINLINE void* GetFragmentData(const int32 FragmentIndex, uint8* ChunkRawMemory, const int32 IndexWithinChunk) const
 	{
 		return FragmentConfigs[FragmentIndex].GetFragmentData(ChunkRawMemory, IndexWithinChunk);
diff --git a/Engine/Plugins/Runtime/MassEntity/Source/MassEntity/Private/MassEntitySubsystem.cpp b/Engine/Plugins/Runtime/MassEntity/Source/MassEntity/Private/MassEntitySubsystem.cpp
--- a/Engine/Plugins/Runtime/MassEntity/Source/MassEntity/Private/MassEntitySubsystem.cpp
+++ b/Engine/Plugins/Runtime/MassEntity/Source/MassEntity/Private/MassEntitySubsystem.cpp
@@ -15,6 +15,30 @@ const FMassEntityHandle UMassEntitySubsystem::InvalidEntity;
 //@TODO: Do we allow GCable types?  If so, need to implement AddReferencedObjects
 //@TODO: Do a UObject stat scope around the system Run call
 
+namespace UE::Mass::Private
+{
+	// note: this function doesn't set EntityHandle.SerialNumber
+	void ConvertArchetypelessSubchunksIntoEntityHandles(FMassArchetypeSubChunks::FConstSubChunkArrayView Subchunks, TArray<FMassEntityHandle>& OutEntityHandles)
+	{
+		int32 TotalCount = 0;
+		for (const FMassArchetypeSubChunks::FSubChunkInfo& Subchunk : Subchunks)
+		{
+			TotalCount += Subchunk.Length;
+		}
+
+		int32 Index = OutEntityHandles.Num();
+		OutEntityHandles.AddDefaulted(TotalCount);
+
+		for (const FMassArchetypeSubChunks::FSubChunkInfo& Subchunk : Subchunks)
+		{
+			for (int i = Subchunk.SubchunkStart; i < Subchunk.SubchunkStart + Subchunk.Length; ++i)
+			{
+				OutEntityHandles[Index++].Index = i;
+			}
+		}
+	}
+}
+
 //////////////////////////////////////////////////////////////////////
 // UMassEntitySubsystem
 
@@ -151,7 +175,7 @@ FMassArchetypeHandle UMassEntitySubsystem::CreateArchetype(const FMassArchetypeC
 		NewArchetype->Initialize(Composition, SharedFragmentValues);
 		HashRow.Add(NewArchetype);
 
-		for (const FMassArchetypeFragmentConfig & FragmentConfig : NewArchetype->GetFragmentConfigs())
+		for (const FMassArchetypeFragmentConfig& FragmentConfig : NewArchetype->GetFragmentConfigs())
 		{
 			checkSlow(FragmentConfig.FragmentType)
 			FragmentTypeToArchetypeMap.FindOrAdd(FragmentConfig.FragmentType).Add(NewArchetype);
@@ -215,6 +239,12 @@ FMassArchetypeHandle UMassEntitySubsystem::GetArchetypeForEntity(FMassEntityHand
 	return Result;
 }
 
+FMassArchetypeHandle UMassEntitySubsystem::GetArchetypeForEntityUnsafe(FMassEntityHandle Entity) const
+{
+	check(Entities.IsValidIndex(Entity.Index));
+	return FMassArchetypeHandle(Entities[Entity.Index].CurrentArchetype);
+}
+
 void UMassEntitySubsystem::ForEachArchetypeFragmentType(const FMassArchetypeHandle Archetype, TFunction< void(const UScriptStruct* /*FragmentType*/)> Function)
 {
 	check(Archetype.DataPtr.IsValid());
@@ -416,23 +446,35 @@ void UMassEntitySubsystem::BatchDestroyEntityChunks(const FMassArchetypeSubChunk
 
 	checkf(IsProcessing() == false, TEXT("Synchronous API function %s called during mass processing. Use asynchronous API instead."), ANSI_TO_TCHAR(__FUNCTION__));
 
+	TArray<FMassEntityHandle> EntitiesRemoved;
+	// note that it's important to place the context instance in the same scope as the loop below that updates 
+	// UMassEntitySubsystem.EntityData, otherwise, if there are commands flushed as part of FMassProcessingContext's 
+	// destruction the commands will work on outdated information (which might result in crashes).
 	FMassProcessingContext ProcessingContext(*this, /*TimeDelta=*/0.0f);
-	ProcessingContext.bFlushCommandBuffer = false;
-	ProcessingContext.CommandBuffer = MakeShareable(new FMassCommandBuffer());
-	ObserverManager.OnPreEntitiesDestroyed(ProcessingContext, Chunks);
 
-	check(Chunks.GetArchetype().IsValid());
-	TArray<FMassEntityHandle> EntitiesRemoved;
-	Chunks.GetArchetype().DataPtr->BatchDestroyEntityChunks(Chunks.GetChunks(), EntitiesRemoved);
-	
-	EntityFreeIndexList.Reserve(EntityFreeIndexList.Num() + EntitiesRemoved.Num());
+	bool bValidArchetype = Chunks.GetArchetype().IsValid();
+	if (bValidArchetype)
+	{
+		ProcessingContext.bFlushCommandBuffer = false;
+		ProcessingContext.CommandBuffer = MakeShareable(new FMassCommandBuffer());
+		ObserverManager.OnPreEntitiesDestroyed(ProcessingContext, Chunks);
+
+		check(Chunks.GetArchetype().IsValid());		
+		Chunks.GetArchetype().DataPtr->BatchDestroyEntityChunks(Chunks.GetChunks(), EntitiesRemoved);
+
+		EntityFreeIndexList.Reserve(EntityFreeIndexList.Num() + EntitiesRemoved.Num());
+	}
+	else
+	{
+		UE::Mass::Private::ConvertArchetypelessSubchunksIntoEntityHandles(Chunks.GetChunks(), EntitiesRemoved);
+	}
 
 	for (const FMassEntityHandle& Entity : EntitiesRemoved)
 	{
 		check(Entities.IsValidIndex(Entity.Index));
 
 		FEntityData& EntityData = Entities[Entity.Index];
-		if (EntityData.SerialNumber == Entity.SerialNumber)
+		if (!bValidArchetype || EntityData.SerialNumber == Entity.SerialNumber)
 		{
 			EntityData.Reset();
 			EntityFreeIndexList.Add(Entity.Index);
@@ -703,6 +745,53 @@ void UMassEntitySubsystem::RemoveTagFromEntity(FMassEntityHandle Entity, const U
 	}
 }
 
+void UMassEntitySubsystem::BatchChangeTagsForEntities(TConstArrayView<FMassArchetypeSubChunks> ChunkCollections, const FMassTagBitSet& TagsToAdd, const FMassTagBitSet& TagsToRemove)
+{
+	for (const FMassArchetypeSubChunks& Chunks : ChunkCollections)
+	{
+		FMassArchetypeData* CurrentArchetype = Chunks.GetArchetype().DataPtr.Get();
+		const FMassTagBitSet NewTagComposition = CurrentArchetype
+			? (CurrentArchetype->GetTagBitSet() + TagsToAdd - TagsToRemove)
+			: (TagsToAdd - TagsToRemove);
+
+		if (ensure(CurrentArchetype) && CurrentArchetype->GetTagBitSet() != NewTagComposition)
+		{
+			FMassTagBitSet TagsAdded = TagsToAdd - CurrentArchetype->GetTagBitSet();
+			FMassTagBitSet TagsRemoved = TagsToRemove.GetOverlap(CurrentArchetype->GetTagBitSet());
+
+			if (ObserverManager.HasObserversForBitSet(TagsRemoved, EMassObservedOperation::Remove))
+			{
+				ObserverManager.OnCompositionChanged(Chunks, FMassArchetypeCompositionDescriptor(MoveTemp(TagsRemoved)), EMassObservedOperation::Remove);
+			}
+			const bool bTagsAddedAreObserved = ObserverManager.HasObserversForBitSet(TagsAdded, EMassObservedOperation::Add);
+
+			FMassArchetypeHandle NewArchetypeHandle = InternalCreateSiblingArchetype(Chunks.GetArchetype().DataPtr, NewTagComposition);
+			checkSlow(NewArchetypeHandle.IsValid());
+
+			// Move the entity over
+			FMassArchetypeSubChunks::FSubChunkArray NewArchetypeSubChunks;
+			TArray<FMassEntityHandle> EntitesBeingMoved;
+			CurrentArchetype->BatchMoveEntitiesToAnotherArchetype(Chunks, *NewArchetypeHandle.DataPtr.Get(), EntitesBeingMoved
+				, bTagsAddedAreObserved ? &NewArchetypeSubChunks : nullptr);
+
+			for (const FMassEntityHandle& Entity : EntitesBeingMoved)
+			{
+				check(Entities.IsValidIndex(Entity.Index));
+
+				FEntityData& EntityData = Entities[Entity.Index];
+				EntityData.CurrentArchetype = NewArchetypeHandle.DataPtr;
+			}
+
+			if (bTagsAddedAreObserved)
+			{
+				ObserverManager.OnCompositionChanged(
+					FMassArchetypeSubChunks(NewArchetypeHandle, MoveTemp(NewArchetypeSubChunks))
+					, FMassArchetypeCompositionDescriptor(MoveTemp(TagsAdded))
+					, EMassObservedOperation::Add);
+			}
+		}
+	}
+}
 
 void UMassEntitySubsystem::MoveEntityToAnotherArchetype(FMassEntityHandle Entity, FMassArchetypeHandle NewArchetypeHandle)
 {
@@ -730,11 +819,9 @@ void UMassEntitySubsystem::BatchSetEntityFragmentsValues(const FMassArchetypeSub
 	FMassArchetypeData* Archetype = SparseEntities.GetArchetype().DataPtr.Get();
 	check(Archetype);
 
-	const FMassArchetypeSubChunks::FConstSubChunkArrayView SparseChunks = SparseEntities.GetChunks();
-
 	for (const FInstancedStruct& FragmentTemplate : FragmentInstanceList)
 	{
-		Archetype->SetFragmentData(SparseChunks, FragmentTemplate);
+		Archetype->SetFragmentData(SparseEntities.GetChunks(), FragmentTemplate);
 	}
 }
 
diff --git a/Engine/Plugins/Runtime/MassEntity/Source/MassEntity/Private/MassEntityUtils.cpp b/Engine/Plugins/Runtime/MassEntity/Source/MassEntity/Private/MassEntityUtils.cpp
--- a/Engine/Plugins/Runtime/MassEntity/Source/MassEntity/Private/MassEntityUtils.cpp
+++ b/Engine/Plugins/Runtime/MassEntity/Source/MassEntity/Private/MassEntityUtils.cpp
@@ -43,9 +43,12 @@ void CreateSparseChunks(const UMassEntitySubsystem& EntitySystem, const TConstAr
 
 	for (const FMassEntityHandle& Entity : Entities)
 	{
-		FMassArchetypeHandle Archetype = EntitySystem.GetArchetypeForEntity(Entity);
-		TArray<FMassEntityHandle>& PerArchetypeEntities = ArchetypeToEntities.FindOrAdd(Archetype);
-		PerArchetypeEntities.Add(Entity);
+		if (EntitySystem.IsEntityValid(Entity))
+		{
+			FMassArchetypeHandle Archetype = EntitySystem.GetArchetypeForEntityUnsafe(Entity);
+			TArray<FMassEntityHandle>& PerArchetypeEntities = ArchetypeToEntities.FindOrAdd(Archetype);
+			PerArchetypeEntities.Add(Entity);
+		}
 	}
 
 	for (auto& Pair : ArchetypeToEntities)
diff --git a/Engine/Plugins/Runtime/MassEntity/Source/MassEntity/Public/MassArchetypeTypes.h b/Engine/Plugins/Runtime/MassEntity/Source/MassEntity/Public/MassArchetypeTypes.h
--- a/Engine/Plugins/Runtime/MassEntity/Source/MassEntity/Public/MassArchetypeTypes.h
+++ b/Engine/Plugins/Runtime/MassEntity/Source/MassEntity/Public/MassArchetypeTypes.h
@@ -97,6 +97,10 @@ public:
 	FMassArchetypeSubChunks(const FMassArchetypeHandle& InArchetype, TConstArrayView<FMassEntityHandle> InEntities, EDuplicatesHandling DuplicatesHandling);
 	explicit FMassArchetypeSubChunks(FMassArchetypeHandle& InArchetypeHandle);
 	explicit FMassArchetypeSubChunks(TSharedPtr<FMassArchetypeData>& InArchetype);
+	FMassArchetypeSubChunks(FMassArchetypeHandle& InArchetypeHandle, FSubChunkArray&& InSubChunks)
+		: SubChunks(MoveTemp(InSubChunks))
+		, Archetype(InArchetypeHandle)
+	{}
 
 	FConstSubChunkArrayView GetChunks() const { return SubChunks; }
 	const FMassArchetypeHandle& GetArchetype() const { return Archetype; }
diff --git a/Engine/Plugins/Runtime/MassEntity/Source/MassEntity/Public/MassEntitySubsystem.h b/Engine/Plugins/Runtime/MassEntity/Source/MassEntity/Public/MassEntitySubsystem.h
--- a/Engine/Plugins/Runtime/MassEntity/Source/MassEntity/Public/MassEntitySubsystem.h
+++ b/Engine/Plugins/Runtime/MassEntity/Source/MassEntity/Public/MassEntitySubsystem.h
@@ -106,7 +106,14 @@ public:
 	 */
 	FMassArchetypeHandle CreateArchetype(const TSharedPtr<FMassArchetypeData>& SourceArchetype, const FMassFragmentBitSet& NewFragmentList);
 
+	/** Fetches the archetype for a given Entity. If Entity is not valid it will still return a handle, just with an invalid archetype */
 	FMassArchetypeHandle GetArchetypeForEntity(FMassEntityHandle Entity) const;
+	/**
+	 * Fetches the archetype for a given Entity. Note that it's callers responsibility the given Entity handle is valid.
+	 * If you can't ensure that call GetArchetypeForEntity.
+	 */
+	FMassArchetypeHandle GetArchetypeForEntityUnsafe(FMassEntityHandle Entity) const;
+
 	/** Method to iterate on all the fragment types of an archetype */
 	static void ForEachArchetypeFragmentType(const FMassArchetypeHandle Archetype, TFunction< void(const UScriptStruct* /*FragmentType*/)> Function);
 
@@ -204,6 +211,7 @@ public:
 	void RemoveTagFromEntity(FMassEntityHandle Entity, const UScriptStruct* TagType);
 	void SwapTagsForEntity(FMassEntityHandle Entity, const UScriptStruct* FromFragmentType, const UScriptStruct* ToFragmentType);
 
+	void BatchChangeTagsForEntities(TConstArrayView<FMassArchetypeSubChunks> ChunkCollections, const FMassTagBitSet& TagsToAdd, const FMassTagBitSet& TagsToRemove);
 
 	/**
 	 * Adds fragments and tags indicated by InOutDescriptor to the Entity. The function also figures out which elements
diff --git a/Engine/Plugins/Runtime/MassEntity/Source/MassEntity/Public/MassEntityTypes.h b/Engine/Plugins/Runtime/MassEntity/Source/MassEntity/Public/MassEntityTypes.h
--- a/Engine/Plugins/Runtime/MassEntity/Source/MassEntity/Public/MassEntityTypes.h
+++ b/Engine/Plugins/Runtime/MassEntity/Source/MassEntity/Public/MassEntityTypes.h
@@ -139,6 +139,10 @@ struct FMassArchetypeCompositionDescriptor
 		, SharedFragments(MoveTemp(InSharedFragments))
 	{}
 
+	FMassArchetypeCompositionDescriptor(FMassTagBitSet&& InTags)
+		: Tags(MoveTemp(InTags))
+	{}
+
 	void Reset()
 	{
 		Fragments.Reset();
