New, "batched" commands for manipulating Mass entity tags

Note that LOD and Representation tag determining changes are required due to the new batched commands need to know the tags types at compilation time.

diff --git a/Engine/Plugins/Experimental/StructUtils/Source/StructUtils/Public/ScriptStructTypeBitSet.h b/Engine/Plugins/Experimental/StructUtils/Source/StructUtils/Public/ScriptStructTypeBitSet.h
--- a/Engine/Plugins/Experimental/StructUtils/Source/StructUtils/Public/ScriptStructTypeBitSet.h
+++ b/Engine/Plugins/Experimental/StructUtils/Source/StructUtils/Public/ScriptStructTypeBitSet.h
@@ -463,6 +463,11 @@ public:
 		}
 	}
 
+	SIZE_T GetAllocatedSize() const
+	{
+		return StructTypesBitArray.GetAllocatedSize();
+	}
+
 	FString DebugGetStringDesc() const
 	{
 #if WITH_STRUCTUTILS_DEBUG
diff --git a/Engine/Plugins/Runtime/MassEntity/Source/MassEntity/Private/MassCommandBuffer.cpp b/Engine/Plugins/Runtime/MassEntity/Source/MassEntity/Private/MassCommandBuffer.cpp
--- a/Engine/Plugins/Runtime/MassEntity/Source/MassEntity/Private/MassCommandBuffer.cpp
+++ b/Engine/Plugins/Runtime/MassEntity/Source/MassEntity/Private/MassCommandBuffer.cpp
@@ -11,7 +11,7 @@
 CSV_DEFINE_CATEGORY(MassEntities, true);
 CSV_DEFINE_CATEGORY(MassEntitiesCounters, true);
 
-namespace UE::FLWCCommand {
+namespace UE::Mass::Command {
 
 #if CSV_PROFILER
 bool bEnableDetailedStats = false;
@@ -22,12 +22,16 @@ FAutoConsoleVariableRef CVarEnableDetailedCommandStats(TEXT("massentities.Enable
 /** CSV stat names */
 static FString DefaultName = TEXT("Command");
 static TArray<FName> CommandFNames;
+static FString DefaultBatchedName = TEXT("BatchedCommand");
+static TArray<FName> CommandBatchedFNames;
 
 /** CSV custom stat names (ANSI) */
 static const int32 MaxNameLength = 64;
 typedef ANSICHAR ANSIName[MaxNameLength];
 static ANSIName DefaultANSIName = "Command";
 static TArray<ANSIName> CommandANSINames;
+static ANSIName DefaultANSIBatchedName = "BatchedCommand";
+static TArray<ANSIName> CommandANSIBatchedNames;
 
 /**
  * Provides valid names for CSV profiling.
@@ -64,8 +68,42 @@ void GetCommandStatNames(FStructView Entry, FString& OutName, ANSIName*& OutANSI
 	}
 }
 
+/**
+ * Provides valid names for CSV profiling.
+ * @param Command is the command instance
+ * @param OutName is the name to use for csv custom stats
+ * @param OutANSIName is the name to use for csv stats
+ */
+void GetCommandStatNames(FMassBatchedCommand& Command, FString& OutName, ANSIName*& OutANSIName)
+{
+	OutANSIName = &DefaultANSIBatchedName;
+	OutName = DefaultANSIBatchedName;
+	if (!bEnableDetailedStats)
+	{
+		return;
+	}
+
+	const FName CommandFName = Command.GetFName();
+	OutName = CommandFName.ToString();
+
+	const int32 Index = CommandBatchedFNames.Find(CommandFName);
+	if (Index == INDEX_NONE)
+	{
+		CommandBatchedFNames.Emplace(CommandFName);
+		OutANSIName = &CommandANSIBatchedNames.AddZeroed_GetRef();
+		// Use prefix for easier parsing in reports
+		//const FString CounterName = FString::Printf(TEXT("Num%s"), *OutName);
+		//FMemory::Memcpy(OutANSIName, StringCast<ANSICHAR>(*CounterName).Get(), FMath::Min(CounterName.Len(), MaxNameLength - 1) * sizeof(ANSICHAR));
+		FMemory::Memcpy(OutANSIName, StringCast<ANSICHAR>(*OutName).Get(), FMath::Min(OutName.Len(), MaxNameLength - 1) * sizeof(ANSICHAR));
+	}
+	else
+	{
+		OutANSIName = &CommandANSIBatchedNames[Index];
+	}
+}
+
 #endif
-} // UE::FLWCCommand
+} // UE::Mass::Command
 
 
 //////////////////////////////////////////////////////////////////////
@@ -107,6 +145,10 @@ void FMassCommandsObservedTypes::Append(const FMassCommandsObservedTypes& Other)
 	}
 }
 
+//////////////////////////////////////////////////////////////////////
+// FMassBatchedCommand
+std::atomic<uint32> FMassBatchedCommand::CommandsCounter;
+
 //////////////////////////////////////////////////////////////////////
 // FMassCommandBuffer
 
@@ -121,7 +163,7 @@ void FMassCommandBuffer::Flush(UMassEntitySubsystem& EntitySystem)
 	TGuardValue FlushingGuard(bIsFlushing, true);
 
 	// short-circuit exit
-	if (EntitiesToDestroy.Num() == 0 && PendingCommands.IsEmpty())
+	if (EntitiesToDestroy.Num() == 0 && PendingCommands.IsEmpty() && CommandInstances.IsEmpty())
 	{
 		return;
 	}
@@ -160,43 +202,70 @@ void FMassCommandBuffer::Flush(UMassEntitySubsystem& EntitySystem)
 			}
 		}
 	}
-	
-	TArray<FMassArchetypeSubChunks> EntityChunksToDestroy;
-	if (EntitiesToDestroy.Num())
-	{
-		UE::Mass::Utils::CreateSparseChunks(EntitySystem, EntitiesToDestroy, FMassArchetypeSubChunks::FoldDuplicates, EntityChunksToDestroy);
-		for (FMassArchetypeSubChunks& Collection : EntityChunksToDestroy)
-		{
-			EntitySystem.BatchDestroyEntityChunks(Collection);
-		}
-	}
-	EntitiesToDestroy.Reset();
 
-
-	UE_MT_SCOPED_WRITE_ACCESS(PendingCommandsDetector);
-	PendingCommands.ForEach([&EntitySystem](FStructView Entry)
 	{
-		const FCommandBufferEntryBase* Command = Entry.GetPtr<FCommandBufferEntryBase>();
-		checkf(Command, TEXT("Either the entry is null or the command does not derive from FCommandBufferEntryBase"));
+		UE_MT_SCOPED_WRITE_ACCESS(PendingCommandsDetector);
+		PendingCommands.ForEach([&EntitySystem](FStructView Entry)
+			{
+				const FCommandBufferEntryBase* Command = Entry.GetPtr<FCommandBufferEntryBase>();
+				checkf(Command, TEXT("Either the entry is null or the command does not derive from FCommandBufferEntryBase"));
 
 #if CSV_PROFILER
-		using namespace UE::FLWCCommand;
+				using namespace UE::Mass::Command;
 
-		// Extract name (default or detailed)
-		ANSIName* ANSIName = &DefaultANSIName;
-		FString Name = DefaultName;
-		GetCommandStatNames(Entry, Name, ANSIName);
+				// Extract name (default or detailed)
+				ANSIName* ANSIName = &DefaultANSIName;
+				FString Name = DefaultName;
+				GetCommandStatNames(Entry, Name, ANSIName);
 
-		// Push stats
-		FScopedCsvStat ScopedCsvStat(*ANSIName, CSV_CATEGORY_INDEX(MassEntities));
-		FCsvProfiler::RecordCustomStat(*Name, CSV_CATEGORY_INDEX(MassEntitiesCounters), 1, ECsvCustomStatOp::Accumulate);
+				// Push stats
+				FScopedCsvStat ScopedCsvStat(*ANSIName, CSV_CATEGORY_INDEX(MassEntities));
+				FCsvProfiler::RecordCustomStat(*Name, CSV_CATEGORY_INDEX(MassEntitiesCounters), 1, ECsvCustomStatOp::Accumulate);
 #endif // CSV_PROFILER
 
-		Command->Execute(EntitySystem);
-	});
+				Command->Execute(EntitySystem);
+			});
+
+		// Using Clear() instead of Reset(), as otherwise the chunks moved into the PendingCommands in MoveAppend() can accumulate.
+		PendingCommands.Clear();
+	}
+
+	{
+		UE_MT_SCOPED_WRITE_ACCESS(PendingBatchCommandsDetector);
+		for (FMassBatchedCommand* Command : CommandInstances)
+		{
+			if (Command && Command->HasWork())
+			{
+#if CSV_PROFILER
+				using namespace UE::Mass::Command;
 
-	// Using Clear() instead of Reset(), as otherwise the chunks moved into the PendingCommands in MoveAppend() can accumulate.
-	PendingCommands.Clear();
+				// Extract name (default or detailed)
+				ANSIName* ANSIName = &DefaultANSIName;
+				FString Name = DefaultName;
+				GetCommandStatNames(*Command, Name, ANSIName);
+
+				// Push stats
+				FScopedCsvStat ScopedCsvStat(*ANSIName, CSV_CATEGORY_INDEX(MassEntities));
+				FCsvProfiler::RecordCustomStat(*Name, CSV_CATEGORY_INDEX(MassEntitiesCounters), Command->GetNumEntitiesStat(), ECsvCustomStatOp::Accumulate);
+#endif // CSV_PROFILER
+				Command->Execute(EntitySystem);
+				Command->Reset();
+			}
+		}
+		ObservedTypes.Reset();
+	}
+
+	// note that entity destruction has been moved to the end to avoid the burden of checking if an entity has been destroyed with every operation
+	TArray<FMassArchetypeSubChunks > EntityChunksToDestroy;
+	if (EntitiesToDestroy.Num())
+	{
+		UE::Mass::Utils::CreateSparseChunks(EntitySystem, EntitiesToDestroy, FMassArchetypeSubChunks::FoldDuplicates, EntityChunksToDestroy);
+		for (FMassArchetypeSubChunks& Collection : EntityChunksToDestroy)
+		{
+			EntitySystem.BatchDestroyEntityChunks(Collection);
+		}
+	}
+	EntitiesToDestroy.Reset();
 
 	for (auto It : ObservedTypes.GetObservedFragments(EMassObservedOperation::Add))
 	{
@@ -227,8 +296,6 @@ void FMassCommandBuffer::Flush(UMassEntitySubsystem& EntitySystem)
 			}
 		}
 	}
-
-	ObservedTypes.Reset();
 }
  
 void FMassCommandBuffer::CleanUp()
@@ -248,10 +315,27 @@ void FMassCommandBuffer::MoveAppend(FMassCommandBuffer& Other)
 		UE_MT_SCOPED_WRITE_ACCESS(PendingCommandsDetector);
 		PendingCommands.Append(MoveTemp(Other.PendingCommands));
 		EntitiesToDestroy.Append(MoveTemp(Other.EntitiesToDestroy));
+		CommandInstances.Append(MoveTemp(Other.CommandInstances));
 		ObservedTypes.Append(Other.ObservedTypes);
 	}
 }
 
+SIZE_T FMassCommandBuffer::GetAllocatedSize() const
+{
+	SIZE_T TotalSize = 0;
+	for (FMassBatchedCommand* Command : CommandInstances)
+	{
+		TotalSize += Command ? Command->GetAllocatedSize() : 0;
+	}
+
+	TotalSize += PendingCommands.GetAllocatedSize();
+	TotalSize += EntitiesToDestroy.GetAllocatedSize();
+	TotalSize += ObservedTypes.GetAllocatedSize();
+	TotalSize += CommandInstances.GetAllocatedSize();
+	
+	return TotalSize;
+}
+
 //////////////////////////////////////////////////////////////////////
 // Command implementations
 
@@ -265,27 +349,6 @@ void FMassCommandAddFragmentInstanceList::Execute(UMassEntitySubsystem& System)
 	System.AddFragmentInstanceListToEntity(TargetEntity, FragmentList);
 }
 
-void FCommandSwapTags::Execute(UMassEntitySubsystem& System) const
-{
-	if (System.IsEntityValid(TargetEntity) == false)
-	{
-		return;
-}
-
-	if (OldTagType && NewTagType)
-	{
-		System.SwapTagsForEntity(TargetEntity, OldTagType, NewTagType);
-	}
-	else if (OldTagType)
-	{
-		System.RemoveTagFromEntity(TargetEntity, OldTagType);
-	}
-	else if (NewTagType)
-	{
-		System.AddTagToEntity(TargetEntity, NewTagType);
-	}
-}
-
 void FBuildEntityFromFragmentInstance::Execute(UMassEntitySubsystem& System) const
 {
 	System.BuildEntity(TargetEntity, MakeArrayView(&Struct, 1), SharedFragmentValues);
@@ -322,21 +385,6 @@ void FCommandRemoveFragmentList::Execute(UMassEntitySubsystem& System) const
 	System.RemoveFragmentListFromEntity(TargetEntity, FragmentList);
 }
 
-void FCommandAddTag::Execute(UMassEntitySubsystem & System) const 
-{
-	if (System.IsEntityValid(TargetEntity) == false)
-	{
-		return;
-	}
-
-	System.AddTagToEntity(TargetEntity, StructParam);
-}
-
-void FCommandRemoveTag::Execute(UMassEntitySubsystem& System) const
-{
-	System.RemoveTagFromEntity(TargetEntity, StructParam);
-}
-
 void FCommandRemoveComposition::Execute(UMassEntitySubsystem& System) const
 {
 	if (System.IsEntityValid(TargetEntity) == false)
diff --git a/Engine/Plugins/Runtime/MassEntity/Source/MassEntity/Public/MassCommandBuffer.h b/Engine/Plugins/Runtime/MassEntity/Source/MassEntity/Public/MassCommandBuffer.h
--- a/Engine/Plugins/Runtime/MassEntity/Source/MassEntity/Public/MassCommandBuffer.h
+++ b/Engine/Plugins/Runtime/MassEntity/Source/MassEntity/Public/MassCommandBuffer.h
@@ -6,7 +6,8 @@
 #include "MassEntitySubsystem.h"
 #include "InstancedStructStream.h"
 #include "Misc/MTAccessDetector.h"
-
+#include "MassEntityUtils.h"
+#include "MassCommands.h"
 #include "MassCommandBuffer.generated.h"
 
 
@@ -48,6 +49,11 @@ struct FMassObservedTypeCollection
 		return Types; 
 	}
 
+	SIZE_T GetAllocatedSize() const
+	{
+		return LastUsedCollection->GetAllocatedSize() + Types.GetAllocatedSize();
+	}
+
 private:
 	const UScriptStruct* LastAddedType = nullptr;
 	TArray<FMassEntityHandle>* LastUsedCollection = nullptr;
@@ -90,6 +96,17 @@ struct FMassCommandsObservedTypes
 
 	void Append(const FMassCommandsObservedTypes& Other);
 
+	SIZE_T GetAllocatedSize() const
+	{
+		SIZE_T Total = 0;
+		for (int i = 0; i < (int)EMassObservedOperation::MAX; ++i)
+		{
+			Total += Fragments[i].GetAllocatedSize();
+			Total += Tags[i].GetAllocatedSize();
+		}
+		return Total;
+	}
+
 protected:
 	FMassObservedTypeCollection Fragments[(uint8)EMassObservedOperation::MAX];
 	FMassObservedTypeCollection Tags[(uint8)EMassObservedOperation::MAX];
@@ -400,109 +417,6 @@ protected:
 	TArray<const UScriptStruct*> FragmentList;
 };
 
-/**
- * Command dedicated to add a new tag to an existing entity
- */
-USTRUCT()
-struct MASSENTITY_API FCommandAddTag: public FCommandBufferEntryBase
-{
-	GENERATED_BODY()
-
-	enum
-	{
-		Type = ECommandBufferOperationType::Add
-	};
-
-	FCommandAddTag() = default;
-	FCommandAddTag(FMassEntityHandle InEntity, UScriptStruct* InStruct)
-		: FCommandBufferEntryBase(InEntity)
-		, StructParam(InStruct)
-	{}
-
-	void AppendAffectedEntitiesPerType(FMassCommandsObservedTypes& ObservedTypes)
-	{
-		ObservedTypes.TagAdded(StructParam, TargetEntity);
-	}
-
-protected:
-	virtual void Execute(UMassEntitySubsystem& System) const override;
-
-	UScriptStruct* StructParam = nullptr;
-};
-
-/**
- * Command dedicated to remove a tag from an existing entity
- */
-USTRUCT()
-struct MASSENTITY_API FCommandRemoveTag : public FCommandBufferEntryBase
-{
-	GENERATED_BODY()
-
-	enum
-	{
-		Type = ECommandBufferOperationType::Remove
-	};
-
-	FCommandRemoveTag() = default;
-	FCommandRemoveTag(FMassEntityHandle InEntity, UScriptStruct* InStruct)
-		: FCommandBufferEntryBase(InEntity)
-		, StructParam(InStruct)
-	{}
-
-	void AppendAffectedEntitiesPerType(FMassCommandsObservedTypes& ObservedTypes)
-	{
-		ObservedTypes.TagRemoved(StructParam, TargetEntity);
-	}
-
-protected:
-	virtual void Execute(UMassEntitySubsystem& System) const override;
-
-	UScriptStruct* StructParam = nullptr;
-};
-
-/**
- * Command for swapping given tags for a given entity. Note that the entity doesn't need to own the "old tag", the 
- * "new tag" will be added regardless.
- */
-USTRUCT()
-struct MASSENTITY_API FCommandSwapTags : public FCommandBufferEntryBase
-{
-	GENERATED_BODY()
-
-	enum
-	{
-		Type = ECommandBufferOperationType::Add | ECommandBufferOperationType::Remove
-	};
-
-	FCommandSwapTags() = default;
-	FCommandSwapTags(const FMassEntityHandle InEntity, const UScriptStruct* InOldTagType, const UScriptStruct* InNewTagType)
-		: FCommandBufferEntryBase(InEntity)
-		, OldTagType(InOldTagType)
-		, NewTagType(InNewTagType)
-	{
-		checkf((InOldTagType == nullptr) || InOldTagType->IsChildOf(FMassTag::StaticStruct()), TEXT("FCommandSwapTags works only with tags while '%s' is not one."), *GetPathNameSafe(InOldTagType));
-		checkf((InNewTagType == nullptr) || InNewTagType->IsChildOf(FMassTag::StaticStruct()), TEXT("FCommandSwapTags works only with tags while '%s' is not one."), *GetPathNameSafe(InNewTagType));
-	}
-
-	void AppendAffectedEntitiesPerType(FMassCommandsObservedTypes& ObservedTypes)
-	{
-		if (OldTagType)
-		{
-			ObservedTypes.TagRemoved(OldTagType, TargetEntity);
-		}
-		if (NewTagType)
-		{
-			ObservedTypes.TagAdded(NewTagType, TargetEntity);
-		}
-	}
-
-protected:
-	virtual void Execute(UMassEntitySubsystem& System) const override;
-
-	const UScriptStruct* OldTagType = nullptr;
-	const UScriptStruct* NewTagType = nullptr;
-};
-
 /**
  * Command performing a removal of a collection of fragments and tags
  */
@@ -575,8 +489,24 @@ public:
 		}
 		return Command;
 	}
-	
-public:
+
+	/** Adds a new entry to a given TCommand batch command instance */
+	template< template<typename... TArgs> typename TCommand, typename... TArgs >
+	void PushCommand(const FMassEntityHandle Entity, TArgs&&... InArgs)
+	{
+		UE_MT_SCOPED_WRITE_ACCESS(PendingCommandsDetector);
+		TCommand<TArgs...>& Instance = CreateOrAddCommand<TCommand<TArgs...>>();
+		Instance.Add(Entity, Forward<TArgs>(InArgs)...);
+	}
+
+	/** Adds a new entry to a given TCommand batch command instance */
+	template< typename TCommand, typename = typename TEnableIf<TIsDerivedFrom<typename TRemoveReference<TCommand>::Type, FMassBatchedCommand>::IsDerived, void>::Type>
+	void PushCommand(const FMassEntityHandle Entity)
+	{
+		UE_MT_SCOPED_WRITE_ACCESS(PendingCommandsDetector);
+		CreateOrAddCommand<TCommand>().Add(Entity);
+	}
+
 	template<typename T>
 	void AddFragment(FMassEntityHandle Entity)
 	{
@@ -591,18 +521,29 @@ public:
 		EmplaceCommand<FCommandRemoveFragment>(Entity, T::StaticStruct());
 	}
 
+	/** the convenience function equivalent to calling PushCommand<FCommandAddTag<T>>(Entity) */
 	template<typename T>
 	void AddTag(FMassEntityHandle Entity)
 	{
 		static_assert(TIsDerivedFrom<T, FMassTag>::IsDerived, "Given struct type is not a valid tag type.");
-		EmplaceCommand<FCommandAddTag>(Entity, T::StaticStruct());
+		PushCommand<FCommandAddTag<T>>(Entity);
 	}
 
+	/** the convenience function equivalent to calling PushCommand<FCommandRemoveTag<T>>(Entity) */
 	template<typename T>
 	void RemoveTag(FMassEntityHandle Entity)
 	{
 		static_assert(TIsDerivedFrom<T, FMassTag>::IsDerived, "Given struct type is not a valid tag type.");
-		EmplaceCommand<FCommandRemoveTag>(Entity, T::StaticStruct());
+		PushCommand<FCommandRemoveTag<T>>(Entity);
+	}
+
+	/** the convenience function equivalent to calling PushCommand<FCommandSwapTags<TOld, TNew>>(Entity)  */
+	template<typename TOld, typename TNew>
+	void SwapTags(FMassEntityHandle Entity)
+	{
+		static_assert(TIsDerivedFrom<TOld, FMassTag>::IsDerived, "Given struct type is not a valid tag type.");
+		static_assert(TIsDerivedFrom<TNew, FMassTag>::IsDerived, "Given struct type is not a valid tag type.");
+		PushCommand<FCommandSwapTags<TOld, TNew>>(Entity);
 	}
 
 	void DestroyEntity(FMassEntityHandle Entity)
@@ -615,7 +556,7 @@ public:
 		EntitiesToDestroy.Append(InEntitiesToDestroy);
 	}
 
-	SIZE_T GetAllocatedSize() const { return PendingCommands.GetAllocatedSize(); }
+	SIZE_T GetAllocatedSize() const;
 
 	/** 
 	 * Appends the commands from the passed buffer into this one
@@ -629,6 +570,28 @@ public:
 
 private:
 	friend UMassEntitySubsystem;
+
+	template<typename T>
+	T& CreateOrAddCommand()
+	{
+		const int32 Index = FMassBatchedCommand::GetCommandIndex<T>();
+
+		UE_MT_SCOPED_WRITE_ACCESS(PendingBatchCommandsDetector);
+		if (CommandInstances.IsValidIndex(Index) == false)
+		{
+			CommandInstances.AddZeroed(Index - CommandInstances.Num() + 1);
+		}
+		else if (CommandInstances[Index])
+		{
+			return (T&)(*CommandInstances[Index]);
+		}
+
+		T* NewCommandInstance = new T();
+		check(NewCommandInstance);
+		CommandInstances[Index] = NewCommandInstance;
+		return *NewCommandInstance;
+	}
+
 	void Flush(UMassEntitySubsystem& EntitySystem);
 	void CleanUp();
 
@@ -640,6 +603,10 @@ private:
 
 	FMassCommandsObservedTypes ObservedTypes;
 
+	UE_MT_DECLARE_RW_ACCESS_DETECTOR(PendingBatchCommandsDetector);
+	// @todo consider using sparse array
+	TArray<FMassBatchedCommand*> CommandInstances;
+
 	/** Indicates that this specific MassCommandBuffer is currently flushing its contents */
 	bool bIsFlushing = false;
 };
diff --git a/Engine/Plugins/Runtime/MassEntity/Source/MassEntity/Public/MassCommands.h b/Engine/Plugins/Runtime/MassEntity/Source/MassEntity/Public/MassCommands.h
new file mode 100644
--- /dev/null
+++ b/Engine/Plugins/Runtime/MassEntity/Source/MassEntity/Public/MassCommands.h
@@ -0,0 +1,116 @@
+// Copyright Epic Games, Inc. All Rights Reserved.
+
+#pragma once
+
+#include "ProfilingDebugging/CsvProfilerConfig.h"
+#include "MassEntityTypes.h"
+#include "MassEntitySubsystem.h"
+
+#define WITH_MASS_BATCHED_COMMANDS 1
+
+struct FMassBatchedCommand
+{
+	virtual ~FMassBatchedCommand() {}
+
+	virtual void Execute(UMassEntitySubsystem& System) const = 0;
+	virtual void Reset()
+	{
+		UE_MT_SCOPED_WRITE_ACCESS(EntitiesAccessDetector);
+		TargetEntities.Reset();
+	}
+
+	void Add(FMassEntityHandle Entity)
+	{
+		UE_MT_SCOPED_WRITE_ACCESS(EntitiesAccessDetector);
+		TargetEntities.Add(Entity);
+	}
+
+	bool HasWork() const 
+	{ 
+		UE_MT_SCOPED_READ_ACCESS(EntitiesAccessDetector);
+		return TargetEntities.Num() > 0; 
+	}
+
+	template<typename T>
+	FORCENOINLINE static uint32 GetCommandIndex()
+	{
+		static const uint32 ThisTypesStaticIndex = CommandsCounter++;
+		return ThisTypesStaticIndex;
+	}
+
+	virtual SIZE_T GetAllocatedSize() const
+	{
+		return TargetEntities.GetAllocatedSize();
+	}
+
+#if CSV_PROFILER || WITH_MASSENTITY_DEBUG
+	int32 GetNumEntitiesStat() const  { return TargetEntities.Num(); }
+	FName GetFName() const { return DebugName; }
+#endif // CSV_PROFILER || WITH_MASSENTITY_DEBUG
+
+protected:
+	UE_MT_DECLARE_RW_ACCESS_DETECTOR(EntitiesAccessDetector); 
+	TArray<FMassEntityHandle> TargetEntities;
+#if CSV_PROFILER || WITH_MASSENTITY_DEBUG
+	FName DebugName;
+#endif // CSV_PROFILER || WITH_MASSENTITY_DEBUG
+private:
+	MASSENTITY_API static std::atomic<uint32> CommandsCounter;
+};
+
+struct FMassBatchedCommandChangeTags : public FMassBatchedCommand
+{
+protected:
+	virtual void Execute(UMassEntitySubsystem& System) const override
+	{
+		TArray<FMassArchetypeSubChunks> ChunkCollections;
+		UE::Mass::Utils::CreateSparseChunks(System, TargetEntities, FMassArchetypeSubChunks::FoldDuplicates, ChunkCollections);
+
+		System.BatchChangeTagsForEntities(ChunkCollections, TagsToAdd, TagsToRemove);
+	}
+
+	virtual SIZE_T GetAllocatedSize() const override
+	{
+		return TagsToAdd.GetAllocatedSize() + TagsToRemove.GetAllocatedSize() + FMassBatchedCommand::GetAllocatedSize();
+	}
+
+	FMassTagBitSet TagsToAdd;
+	FMassTagBitSet TagsToRemove;
+};
+
+template<typename T>
+struct FCommandAddTag : public FMassBatchedCommandChangeTags
+{
+	FCommandAddTag()
+	{
+		TagsToAdd.Add<T>();
+#if CSV_PROFILER || WITH_MASSENTITY_DEBUG
+		DebugName = TEXT("CommandAddTag");
+#endif // CSV_PROFILER || WITH_MASSENTITY_DEBUG
+	}
+};
+
+template<typename T>
+struct FCommandRemoveTag : public FMassBatchedCommandChangeTags
+{
+	FCommandRemoveTag()
+	{
+		TagsToRemove.Add<T>();
+#if CSV_PROFILER || WITH_MASSENTITY_DEBUG
+		DebugName = TEXT("RemoveAddTag");
+#endif // CSV_PROFILER || WITH_MASSENTITY_DEBUG
+	}
+};
+
+template<typename TOld, typename TNew>
+struct FCommandSwapTags : public FMassBatchedCommandChangeTags
+{
+	FCommandSwapTags()
+	{
+		TagsToRemove.Add<TOld>();
+		TagsToAdd.Add<TNew>();
+#if CSV_PROFILER || WITH_MASSENTITY_DEBUG
+		DebugName = TEXT("SwapAddTags");
+#endif // CSV_PROFILER || WITH_MASSENTITY_DEBUG
+	}
+};
diff --git a/Engine/Plugins/Runtime/MassEntity/Source/MassEntityTestSuite/Private/MassObserverTest.cpp b/Engine/Plugins/Runtime/MassEntity/Source/MassEntityTestSuite/Private/MassObserverTest.cpp
--- a/Engine/Plugins/Runtime/MassEntity/Source/MassEntityTestSuite/Private/MassObserverTest.cpp
+++ b/Engine/Plugins/Runtime/MassEntity/Source/MassEntityTestSuite/Private/MassObserverTest.cpp
@@ -206,7 +206,7 @@ struct FTag_MultipleArchetypeSwap : FTagBaseOperation
 		AITEST_EQUAL(TEXT("Tag addition is not being observed and is not expected to produce results yet"), AffectedEntities.Num(), 0);
 		for (const FMassEntityHandle& ModifiedEntity : ExpectedEntities)
 		{
-			EntitySubsystem->Defer().PushCommand(FCommandSwapTags(ModifiedEntity, FTestTag_A::StaticStruct(), FTestTag_B::StaticStruct()));
+			EntitySubsystem->Defer().SwapTags<FTestTag_A, FTestTag_B>(ModifiedEntity);
 		}
 		return true;
 	}
diff --git a/Engine/Plugins/Runtime/MassGameplay/Source/MassLOD/Private/MassLODUtils.cpp b/Engine/Plugins/Runtime/MassGameplay/Source/MassLOD/Private/MassLODUtils.cpp
--- a/Engine/Plugins/Runtime/MassGameplay/Source/MassLOD/Private/MassLODUtils.cpp
+++ b/Engine/Plugins/Runtime/MassGameplay/Source/MassLOD/Private/MassLODUtils.cpp
@@ -1,6 +1,7 @@
 // Copyright Epic Games, Inc. All Rights Reserved.
 
 #include "MassLODUtils.h"
+#include "MassCommandBuffer.h"
 
 DEFINE_LOG_CATEGORY(LogMassLOD);
 #if WITH_MASSGAMEPLAY_DEBUG
@@ -10,4 +11,86 @@ namespace UE::MassLOD::Debug
 	
 	FAutoConsoleVariableRef CVarLODPause(TEXT("mass.lod.pause"), bLODCalculationsPaused, TEXT("If non zero will pause all LOD calculations"));
 } // UE::MassLOD::Debug
-#endif // WITH_MASSGAMEPLAY_DEBUG
\ No newline at end of file
+#endif // WITH_MASSGAMEPLAY_DEBUG
+
+namespace UE::MassLOD
+{
+void PushSwapTagsCommand(FMassCommandBuffer& CommandBuffer, const FMassEntityHandle Entity, const EMassLOD::Type PrevLOD, const EMassLOD::Type NewLOD)
+{
+#define CASE_SWAP_TAGS(OldLOD, NewLOD) \
+case NewLOD: \
+	CommandBuffer.SwapTags<TMassLODTagForLevel<OldLOD>::FTag, TMassLODTagForLevel<NewLOD>::FTag>(Entity); \
+	break
+
+#define CASE_ADD_TAG(NewLOD) \
+case NewLOD: \
+	CommandBuffer.AddTag<TMassLODTagForLevel<NewLOD>::FTag>(Entity); \
+	break
+
+#define DEFAULT_REMOVE_TAG(OldLOD) \
+case EMassLOD::Max: /* fall through on purpose */ \
+default: \
+	CommandBuffer.RemoveTag<TMassLODTagForLevel<OldLOD>::FTag>(Entity); \
+	break
+
+	check(PrevLOD != NewLOD);
+
+	switch (PrevLOD)
+	{
+	case EMassLOD::High:
+		switch (NewLOD)
+		{
+			CASE_SWAP_TAGS(EMassLOD::High, EMassLOD::Medium);
+			CASE_SWAP_TAGS(EMassLOD::High, EMassLOD::Low);
+			CASE_SWAP_TAGS(EMassLOD::High, EMassLOD::Off);
+			DEFAULT_REMOVE_TAG(EMassLOD::High);
+		}
+		break;
+	case EMassLOD::Medium:
+		switch (NewLOD)
+		{
+			CASE_SWAP_TAGS(EMassLOD::Medium, EMassLOD::High);
+			CASE_SWAP_TAGS(EMassLOD::Medium, EMassLOD::Low);
+			CASE_SWAP_TAGS(EMassLOD::Medium, EMassLOD::Off);
+			DEFAULT_REMOVE_TAG(EMassLOD::Medium);
+		}
+		break;
+	case EMassLOD::Low:
+		switch (NewLOD)
+		{
+			CASE_SWAP_TAGS(EMassLOD::Low, EMassLOD::High);
+			CASE_SWAP_TAGS(EMassLOD::Low, EMassLOD::Medium);
+			CASE_SWAP_TAGS(EMassLOD::Low, EMassLOD::Off);
+			DEFAULT_REMOVE_TAG(EMassLOD::Low);
+		}
+		break;
+	case EMassLOD::Off:
+		switch (NewLOD)
+		{
+			CASE_SWAP_TAGS(EMassLOD::Off, EMassLOD::High);
+			CASE_SWAP_TAGS(EMassLOD::Off, EMassLOD::Medium);
+			CASE_SWAP_TAGS(EMassLOD::Off, EMassLOD::Low);
+			DEFAULT_REMOVE_TAG(EMassLOD::Off);
+		}
+		break;
+	case EMassLOD::Max:
+		switch (NewLOD)
+		{
+			CASE_ADD_TAG(EMassLOD::High);
+			CASE_ADD_TAG(EMassLOD::Medium);
+			CASE_ADD_TAG(EMassLOD::Low);
+			CASE_ADD_TAG(EMassLOD::Off);
+			default:
+				checkf(false, TEXT("Unsupported LOD types!"));
+				break;
+		}
+		break;
+	default:
+		checkf(false, TEXT("Unsupported LOD type!"));
+		break;
+	}
+
+#undef CASE_SWAP_TAGS
+#undef DEFAULT_REMOVE_TAG
+}
+}
\ No newline at end of file
diff --git a/Engine/Plugins/Runtime/MassGameplay/Source/MassLOD/Private/MassSimulationLOD.cpp b/Engine/Plugins/Runtime/MassGameplay/Source/MassLOD/Private/MassSimulationLOD.cpp
--- a/Engine/Plugins/Runtime/MassGameplay/Source/MassLOD/Private/MassSimulationLOD.cpp
+++ b/Engine/Plugins/Runtime/MassGameplay/Source/MassLOD/Private/MassSimulationLOD.cpp
@@ -174,8 +174,8 @@ void UMassSimulationLODProcessor::Execute(UMassEntitySubsystem& EntitySubsystem,
 				const FMassSimulationLODFragment& EntityLOD = SimulationLODFragments[Index];
 				if (EntityLOD.PrevLOD != EntityLOD.LOD)
 				{
-					FMassEntityHandle Entity = Context.GetEntity(Index);
-					Context.Defer().PushCommand(FCommandSwapTags(Entity, UE::MassLOD::GetLODTagFromLOD(EntityLOD.PrevLOD), UE::MassLOD::GetLODTagFromLOD(EntityLOD.LOD)));
+					const FMassEntityHandle Entity = Context.GetEntity(Index);
+					UE::MassLOD::PushSwapTagsCommand(Context.Defer(), Entity, EntityLOD.PrevLOD, EntityLOD.LOD);
 				}
 			}
 		});
@@ -193,4 +193,4 @@ void UMassSimulationLODProcessor::Execute(UMassEntitySubsystem& EntitySubsystem,
 			LODSharedFragment.LODCalculator.DebugDisplayLOD(Context, SimulationLODList, LocationList, World);
 		});
 	}
-}
\ No newline at end of file
+}
diff --git a/Engine/Plugins/Runtime/MassGameplay/Source/MassLOD/Public/MassLODTickRateController.h b/Engine/Plugins/Runtime/MassGameplay/Source/MassLOD/Public/MassLODTickRateController.h
--- a/Engine/Plugins/Runtime/MassGameplay/Source/MassLOD/Public/MassLODTickRateController.h
+++ b/Engine/Plugins/Runtime/MassGameplay/Source/MassLOD/Public/MassLODTickRateController.h
@@ -152,8 +152,8 @@ bool TMassLODTickRateController<TVariableTickChunkFragment, FLODLogic>::UpdateTi
 			TickRate.LastTickedTime = Time;
 			if (EntityLOD.LOD != ChunkLOD)
 			{
-				FMassEntityHandle Entity = Context.GetEntity(Index);
-				Context.Defer().PushCommand(FCommandSwapTags(Entity, UE::MassLOD::GetLODTagFromLOD(ChunkLOD), UE::MassLOD::GetLODTagFromLOD(EntityLOD.LOD)));
+				const FMassEntityHandle Entity = Context.GetEntity(Index);
+				UE::MassLOD::PushSwapTagsCommand(Context.Defer(), Entity, ChunkLOD, EntityLOD.LOD);
 			}
 		}
 	}
diff --git a/Engine/Plugins/Runtime/MassGameplay/Source/MassLOD/Public/MassLODUtils.h b/Engine/Plugins/Runtime/MassGameplay/Source/MassLOD/Public/MassLODUtils.h
--- a/Engine/Plugins/Runtime/MassGameplay/Source/MassLOD/Public/MassLODUtils.h
+++ b/Engine/Plugins/Runtime/MassGameplay/Source/MassLOD/Public/MassLODUtils.h
@@ -10,6 +10,8 @@
 
 MASSLOD_API DECLARE_LOG_CATEGORY_EXTERN(LogMassLOD, Log, All);
 
+struct FMassCommandBuffer;
+
 namespace UE::MassLOD 
 {
 
@@ -34,18 +36,43 @@ inline EMassLOD::Type GetLODFromArchetype(const FMassExecutionContext& Context)
 	return EMassLOD::Max;
 }
 
+template <EMassLOD::Type Level>
+struct TMassLODTagForLevel
+{
+	typedef FMassOffLODTag FTag;
+};
+
+template<>
+struct TMassLODTagForLevel<EMassLOD::High>
+{
+	typedef FMassHighLODTag FTag;
+};
+
+template<>
+struct TMassLODTagForLevel<EMassLOD::Medium>
+{
+	typedef FMassMediumLODTag FTag;
+};
+
+template<>
+struct TMassLODTagForLevel<EMassLOD::Low>
+{
+	typedef FMassLowLODTag FTag;
+};
+
+
 inline const UScriptStruct* GetLODTagFromLOD(EMassLOD::Type LOD)
 {
 	switch (LOD)
 	{
 		case EMassLOD::High:
-			return FMassHighLODTag::StaticStruct();
+			return TMassLODTagForLevel<EMassLOD::High>::FTag::StaticStruct();
 		case EMassLOD::Medium:
-			return FMassMediumLODTag::StaticStruct();
+			return TMassLODTagForLevel<EMassLOD::Medium>::FTag::StaticStruct();
 		case EMassLOD::Low:
-			return FMassLowLODTag::StaticStruct();
+			return TMassLODTagForLevel<EMassLOD::Low>::FTag::StaticStruct();
 		case EMassLOD::Off:
-			return FMassOffLODTag::StaticStruct();
+			return TMassLODTagForLevel<EMassLOD::Off>::FTag::StaticStruct();
 		default:
 			checkf(false, TEXT("Unsupported LOD Type"));
 		case EMassLOD::Max:
@@ -72,6 +99,8 @@ inline bool IsLODTagSet(const FMassExecutionContext& Context, EMassLOD::Type LOD
 	}
 }
 
+void MASSLOD_API PushSwapTagsCommand(FMassCommandBuffer& CommandBuffer, const FMassEntityHandle Entity, const EMassLOD::Type PrevLOD, const EMassLOD::Type NewLOD);
+
 #if WITH_MASSGAMEPLAY_DEBUG
 namespace Debug 
 {
diff --git a/Engine/Plugins/Runtime/MassGameplay/Source/MassRepresentation/Private/MassRepresentationProcessor.cpp b/Engine/Plugins/Runtime/MassGameplay/Source/MassRepresentation/Private/MassRepresentationProcessor.cpp
--- a/Engine/Plugins/Runtime/MassGameplay/Source/MassRepresentation/Private/MassRepresentationProcessor.cpp
+++ b/Engine/Plugins/Runtime/MassGameplay/Source/MassRepresentation/Private/MassRepresentationProcessor.cpp
@@ -346,7 +346,7 @@ void UMassVisualizationProcessor::UpdateEntityVisibility(const FMassEntityHandle
 	const EMassVisibility ChunkVisibility = ChunkData.GetVisibility();
 	if (ChunkVisibility != Visibility)
 	{
-		CommandBuffer.PushCommand(FCommandSwapTags(Entity, UE::MassRepresentation::GetTagFromVisibility(ChunkVisibility), UE::MassRepresentation::GetTagFromVisibility(Visibility)));
+		UE::MassRepresentation::PushSwapTagsCommand(CommandBuffer, Entity, ChunkVisibility, Visibility);
 		ChunkData.SetContainsNewlyVisibleEntity(Visibility == EMassVisibility::CanBeSeen);
 	}
 }
diff --git a/Engine/Plugins/Runtime/MassGameplay/Source/MassRepresentation/Private/MassRepresentationTypes.cpp b/Engine/Plugins/Runtime/MassGameplay/Source/MassRepresentation/Private/MassRepresentationTypes.cpp
--- a/Engine/Plugins/Runtime/MassGameplay/Source/MassRepresentation/Private/MassRepresentationTypes.cpp
+++ b/Engine/Plugins/Runtime/MassGameplay/Source/MassRepresentation/Private/MassRepresentationTypes.cpp
@@ -1,5 +1,75 @@
 // Copyright Epic Games, Inc. All Rights Reserved.
 
 #include "MassRepresentationTypes.h"
+#include "MassRepresentationUtils.h"
+#include "MassCommandBuffer.h"
 
-DEFINE_LOG_CATEGORY(LogMassRepresentation);
\ No newline at end of file
+DEFINE_LOG_CATEGORY(LogMassRepresentation);
+
+namespace UE::MassRepresentation
+{
+void PushSwapTagsCommand(FMassCommandBuffer& CommandBuffer, const FMassEntityHandle Entity, const EMassVisibility PrevVisibility, const EMassVisibility NewVisibility)
+{
+#define CASE_SWAP_TAGS(OldVisibility, NewVisibility) \
+	case NewVisibility: \
+		CommandBuffer.SwapTags<TMassVisibilityTagForLevel<OldVisibility>::FTag, TMassVisibilityTagForLevel<NewVisibility>::FTag>(Entity); \
+		break
+
+#define CASE_ADD_TAG(NewVisibility) \
+case NewVisibility: \
+	CommandBuffer.AddTag<TMassVisibilityTagForLevel<NewVisibility>::FTag>(Entity); \
+	break
+
+#define DEFAULT_REMOVE_TAG(OldVisibility) \
+case EMassVisibility::Max: /* fall through on purpose */ \
+default: \
+	CommandBuffer.RemoveTag<TMassVisibilityTagForLevel<OldVisibility>::FTag>(Entity); \
+	break
+
+	check(PrevVisibility != NewVisibility);
+
+	switch (PrevVisibility)
+	{
+	case EMassVisibility::CanBeSeen:
+		switch (NewVisibility)
+		{
+		CASE_SWAP_TAGS(EMassVisibility::CanBeSeen, EMassVisibility::CulledByFrustum);
+		CASE_SWAP_TAGS(EMassVisibility::CanBeSeen, EMassVisibility::CulledByDistance);
+		DEFAULT_REMOVE_TAG(EMassVisibility::CanBeSeen);
+		}
+		break;
+	case EMassVisibility::CulledByFrustum:
+		switch (NewVisibility)
+		{
+		CASE_SWAP_TAGS(EMassVisibility::CulledByFrustum, EMassVisibility::CanBeSeen);
+		CASE_SWAP_TAGS(EMassVisibility::CulledByFrustum, EMassVisibility::CulledByDistance);
+		DEFAULT_REMOVE_TAG(EMassVisibility::CulledByFrustum);
+		}
+		break;
+	case EMassVisibility::CulledByDistance:
+		switch (NewVisibility)
+		{
+		CASE_SWAP_TAGS(EMassVisibility::CulledByDistance, EMassVisibility::CanBeSeen);
+		CASE_SWAP_TAGS(EMassVisibility::CulledByDistance, EMassVisibility::CulledByFrustum);
+		DEFAULT_REMOVE_TAG(EMassVisibility::CulledByDistance);
+		}
+		break;
+	case EMassVisibility::Max:
+		switch (NewVisibility)
+		{
+		CASE_ADD_TAG(EMassVisibility::CanBeSeen);
+		CASE_ADD_TAG(EMassVisibility::CulledByFrustum);
+		CASE_ADD_TAG(EMassVisibility::CulledByDistance);
+		default:
+			checkf(false, TEXT("Unsupported Visibility types!"));
+			break;
+		}
+		break;
+	default:
+		checkf(false, TEXT("Unsupported Visibility type!"));
+		break;
+	}
+
+#undef CASE_SWAP_TAGS
+}
+} // UE::MassRepresentation
\ No newline at end of file
diff --git a/Engine/Plugins/Runtime/MassGameplay/Source/MassRepresentation/Public/MassRepresentationUtils.h b/Engine/Plugins/Runtime/MassGameplay/Source/MassRepresentation/Public/MassRepresentationUtils.h
--- a/Engine/Plugins/Runtime/MassGameplay/Source/MassRepresentation/Public/MassRepresentationUtils.h
+++ b/Engine/Plugins/Runtime/MassGameplay/Source/MassRepresentation/Public/MassRepresentationUtils.h
@@ -6,6 +6,9 @@
 #include "MassEntitySubsystem.h"
 #include "MassCommonTypes.h"
 #include "MassRepresentationTypes.h"
+#include "MassLODFragments.h"
+
+struct FMassCommandBuffer;
 
 namespace UE::MassRepresentation
 {
@@ -27,16 +30,40 @@ inline EMassVisibility GetVisibilityFromArchetype(const FMassExecutionContext& C
 	return EMassVisibility::Max;
 }
 
+template <EMassVisibility Level>
+struct TMassVisibilityTagForLevel
+{
+	typedef FMassTag FTag;
+};
+
+template<>
+struct TMassVisibilityTagForLevel<EMassVisibility::CanBeSeen>
+{
+	typedef FMassVisibilityCanBeSeenTag FTag;
+};
+
+template<>
+struct TMassVisibilityTagForLevel<EMassVisibility::CulledByFrustum>
+{
+	typedef FMassVisibilityCulledByFrustumTag FTag;
+};
+
+template<>
+struct TMassVisibilityTagForLevel<EMassVisibility::CulledByDistance>
+{
+	typedef FMassVisibilityCulledByDistanceTag FTag;
+};
+
 inline const UScriptStruct* GetTagFromVisibility(EMassVisibility Visibility)
 {
 	switch (Visibility)
 	{
 		case EMassVisibility::CanBeSeen:
-			return FMassVisibilityCanBeSeenTag::StaticStruct();
+			return TMassVisibilityTagForLevel<EMassVisibility::CanBeSeen>::FTag::StaticStruct();
 		case EMassVisibility::CulledByFrustum:
-			return FMassVisibilityCulledByFrustumTag::StaticStruct();
+			return TMassVisibilityTagForLevel<EMassVisibility::CulledByFrustum>::FTag::StaticStruct();
 		case EMassVisibility::CulledByDistance:
-			return FMassVisibilityCulledByDistanceTag::StaticStruct();
+			return TMassVisibilityTagForLevel<EMassVisibility::CulledByDistance>::FTag::StaticStruct();
 		default:
 			checkf(false, TEXT("Unsupported visibility Type"));
 		case EMassVisibility::Max:
@@ -61,4 +88,6 @@ inline bool IsVisibilityTagSet(const FMassExecutionContext& Context, EMassVisibi
 	}
 }
 
+void PushSwapTagsCommand(FMassCommandBuffer& CommandBuffer, const FMassEntityHandle Entity, const EMassVisibility PrevVisibility, const EMassVisibility NewVisibility);
+
 } // UE::MassRepresentation
\ No newline at end of file
