Added a quick debug-logging console-triggered function for printing out all "known" tags and fragments (including shared- and chunk-).

diff --git a/Engine/Plugins/Runtime/MassEntity/Source/MassEntity/Private/MassEntityDebug.cpp b/Engine/Plugins/Runtime/MassEntity/Source/MassEntity/Private/MassEntityDebug.cpp
--- a/Engine/Plugins/Runtime/MassEntity/Source/MassEntity/Private/MassEntityDebug.cpp
+++ b/Engine/Plugins/Runtime/MassEntity/Source/MassEntity/Private/MassEntityDebug.cpp
@@ -5,6 +5,7 @@
 #include "MassProcessor.h"
 #include "Engine/World.h"
 #include "Engine/Engine.h"
+#include "ScriptStructTypeBitSet.h"
 
 DEFINE_ENUM_TO_STRING(EMassProcessingPhase);
 
@@ -158,6 +159,34 @@ FAutoConsoleCommandWithWorldArgsAndOutputDevice LogMemoryUsage(
 	}
 }));
 
+FAutoConsoleCommandWithOutputDevice LogFragments(
+	TEXT("mass.LogKnownFragments"),
+	TEXT("Logs all the known tags and fragments along with their \"index\" as stored via bitsets."),
+	FConsoleCommandWithOutputDeviceDelegate::CreateStatic([](FOutputDevice& OutputDevice)
+{
+	auto PrintKnownTypes = [&OutputDevice](TConstArrayView<TWeakObjectPtr<const UScriptStruct>> AllStructs) {
+		int i = 0;
+		for (TWeakObjectPtr<const UScriptStruct> Struct : AllStructs)
+		{
+			if (Struct.IsValid())
+			{
+				OutputDevice.Logf(TEXT("\t%d. %s"), i++, *Struct->GetName());
+			}
+		}
+	};
+
+	OutputDevice.Logf(TEXT("Known tags:"));
+	PrintKnownTypes(FMassTagBitSet::DebugGetAllStructTypes());
+
+	OutputDevice.Logf(TEXT("Known Fragments:"));
+	PrintKnownTypes(FMassFragmentBitSet::DebugGetAllStructTypes());
+
+	OutputDevice.Logf(TEXT("Known Shared Fragments:"));
+	PrintKnownTypes(FMassSharedFragmentBitSet::DebugGetAllStructTypes());
+
+	OutputDevice.Logf(TEXT("Known Chunk Fragments:"));
+	PrintKnownTypes(FMassChunkFragmentBitSet::DebugGetAllStructTypes());
+}));
 
 #endif // WITH_MASSENTITY_DEBUG && WITH_MASSENTITY_DEBUG
 
