Renamed FMassCommandAddFragment and FMassCommandRemoveFragment to FMassCommandAddFragments and FMassCommandRemoveFragments respectively.

diff --git a/Engine/Plugins/Runtime/MassEntity/Source/MassEntity/Public/MassCommandBuffer.h b/Engine/Plugins/Runtime/MassEntity/Source/MassEntity/Public/MassCommandBuffer.h
--- a/Engine/Plugins/Runtime/MassEntity/Source/MassEntity/Public/MassCommandBuffer.h
+++ b/Engine/Plugins/Runtime/MassEntity/Source/MassEntity/Public/MassCommandBuffer.h
@@ -50,14 +50,14 @@ public:
 	void AddFragment(FMassEntityHandle Entity)
 	{
 		static_assert(TIsDerivedFrom<T, FMassFragment>::IsDerived, "Given struct type is not a valid fragment type.");
-		PushCommand<FMassCommandAddFragment<T>>(Entity);
+		PushCommand<FMassCommandAddFragments<T>>(Entity);
 	}
 
 	template<typename T>
 	void RemoveFragment(FMassEntityHandle Entity)
 	{
 		static_assert(TIsDerivedFrom<T, FMassFragment>::IsDerived, "Given struct type is not a valid fragment type.");
-		PushCommand<FMassCommandRemoveFragment<T>>(Entity);
+		PushCommand<FMassCommandRemoveFragments<T>>(Entity);
 	}
 
 	/** the convenience function equivalent to calling PushCommand<FMassCommandAddTag<T>>(Entity) */
diff --git a/Engine/Plugins/Runtime/MassEntity/Source/MassEntity/Public/MassCommands.h b/Engine/Plugins/Runtime/MassEntity/Source/MassEntity/Public/MassCommands.h
--- a/Engine/Plugins/Runtime/MassEntity/Source/MassEntity/Public/MassCommands.h
+++ b/Engine/Plugins/Runtime/MassEntity/Source/MassEntity/Public/MassCommands.h
@@ -130,9 +130,9 @@ protected:
 // Simple fragment composition change
 
 template<typename... TTypes>
-struct FMassCommandAddFragment : public FMassBatchedEntityCommand
+struct FMassCommandAddFragments : public FMassBatchedEntityCommand
 {
-	FMassCommandAddFragment()
+	FMassCommandAddFragments()
 	{
 		UE::Mass::TMultiTypeList<TTypes...>::PopulateBitSet(FragmentsAffected);
 		OperationType = EMassCommandOperationType::Add;
@@ -148,9 +148,9 @@ protected:
 };
 
 template<typename... TTypes>
-struct FMassCommandRemoveFragment : public FMassBatchedEntityCommand
+struct FMassCommandRemoveFragments : public FMassBatchedEntityCommand
 {
-	FMassCommandRemoveFragment()
+	FMassCommandRemoveFragments()
 	{
 		UE::Mass::TMultiTypeList<TTypes...>::PopulateBitSet(FragmentsAffected);
 		OperationType = EMassCommandOperationType::Remove;
