Shipping editor compilation fix to MassCommandsTest.cpp

diff --git a/Engine/Plugins/Runtime/MassEntity/Source/MassEntityTestSuite/Private/MassCommandsTest.cpp b/Engine/Plugins/Runtime/MassEntity/Source/MassEntityTestSuite/Private/MassCommandsTest.cpp
--- a/Engine/Plugins/Runtime/MassEntity/Source/MassEntityTestSuite/Private/MassCommandsTest.cpp
+++ b/Engine/Plugins/Runtime/MassEntity/Source/MassEntityTestSuite/Private/MassCommandsTest.cpp
@@ -20,7 +20,7 @@ PRAGMA_DISABLE_OPTIMIZATION
 
 namespace FMassCommandsTest
 {
-
+#if WITH_MASSENTITY_DEBUG
 struct FCommands_FragmentInstanceList : FEntityTestBase
 {
 	virtual bool InstantTest() override
@@ -242,6 +242,7 @@ struct FCommands_DeferredFunction : FEntityTestBase
 };
 IMPLEMENT_AI_INSTANT_TEST(FCommands_DeferredFunction, "System.Mass.Commands.DeferredFunction");
 
+#endif // WITH_MASSENTITY_DEBUG
 } // FMassCommandsTest
 
 PRAGMA_ENABLE_OPTIMIZATION
