Fixed FMassCommandBuildEntityWithSharedFragments not clearing out its payload on reset.

diff --git a/Engine/Plugins/Runtime/MassEntity/Source/MassEntity/Public/MassCommands.h b/Engine/Plugins/Runtime/MassEntity/Source/MassEntity/Public/MassCommands.h
--- a/Engine/Plugins/Runtime/MassEntity/Source/MassEntity/Public/MassCommands.h
+++ b/Engine/Plugins/Runtime/MassEntity/Source/MassEntity/Public/MassCommands.h
@@ -70,6 +70,8 @@ private:
 
 struct FMassBatchedEntityCommand : public FMassBatchedCommand
 {
+	using Super = FMassBatchedCommand;
+
 	void Add(FMassEntityHandle Entity)
 	{
 		UE_MT_SCOPED_READ_ACCESS(EntitiesAccessDetector);
@@ -90,10 +92,10 @@ protected:
 		return TargetEntities.GetAllocatedSize();
 	}
 
-	virtual void Reset()
+	virtual void Reset() override
 	{
 		TargetEntities.Reset();
-		FMassBatchedCommand::Reset();
+		Super::Reset();
 	}
 
 #if CSV_PROFILER || WITH_MASSENTITY_DEBUG
@@ -321,6 +323,8 @@ protected:
 template<typename TSharedFragmentValues, typename... TOthers>
 struct FMassCommandBuildEntityWithSharedFragments : public FMassBatchedCommand
 {
+	using Super = FMassBatchedCommand;
+
 	FMassCommandBuildEntityWithSharedFragments()
 	{
 		OperationType = EMassCommandOperationType::Create;
@@ -377,6 +381,12 @@ protected:
 		}
 	}
 
+	virtual void Reset() override
+	{
+		Data.Reset();
+		Super::Reset();
+	}
+
 #if CSV_PROFILER || WITH_MASSENTITY_DEBUG
 	virtual int32 GetNumOperationsStat() const override
 	{
@@ -454,7 +464,7 @@ protected:
 		}
 	}
 
-	virtual void Reset()
+	virtual void Reset() override
 	{
 		DeferredFunctions.Reset();
 		Super::Reset();
