Massaged bits of FMassCommandBuffer::Flush to avoid erronous PSV warning.

diff --git a/Engine/Plugins/Runtime/MassEntity/Source/MassEntity/Private/MassCommandBuffer.cpp b/Engine/Plugins/Runtime/MassEntity/Source/MassEntity/Private/MassCommandBuffer.cpp
--- a/Engine/Plugins/Runtime/MassEntity/Source/MassEntity/Private/MassCommandBuffer.cpp
+++ b/Engine/Plugins/Runtime/MassEntity/Source/MassEntity/Private/MassCommandBuffer.cpp
@@ -108,6 +108,10 @@ void FMassCommandBuffer::Flush(UMassEntitySubsystem& EntitySystem)
 
 		struct FBatchedCommandsSortedIndex
 		{
+			FBatchedCommandsSortedIndex(const int32 InIndex, const int32 InGroupOrder)
+				: Index(InIndex), GroupOrder(InGroupOrder)
+			{}
+
 			const int32 Index = -1;
 			const int32 GroupOrder = MAX_int32;
 			bool IsValid() const { return GroupOrder < MAX_int32; }
@@ -118,7 +122,7 @@ void FMassCommandBuffer::Flush(UMassEntitySubsystem& EntitySystem)
 		for (int32 i = 0; i < CommandInstances.Num(); ++i)
 		{
 			const FMassBatchedCommand* Command = CommandInstances[i];
-			CommandsOrder.Add({ i, (Command && Command->HasWork())? CommandTypeOrder[(int)Command->GetOperationType()] : MAX_int32 });
+			CommandsOrder.Add(FBatchedCommandsSortedIndex(i, (Command && Command->HasWork())? CommandTypeOrder[(int)Command->GetOperationType()] : MAX_int32));
 		}
 		CommandsOrder.Sort();
 				
