[Mass] fixed FMassCommandBuildEntityWithSharedFragments hash computation
#rnx

diff --git a/Engine/Plugins/Runtime/MassEntity/Source/MassEntity/Public/MassCommands.h b/Engine/Plugins/Runtime/MassEntity/Source/MassEntity/Public/MassCommands.h
--- a/Engine/Plugins/Runtime/MassEntity/Source/MassEntity/Public/MassCommands.h
+++ b/Engine/Plugins/Runtime/MassEntity/Source/MassEntity/Public/MassCommands.h
@@ -337,7 +337,12 @@ struct FMassCommandBuildEntityWithSharedFragments : public FMassBatchedCommand
 	void Add(FMassEntityHandle Entity, FMassArchetypeSharedFragmentValues&& InSharedFragments, TOthers... InFragments)
 	{
 		InSharedFragments.Sort();
-		FPerSharedFragmentsHashData& Instance = Data.FindOrAdd(GetTypeHash(InSharedFragments), MoveTemp(InSharedFragments));
+
+		// Compute hash before adding to the map since evaluation order is not guaranteed
+		// and MoveTemp will invalidate InSharedFragments
+		const uint32 Hash = GetTypeHash(InSharedFragments);
+
+		FPerSharedFragmentsHashData& Instance = Data.FindOrAdd(Hash, MoveTemp(InSharedFragments));
 		Instance.Fragments.Add(InFragments...);
 		Instance.TargetEntities.Add(Entity);
 
