Switched UMassEntitySubsystem.FlushedCommandBufferQueue from TQueue<T, EQueueMode::Mpsc> to more efficient TMpscQueue<T>.

diff --git a/Engine/Plugins/Runtime/MassEntity/Source/MassEntity/Private/MassEntitySubsystem.cpp b/Engine/Plugins/Runtime/MassEntity/Source/MassEntity/Private/MassEntitySubsystem.cpp
--- a/Engine/Plugins/Runtime/MassEntity/Source/MassEntity/Private/MassEntitySubsystem.cpp
+++ b/Engine/Plugins/Runtime/MassEntity/Source/MassEntity/Private/MassEntitySubsystem.cpp
@@ -1242,12 +1242,13 @@ void UMassEntitySubsystem::FlushCommands(const TSharedPtr<FMassCommandBuffer>& I
 		bCommandBufferFlushingInProgress = true;
 		
 		int IterationsCounter = 0;
-		TSharedPtr<FMassCommandBuffer> CurrentCommandBuffer;
-		while (IterationsCounter++ < MaxIterations && FlushedCommandBufferQueue.Dequeue(CurrentCommandBuffer))
+		TOptional<TSharedPtr<FMassCommandBuffer>> CurrentCommandBuffer = FlushedCommandBufferQueue.Dequeue();
+		while (IterationsCounter++ < MaxIterations && CurrentCommandBuffer.IsSet())
 		{
-			CurrentCommandBuffer->Flush(*this);
+			(*CurrentCommandBuffer)->Flush(*this);
+			CurrentCommandBuffer = FlushedCommandBufferQueue.Dequeue();
 		}
-		ensure(IterationsCounter >= MaxIterations || FlushedCommandBufferQueue.IsEmpty());
+		ensure(IterationsCounter >= MaxIterations || CurrentCommandBuffer.IsSet() == false);
 		UE_CVLOG_UELOG(IterationsCounter >= MaxIterations, this, LogMass, Error, TEXT("Reached loop count limit while flushing commands"));
 
 		bCommandBufferFlushingInProgress = false;
diff --git a/Engine/Plugins/Runtime/MassEntity/Source/MassEntity/Public/MassEntitySubsystem.h b/Engine/Plugins/Runtime/MassEntity/Source/MassEntity/Public/MassEntitySubsystem.h
--- a/Engine/Plugins/Runtime/MassEntity/Source/MassEntity/Public/MassEntitySubsystem.h
+++ b/Engine/Plugins/Runtime/MassEntity/Source/MassEntity/Public/MassEntitySubsystem.h
@@ -9,7 +9,7 @@
 #include "MassEntityQuery.h"
 #include "StructUtilsTypes.h"
 #include "MassObserverManager.h"
-#include "Containers/Queue.h"
+#include "Containers/MpscQueue.h"
 #include "MassEntitySubsystem.generated.h"
 
 
@@ -391,7 +391,7 @@ private:
 
 	std::atomic<bool> bCommandBufferFlushingInProgress = false;
 	TSharedPtr<FMassCommandBuffer> DeferredCommandBuffer;
-	TQueue<TSharedPtr<FMassCommandBuffer>, EQueueMode::Mpsc> FlushedCommandBufferQueue;
+	TMpscQueue<TSharedPtr<FMassCommandBuffer>> FlushedCommandBufferQueue;
 
 	std::atomic<int32> SerialNumberGenerator;
 	std::atomic<int32> ProcessingScopeCount;
