Extracted MassExecutionContext into a separate header/cpp files. No changes other than the ones required by the move have been done.

diff --git a/Engine/Plugins/Runtime/MassEntity/Source/MassEntity/Private/MassEntitySubsystem.cpp b/Engine/Plugins/Runtime/MassEntity/Source/MassEntity/Private/MassEntitySubsystem.cpp
--- a/Engine/Plugins/Runtime/MassEntity/Source/MassEntity/Private/MassEntitySubsystem.cpp
+++ b/Engine/Plugins/Runtime/MassEntity/Source/MassEntity/Private/MassEntitySubsystem.cpp
@@ -7,6 +7,7 @@
 #include "Engine/World.h"
 #include "UObject/UObjectIterator.h"
 #include "VisualLogger/VisualLogger.h"
+#include "MassExecutionContext.h"
 
 const FMassEntityHandle UMassEntitySubsystem::InvalidEntity;
 
@@ -1425,76 +1426,3 @@ FAutoConsoleCommandWithWorldArgsAndOutputDevice GPrintArchetypesCmd(
 	}
 }));
 #endif // WITH_MASSENTITY_DEBUG
-
-//////////////////////////////////////////////////////////////////////
-// FMassExecutionContext
-
-void FMassExecutionContext::FlushDeferred(UMassEntitySubsystem& EntitySystem) const
-{
-	if (bFlushDeferredCommands && DeferredCommandBuffer)
-	{
-		EntitySystem.FlushCommands(DeferredCommandBuffer);
-	}
-}
-
-void FMassExecutionContext::ClearExecutionData()
-{
-	FragmentViews.Reset();
-	CurrentArchetypesTagBitSet.Reset();
-	ChunkSerialModificationNumber = -1;
-}
-
-void FMassExecutionContext::SetEntityCollection(const FMassArchetypeEntityCollection& InEntityCollection)
-{
-	check(EntityCollection.IsEmpty());
-	EntityCollection = InEntityCollection;
-}
-
-void FMassExecutionContext::SetEntityCollection(FMassArchetypeEntityCollection&& InEntityCollection)
-{
-	check(EntityCollection.IsEmpty());
-	EntityCollection = MoveTemp(InEntityCollection);
-}
-
-void FMassExecutionContext::SetRequirements(TConstArrayView<FMassFragmentRequirement> InRequirements, 
-	TConstArrayView<FMassFragmentRequirement> InChunkRequirements, 
-	TConstArrayView<FMassFragmentRequirement> InConstSharedRequirements, 
-	TConstArrayView<FMassFragmentRequirement> InSharedRequirements)
-{ 
-	FragmentViews.Reset();
-	for (const FMassFragmentRequirement& Requirement : InRequirements)
-	{
-		if (Requirement.RequiresBinding())
-		{
-			FragmentViews.Emplace(Requirement);
-		}
-	}
-
-	ChunkFragmentViews.Reset();
-	for (const FMassFragmentRequirement& Requirement : InChunkRequirements)
-	{
-		if (Requirement.RequiresBinding())
-		{
-			ChunkFragmentViews.Emplace(Requirement);
-		}
-	}
-
-	ConstSharedFragmentViews.Reset();
-	for (const FMassFragmentRequirement& Requirement : InConstSharedRequirements)
-	{
-		if (Requirement.RequiresBinding())
-		{
-			ConstSharedFragmentViews.Emplace(Requirement);
-		}
-	}
-
-	SharedFragmentViews.Reset();
-	for (const FMassFragmentRequirement& Requirement : InSharedRequirements)
-	{
-		if (Requirement.RequiresBinding())
-		{
-			SharedFragmentViews.Emplace(Requirement);
-		}
-	}
-}
-
diff --git a/Engine/Plugins/Runtime/MassEntity/Source/MassEntity/Private/MassExecutionContext.cpp b/Engine/Plugins/Runtime/MassEntity/Source/MassEntity/Private/MassExecutionContext.cpp
new file mode 100644
--- /dev/null
+++ b/Engine/Plugins/Runtime/MassEntity/Source/MassEntity/Private/MassExecutionContext.cpp
@@ -0,0 +1,75 @@
+// Copyright Epic Games, Inc. All Rights Reserved.
+
+#include "MassExecutionContext.h"
+#include "MassArchetypeData.h"
+
+
+void FMassExecutionContext::FlushDeferred(UMassEntitySubsystem& EntitySystem) const
+{
+	if (bFlushDeferredCommands && DeferredCommandBuffer)
+	{
+		EntitySystem.FlushCommands(DeferredCommandBuffer);
+	}
+}
+
+void FMassExecutionContext::ClearExecutionData()
+{
+	FragmentViews.Reset();
+	CurrentArchetypesTagBitSet.Reset();
+	ChunkSerialModificationNumber = -1;
+}
+
+void FMassExecutionContext::SetEntityCollection(const FMassArchetypeEntityCollection& InEntityCollection)
+{
+	check(EntityCollection.IsEmpty());
+	EntityCollection = InEntityCollection;
+}
+
+void FMassExecutionContext::SetEntityCollection(FMassArchetypeEntityCollection&& InEntityCollection)
+{
+	check(EntityCollection.IsEmpty());
+	EntityCollection = MoveTemp(InEntityCollection);
+}
+
+void FMassExecutionContext::SetRequirements(TConstArrayView<FMassFragmentRequirement> InRequirements, 
+	TConstArrayView<FMassFragmentRequirement> InChunkRequirements, 
+	TConstArrayView<FMassFragmentRequirement> InConstSharedRequirements, 
+	TConstArrayView<FMassFragmentRequirement> InSharedRequirements)
+{ 
+	FragmentViews.Reset();
+	for (const FMassFragmentRequirement& Requirement : InRequirements)
+	{
+		if (Requirement.RequiresBinding())
+		{
+			FragmentViews.Emplace(Requirement);
+		}
+	}
+
+	ChunkFragmentViews.Reset();
+	for (const FMassFragmentRequirement& Requirement : InChunkRequirements)
+	{
+		if (Requirement.RequiresBinding())
+		{
+			ChunkFragmentViews.Emplace(Requirement);
+		}
+	}
+
+	ConstSharedFragmentViews.Reset();
+	for (const FMassFragmentRequirement& Requirement : InConstSharedRequirements)
+	{
+		if (Requirement.RequiresBinding())
+		{
+			ConstSharedFragmentViews.Emplace(Requirement);
+		}
+	}
+
+	SharedFragmentViews.Reset();
+	for (const FMassFragmentRequirement& Requirement : InSharedRequirements)
+	{
+		if (Requirement.RequiresBinding())
+		{
+			SharedFragmentViews.Emplace(Requirement);
+		}
+	}
+}
+
diff --git a/Engine/Plugins/Runtime/MassEntity/Source/MassEntity/Public/MassEntityQuery.h b/Engine/Plugins/Runtime/MassEntity/Source/MassEntity/Public/MassEntityQuery.h
--- a/Engine/Plugins/Runtime/MassEntity/Source/MassEntity/Public/MassEntityQuery.h
+++ b/Engine/Plugins/Runtime/MassEntity/Source/MassEntity/Public/MassEntityQuery.h
@@ -15,6 +15,7 @@ struct FMassArchetypeData;
 struct FMassExecutionContext;
 struct FMassFragment;
 struct FMassArchetypeHandle;
+struct FMassExecutionContext;
 
 enum class EMassFragmentAccess : uint8
 {
diff --git a/Engine/Plugins/Runtime/MassEntity/Source/MassEntity/Public/MassEntitySubsystem.h b/Engine/Plugins/Runtime/MassEntity/Source/MassEntity/Public/MassEntitySubsystem.h
--- a/Engine/Plugins/Runtime/MassEntity/Source/MassEntity/Public/MassEntitySubsystem.h
+++ b/Engine/Plugins/Runtime/MassEntity/Source/MassEntity/Public/MassEntitySubsystem.h
@@ -10,6 +10,7 @@
 #include "StructUtilsTypes.h"
 #include "MassObserverManager.h"
 #include "Containers/MpscQueue.h"
+#include "MassExecutionContext.h"
 #include "MassEntitySubsystem.generated.h"
 
 
@@ -419,277 +420,3 @@ private:
 	UPROPERTY(Transient)
 	FMassObserverManager ObserverManager;
 };
-
-
-//////////////////////////////////////////////////////////////////////
-//
-
-struct MASSENTITY_API FMassExecutionContext
-{
-private:
-
-	template< typename ViewType >
-	struct TFragmentView 
-	{
-		FMassFragmentRequirement Requirement;
-		ViewType FragmentView;
-
-		TFragmentView() {}
-		explicit TFragmentView(const FMassFragmentRequirement& InRequirement) : Requirement(InRequirement) {}
-
-		bool operator==(const UScriptStruct* FragmentType) const { return Requirement.StructType == FragmentType; }
-	};
-	using FFragmentView = TFragmentView<TArrayView<FMassFragment>>;
-	TArray<FFragmentView, TInlineAllocator<8>> FragmentViews;
-
-	using FChunkFragmentView = TFragmentView<FStructView>;
-	TArray<FChunkFragmentView, TInlineAllocator<4>> ChunkFragmentViews;
-
-	using FConstSharedFragmentView = TFragmentView<FConstStructView>;
-	TArray<FConstSharedFragmentView, TInlineAllocator<4>> ConstSharedFragmentViews;
-
-	using FSharedFragmentView = TFragmentView<FStructView>;
-	TArray<FSharedFragmentView, TInlineAllocator<4>> SharedFragmentViews;
-	
-	// mz@todo make this shared ptr thread-safe and never auto-flush in MT environment. 
-	TSharedPtr<FMassCommandBuffer> DeferredCommandBuffer;
-	TArrayView<FMassEntityHandle> EntityListView;
-	
-	/** If set this indicates the exact archetype and its chunks to be processed. 
-	 *  @todo this data should live somewhere else, preferably be just a parameter to Query.ForEachEntityChunk function */
-	FMassArchetypeEntityCollection EntityCollection;
-	
-	/** @todo rename to "payload" */
-	FInstancedStruct AuxData;
-	float DeltaTimeSeconds = 0.0f;
-	int32 ChunkSerialModificationNumber = -1;
-	FMassTagBitSet CurrentArchetypesTagBitSet;
-
-#if WITH_MASSENTITY_DEBUG
-	FString DebugExecutionDescription;
-#endif
-	
-	/** If true the EntitySystem will flush the deferred commands stored in DeferredCommandBuffer just after executing 
-	 *  the given system. If False then the party calling UEntitySubsystem::ExecuteSystem is responsible for manually
-	 *  calling MassEntitySubsystem.FlushCommands() */
-	bool bFlushDeferredCommands = true;
-
-	TArrayView<FFragmentView> GetMutableRequirements() { return FragmentViews; }
-	TArrayView<FChunkFragmentView> GetMutableChunkRequirements() { return ChunkFragmentViews; }
-	TArrayView<FConstSharedFragmentView> GetMutableConstSharedRequirements() { return ConstSharedFragmentViews; }
-	TArrayView<FSharedFragmentView> GetMutableSharedRequirements() { return SharedFragmentViews; }
-
-	friend FMassArchetypeData;
-	friend FMassEntityQuery;
-
-public:
-	FMassExecutionContext() = default;
-	explicit FMassExecutionContext(const float InDeltaTimeSeconds, const bool bInFlushDeferredCommands = true)
-		: DeltaTimeSeconds(InDeltaTimeSeconds)
-		, bFlushDeferredCommands(bInFlushDeferredCommands)
-	{}
-
-#if WITH_MASSENTITY_DEBUG
-	const FString& DebugGetExecutionDesc() const { return DebugExecutionDescription; }
-	void DebugSetExecutionDesc(const FString& Description) { DebugExecutionDescription = Description; }
-#endif
-
-	/** Sets bFlushDeferredCommands. Note that setting to True while the system is being executed doesn't result in
-	 *  immediate commands flushing */
-	void SetFlushDeferredCommands(const bool bNewFlushDeferredCommands) { bFlushDeferredCommands = bNewFlushDeferredCommands; } 
-	void SetDeferredCommandBuffer(const TSharedPtr<FMassCommandBuffer>& InDeferredCommandBuffer) { DeferredCommandBuffer = InDeferredCommandBuffer; }
-	void SetEntityCollection(const FMassArchetypeEntityCollection& InEntityCollection);
-	void SetEntityCollection(FMassArchetypeEntityCollection&& InEntityCollection);
-	void ClearEntityCollection() { EntityCollection.Reset(); }
-	void SetAuxData(const FInstancedStruct& InAuxData) { AuxData = InAuxData; }
-
-	float GetDeltaTimeSeconds() const
-	{
-		return DeltaTimeSeconds;
-	}
-
-	TSharedPtr<FMassCommandBuffer> GetSharedDeferredCommandBuffer() const { return DeferredCommandBuffer; }
-	FMassCommandBuffer& Defer() const { checkSlow(DeferredCommandBuffer.IsValid()); return *DeferredCommandBuffer.Get(); }
-
-	TConstArrayView<FMassEntityHandle> GetEntities() const { return EntityListView; }
-	int32 GetNumEntities() const { return EntityListView.Num(); }
-
-	FMassEntityHandle GetEntity(const int32 Index) const
-	{
-		return EntityListView[Index];
-	}
-
-	template<typename T>
-	bool DoesArchetypeHaveTag() const
-	{
-		static_assert(TIsDerivedFrom<T, FMassTag>::IsDerived, "Given struct is not of a valid fragment type.");
-		return CurrentArchetypesTagBitSet.Contains<T>();
-	}
-
-	/** Chunk related operations */
-	void SetCurrentChunkSerialModificationNumber(const int32 SerialModificationNumber) { ChunkSerialModificationNumber = SerialModificationNumber; }
-	int32 GetChunkSerialModificationNumber() const { return ChunkSerialModificationNumber; }
-
-	template<typename T>
-	T* GetMutableChunkFragmentPtr() const
-	{
-		static_assert(TIsDerivedFrom<T, FMassChunkFragment>::IsDerived, "Given struct doesn't represent a valid chunk fragment type. Make sure to inherit from FMassChunkFragment or one of its child-types.");
-
-		const UScriptStruct* Type = T::StaticStruct();
-		const FChunkFragmentView* FoundChunkFragmentData = ChunkFragmentViews.FindByPredicate([Type](const FChunkFragmentView& Element) { return Element.Requirement.StructType == Type; } );
-		return FoundChunkFragmentData ? FoundChunkFragmentData->FragmentView.GetMutablePtr<T>() : static_cast<T*>(nullptr);
-	}
-	
-	template<typename T>
-	T& GetMutableChunkFragment() const
-	{
-		T* ChunkFragment = GetMutableChunkFragmentPtr<T>();
-		checkf(ChunkFragment, TEXT("Chunk Fragment requirement not found: %s"), *T::StaticStruct()->GetName());
-		return *ChunkFragment;
-	}
-
-	template<typename T>
-	const T* GetChunkFragmentPtr() const
-	{
-		return GetMutableChunkFragmentPtr<T>();
-	}
-	
-	template<typename T>
-	const T& GetChunkFragment() const
-	{
-		const T* ChunkFragment = GetMutableChunkFragmentPtr<T>();
-		checkf(ChunkFragment, TEXT("Chunk Fragment requirement not found: %s"), *T::StaticStruct()->GetName());
-		return *ChunkFragment;
-	}
-
-	/** Shared fragment related operations */
-	template<typename T>
-	const T* GetConstSharedFragmentPtr() const
-	{
-		static_assert(TIsDerivedFrom<T, FMassSharedFragment>::IsDerived, "Given struct doesn't represent a valid Shared fragment type. Make sure to inherit from FMassSharedFragment or one of its child-types.");
-
-		const FConstSharedFragmentView* FoundSharedFragmentData = ConstSharedFragmentViews.FindByPredicate([](const FConstSharedFragmentView& Element) { return Element.Requirement.StructType == T::StaticStruct(); });
-		return FoundSharedFragmentData ? FoundSharedFragmentData->FragmentView.GetPtr<T>() : static_cast<const T*>(nullptr);
-	}
-
-	template<typename T>
-	const T& GetConstSharedFragment() const
-	{
-		const T* SharedFragment = GetConstSharedFragmentPtr<T>();
-		checkf(SharedFragment, TEXT("Shared Fragment requirement not found: %s"), *T::StaticStruct()->GetName());
-		return *SharedFragment;
-	}
-
-
-	template<typename T>
-	T* GetMutableSharedFragmentPtr() const
-	{
-		static_assert(TIsDerivedFrom<T, FMassSharedFragment>::IsDerived, "Given struct doesn't represent a valid shared fragment type. Make sure to inherit from FMassSharedFragment or one of its child-types.");
-
-		const FSharedFragmentView* FoundSharedFragmentData = SharedFragmentViews.FindByPredicate([](const FSharedFragmentView& Element) { return Element.Requirement.StructType == T::StaticStruct(); });
-		return FoundSharedFragmentData ? FoundSharedFragmentData->FragmentView.GetMutablePtr<T>() : static_cast<T*>(nullptr);
-	}
-
-	template<typename T>
-	T& GetMutableSharedFragment() const
-	{
-		T* SharedFragment = GetMutableSharedFragmentPtr<T>();
-		checkf(SharedFragment, TEXT("Shared Fragment requirement not found: %s"), *T::StaticStruct()->GetName());
-		return *SharedFragment;
-	}
-
-	template<typename T>
-	const T* GetSharedFragmentPtr() const
-	{
-		return GetMutableSharedFragmentPtr<T>();
-	}
-
-	template<typename T>
-	const T& GetSharedFragment() const
-	{
-		const T* SharedFragment = GetSharedFragmentPtr<T>();
-		checkf(SharedFragment, TEXT("Shared Fragment requirement not found: %s"), *T::StaticStruct()->GetName());
-		return *SharedFragment;
-	}
-
-	/* Fragments related operations */
-	template<typename TFragment>
-	TArrayView<TFragment> GetMutableFragmentView()
-	{
-		const UScriptStruct* FragmentType = TFragment::StaticStruct();
-		const FFragmentView* View = FragmentViews.FindByPredicate([FragmentType](const FFragmentView& Element) { return Element.Requirement.StructType == FragmentType; });
-		//checkfSlow(View != nullptr, TEXT("Requested fragment type not bound"));
-		//checkfSlow(View->Requirement.AccessMode == EMassFragmentAccess::ReadWrite, TEXT("Requested fragment has not been bound for writing"));
-		return MakeArrayView<TFragment>((TFragment*)View->FragmentView.GetData(), View->FragmentView.Num());
-	}
-
-	template<typename TFragment>
-	TConstArrayView<TFragment> GetFragmentView() const
-	{
-		const UScriptStruct* FragmentType = TFragment::StaticStruct();
-		const FFragmentView* View = FragmentViews.FindByPredicate([FragmentType](const FFragmentView& Element) { return Element.Requirement.StructType == FragmentType; });
-		//checkfSlow(View != nullptr, TEXT("Requested fragment type not bound"));
-		return TConstArrayView<TFragment>((const TFragment*)View->FragmentView.GetData(), View->FragmentView.Num());
-	}
-
-	TConstArrayView<FMassFragment> GetFragmentView(const UScriptStruct* FragmentType) const
-	{
-		const FFragmentView* View = FragmentViews.FindByPredicate([FragmentType](const FFragmentView& Element) { return Element.Requirement.StructType == FragmentType; });
-		checkSlow(View);
-		return TConstArrayView<FMassFragment>((const FMassFragment*)View->FragmentView.GetData(), View->FragmentView.Num());;
-	}
-
-	TArrayView<FMassFragment> GetMutableFragmentView(const UScriptStruct* FragmentType) 
-	{
-		const FFragmentView* View = FragmentViews.FindByPredicate([FragmentType](const FFragmentView& Element) { return Element.Requirement.StructType == FragmentType; });
-		checkSlow(View);
-		return View->FragmentView;
-	}
-
-	/** Sparse chunk related operation */
-	const FMassArchetypeEntityCollection& GetEntityCollection() const { return EntityCollection; }
-
-	const FInstancedStruct& GetAuxData() const { return AuxData; }
-	FInstancedStruct& GetMutableAuxData() { return AuxData; }
-	
-	template<typename TFragment>
-	bool ValidateAuxDataType() const
-	{
-		const UScriptStruct* FragmentType = GetAuxData().GetScriptStruct();
-		return FragmentType != nullptr && FragmentType == TFragment::StaticStruct();
-	}
-
-	void FlushDeferred(UMassEntitySubsystem& EntitySystem) const;
-
-	void ClearExecutionData();
-	void SetCurrentArchetypesTagBitSet(const FMassTagBitSet& BitSet)
-	{
-		CurrentArchetypesTagBitSet = BitSet;
-	}
-
-protected:
-	void SetRequirements(TConstArrayView<FMassFragmentRequirement> InRequirements, 
-		TConstArrayView<FMassFragmentRequirement> InChunkRequirements, 
-		TConstArrayView<FMassFragmentRequirement> InConstSharedRequirements, 
-		TConstArrayView<FMassFragmentRequirement> InSharedRequirements);
-
-	void ClearFragmentViews()
-	{
-		for (FFragmentView& View : FragmentViews)
-		{
-			View.FragmentView = TArrayView<FMassFragment>();
-		}
-		for (FChunkFragmentView& View : ChunkFragmentViews)
-		{
-			View.FragmentView.Reset();
-		}
-		for (FConstSharedFragmentView& View : ConstSharedFragmentViews)
-		{
-			View.FragmentView.Reset();
-		}
-		for (FSharedFragmentView& View : SharedFragmentViews)
-		{
-			View.FragmentView.Reset();
-		}
-	}
-};
diff --git a/Engine/Plugins/Runtime/MassEntity/Source/MassEntity/Public/MassExecutionContext.h b/Engine/Plugins/Runtime/MassEntity/Source/MassEntity/Public/MassExecutionContext.h
new file mode 100644
--- /dev/null
+++ b/Engine/Plugins/Runtime/MassEntity/Source/MassEntity/Public/MassExecutionContext.h
@@ -0,0 +1,277 @@
+// Copyright Epic Games, Inc. All Rights Reserved.
+
+#pragma once
+
+#include "MassEntityTypes.h"
+#include "MassEntityQuery.h"
+
+
+struct MASSENTITY_API FMassExecutionContext
+{
+private:
+
+	template< typename ViewType >
+	struct TFragmentView 
+	{
+		FMassFragmentRequirement Requirement;
+		ViewType FragmentView;
+
+		TFragmentView() {}
+		explicit TFragmentView(const FMassFragmentRequirement& InRequirement) : Requirement(InRequirement) {}
+
+		bool operator==(const UScriptStruct* FragmentType) const { return Requirement.StructType == FragmentType; }
+	};
+	using FFragmentView = TFragmentView<TArrayView<FMassFragment>>;
+	TArray<FFragmentView, TInlineAllocator<8>> FragmentViews;
+
+	using FChunkFragmentView = TFragmentView<FStructView>;
+	TArray<FChunkFragmentView, TInlineAllocator<4>> ChunkFragmentViews;
+
+	using FConstSharedFragmentView = TFragmentView<FConstStructView>;
+	TArray<FConstSharedFragmentView, TInlineAllocator<4>> ConstSharedFragmentViews;
+
+	using FSharedFragmentView = TFragmentView<FStructView>;
+	TArray<FSharedFragmentView, TInlineAllocator<4>> SharedFragmentViews;
+	
+	// mz@todo make this shared ptr thread-safe and never auto-flush in MT environment. 
+	TSharedPtr<FMassCommandBuffer> DeferredCommandBuffer;
+	TArrayView<FMassEntityHandle> EntityListView;
+	
+	/** If set this indicates the exact archetype and its chunks to be processed. 
+	 *  @todo this data should live somewhere else, preferably be just a parameter to Query.ForEachEntityChunk function */
+	FMassArchetypeEntityCollection EntityCollection;
+	
+	/** @todo rename to "payload" */
+	FInstancedStruct AuxData;
+	float DeltaTimeSeconds = 0.0f;
+	int32 ChunkSerialModificationNumber = -1;
+	FMassTagBitSet CurrentArchetypesTagBitSet;
+
+#if WITH_MASSENTITY_DEBUG
+	FString DebugExecutionDescription;
+#endif
+	
+	/** If true the EntitySystem will flush the deferred commands stored in DeferredCommandBuffer just after executing 
+	 *  the given system. If False then the party calling UEntitySubsystem::ExecuteSystem is responsible for manually
+	 *  calling MassEntitySubsystem.FlushCommands() */
+	bool bFlushDeferredCommands = true;
+
+	TArrayView<FFragmentView> GetMutableRequirements() { return FragmentViews; }
+	TArrayView<FChunkFragmentView> GetMutableChunkRequirements() { return ChunkFragmentViews; }
+	TArrayView<FConstSharedFragmentView> GetMutableConstSharedRequirements() { return ConstSharedFragmentViews; }
+	TArrayView<FSharedFragmentView> GetMutableSharedRequirements() { return SharedFragmentViews; }
+
+	friend FMassArchetypeData;
+	friend FMassEntityQuery;
+
+public:
+	FMassExecutionContext() = default;
+	explicit FMassExecutionContext(const float InDeltaTimeSeconds, const bool bInFlushDeferredCommands = true)
+		: DeltaTimeSeconds(InDeltaTimeSeconds)
+		, bFlushDeferredCommands(bInFlushDeferredCommands)
+	{}
+
+#if WITH_MASSENTITY_DEBUG
+	const FString& DebugGetExecutionDesc() const { return DebugExecutionDescription; }
+	void DebugSetExecutionDesc(const FString& Description) { DebugExecutionDescription = Description; }
+#endif
+
+	/** Sets bFlushDeferredCommands. Note that setting to True while the system is being executed doesn't result in
+	 *  immediate commands flushing */
+	void SetFlushDeferredCommands(const bool bNewFlushDeferredCommands) { bFlushDeferredCommands = bNewFlushDeferredCommands; } 
+	void SetDeferredCommandBuffer(const TSharedPtr<FMassCommandBuffer>& InDeferredCommandBuffer) { DeferredCommandBuffer = InDeferredCommandBuffer; }
+	void SetEntityCollection(const FMassArchetypeEntityCollection& InEntityCollection);
+	void SetEntityCollection(FMassArchetypeEntityCollection&& InEntityCollection);
+	void ClearEntityCollection() { EntityCollection.Reset(); }
+	void SetAuxData(const FInstancedStruct& InAuxData) { AuxData = InAuxData; }
+
+	float GetDeltaTimeSeconds() const
+	{
+		return DeltaTimeSeconds;
+	}
+
+	TSharedPtr<FMassCommandBuffer> GetSharedDeferredCommandBuffer() const { return DeferredCommandBuffer; }
+	FMassCommandBuffer& Defer() const { checkSlow(DeferredCommandBuffer.IsValid()); return *DeferredCommandBuffer.Get(); }
+
+	TConstArrayView<FMassEntityHandle> GetEntities() const { return EntityListView; }
+	int32 GetNumEntities() const { return EntityListView.Num(); }
+
+	FMassEntityHandle GetEntity(const int32 Index) const
+	{
+		return EntityListView[Index];
+	}
+
+	template<typename T>
+	bool DoesArchetypeHaveTag() const
+	{
+		static_assert(TIsDerivedFrom<T, FMassTag>::IsDerived, "Given struct is not of a valid fragment type.");
+		return CurrentArchetypesTagBitSet.Contains<T>();
+	}
+
+	/** Chunk related operations */
+	void SetCurrentChunkSerialModificationNumber(const int32 SerialModificationNumber) { ChunkSerialModificationNumber = SerialModificationNumber; }
+	int32 GetChunkSerialModificationNumber() const { return ChunkSerialModificationNumber; }
+
+	template<typename T>
+	T* GetMutableChunkFragmentPtr() const
+	{
+		static_assert(TIsDerivedFrom<T, FMassChunkFragment>::IsDerived, "Given struct doesn't represent a valid chunk fragment type. Make sure to inherit from FMassChunkFragment or one of its child-types.");
+
+		const UScriptStruct* Type = T::StaticStruct();
+		const FChunkFragmentView* FoundChunkFragmentData = ChunkFragmentViews.FindByPredicate([Type](const FChunkFragmentView& Element) { return Element.Requirement.StructType == Type; } );
+		return FoundChunkFragmentData ? FoundChunkFragmentData->FragmentView.GetMutablePtr<T>() : static_cast<T*>(nullptr);
+	}
+	
+	template<typename T>
+	T& GetMutableChunkFragment() const
+	{
+		T* ChunkFragment = GetMutableChunkFragmentPtr<T>();
+		checkf(ChunkFragment, TEXT("Chunk Fragment requirement not found: %s"), *T::StaticStruct()->GetName());
+		return *ChunkFragment;
+	}
+
+	template<typename T>
+	const T* GetChunkFragmentPtr() const
+	{
+		return GetMutableChunkFragmentPtr<T>();
+	}
+	
+	template<typename T>
+	const T& GetChunkFragment() const
+	{
+		const T* ChunkFragment = GetMutableChunkFragmentPtr<T>();
+		checkf(ChunkFragment, TEXT("Chunk Fragment requirement not found: %s"), *T::StaticStruct()->GetName());
+		return *ChunkFragment;
+	}
+
+	/** Shared fragment related operations */
+	template<typename T>
+	const T* GetConstSharedFragmentPtr() const
+	{
+		static_assert(TIsDerivedFrom<T, FMassSharedFragment>::IsDerived, "Given struct doesn't represent a valid Shared fragment type. Make sure to inherit from FMassSharedFragment or one of its child-types.");
+
+		const FConstSharedFragmentView* FoundSharedFragmentData = ConstSharedFragmentViews.FindByPredicate([](const FConstSharedFragmentView& Element) { return Element.Requirement.StructType == T::StaticStruct(); });
+		return FoundSharedFragmentData ? FoundSharedFragmentData->FragmentView.GetPtr<T>() : static_cast<const T*>(nullptr);
+	}
+
+	template<typename T>
+	const T& GetConstSharedFragment() const
+	{
+		const T* SharedFragment = GetConstSharedFragmentPtr<T>();
+		checkf(SharedFragment, TEXT("Shared Fragment requirement not found: %s"), *T::StaticStruct()->GetName());
+		return *SharedFragment;
+	}
+
+
+	template<typename T>
+	T* GetMutableSharedFragmentPtr() const
+	{
+		static_assert(TIsDerivedFrom<T, FMassSharedFragment>::IsDerived, "Given struct doesn't represent a valid shared fragment type. Make sure to inherit from FMassSharedFragment or one of its child-types.");
+
+		const FSharedFragmentView* FoundSharedFragmentData = SharedFragmentViews.FindByPredicate([](const FSharedFragmentView& Element) { return Element.Requirement.StructType == T::StaticStruct(); });
+		return FoundSharedFragmentData ? FoundSharedFragmentData->FragmentView.GetMutablePtr<T>() : static_cast<T*>(nullptr);
+	}
+
+	template<typename T>
+	T& GetMutableSharedFragment() const
+	{
+		T* SharedFragment = GetMutableSharedFragmentPtr<T>();
+		checkf(SharedFragment, TEXT("Shared Fragment requirement not found: %s"), *T::StaticStruct()->GetName());
+		return *SharedFragment;
+	}
+
+	template<typename T>
+	const T* GetSharedFragmentPtr() const
+	{
+		return GetMutableSharedFragmentPtr<T>();
+	}
+
+	template<typename T>
+	const T& GetSharedFragment() const
+	{
+		const T* SharedFragment = GetSharedFragmentPtr<T>();
+		checkf(SharedFragment, TEXT("Shared Fragment requirement not found: %s"), *T::StaticStruct()->GetName());
+		return *SharedFragment;
+	}
+
+	/* Fragments related operations */
+	template<typename TFragment>
+	TArrayView<TFragment> GetMutableFragmentView()
+	{
+		const UScriptStruct* FragmentType = TFragment::StaticStruct();
+		const FFragmentView* View = FragmentViews.FindByPredicate([FragmentType](const FFragmentView& Element) { return Element.Requirement.StructType == FragmentType; });
+		//checkfSlow(View != nullptr, TEXT("Requested fragment type not bound"));
+		//checkfSlow(View->Requirement.AccessMode == EMassFragmentAccess::ReadWrite, TEXT("Requested fragment has not been bound for writing"));
+		return MakeArrayView<TFragment>((TFragment*)View->FragmentView.GetData(), View->FragmentView.Num());
+	}
+
+	template<typename TFragment>
+	TConstArrayView<TFragment> GetFragmentView() const
+	{
+		const UScriptStruct* FragmentType = TFragment::StaticStruct();
+		const FFragmentView* View = FragmentViews.FindByPredicate([FragmentType](const FFragmentView& Element) { return Element.Requirement.StructType == FragmentType; });
+		//checkfSlow(View != nullptr, TEXT("Requested fragment type not bound"));
+		return TConstArrayView<TFragment>((const TFragment*)View->FragmentView.GetData(), View->FragmentView.Num());
+	}
+
+	TConstArrayView<FMassFragment> GetFragmentView(const UScriptStruct* FragmentType) const
+	{
+		const FFragmentView* View = FragmentViews.FindByPredicate([FragmentType](const FFragmentView& Element) { return Element.Requirement.StructType == FragmentType; });
+		checkSlow(View);
+		return TConstArrayView<FMassFragment>((const FMassFragment*)View->FragmentView.GetData(), View->FragmentView.Num());;
+	}
+
+	TArrayView<FMassFragment> GetMutableFragmentView(const UScriptStruct* FragmentType) 
+	{
+		const FFragmentView* View = FragmentViews.FindByPredicate([FragmentType](const FFragmentView& Element) { return Element.Requirement.StructType == FragmentType; });
+		checkSlow(View);
+		return View->FragmentView;
+	}
+
+	/** Sparse chunk related operation */
+	const FMassArchetypeEntityCollection& GetEntityCollection() const { return EntityCollection; }
+
+	const FInstancedStruct& GetAuxData() const { return AuxData; }
+	FInstancedStruct& GetMutableAuxData() { return AuxData; }
+	
+	template<typename TFragment>
+	bool ValidateAuxDataType() const
+	{
+		const UScriptStruct* FragmentType = GetAuxData().GetScriptStruct();
+		return FragmentType != nullptr && FragmentType == TFragment::StaticStruct();
+	}
+
+	void FlushDeferred(UMassEntitySubsystem& EntitySystem) const;
+
+	void ClearExecutionData();
+	void SetCurrentArchetypesTagBitSet(const FMassTagBitSet& BitSet)
+	{
+		CurrentArchetypesTagBitSet = BitSet;
+	}
+
+protected:
+	void SetRequirements(TConstArrayView<FMassFragmentRequirement> InRequirements, 
+		TConstArrayView<FMassFragmentRequirement> InChunkRequirements, 
+		TConstArrayView<FMassFragmentRequirement> InConstSharedRequirements, 
+		TConstArrayView<FMassFragmentRequirement> InSharedRequirements);
+
+	void ClearFragmentViews()
+	{
+		for (FFragmentView& View : FragmentViews)
+		{
+			View.FragmentView = TArrayView<FMassFragment>();
+		}
+		for (FChunkFragmentView& View : ChunkFragmentViews)
+		{
+			View.FragmentView.Reset();
+		}
+		for (FConstSharedFragmentView& View : ConstSharedFragmentViews)
+		{
+			View.FragmentView.Reset();
+		}
+		for (FSharedFragmentView& View : SharedFragmentViews)
+		{
+			View.FragmentView.Reset();
+		}
+	}
+};
