Modified bitset initialization performed in UMassEntitySubsystem::Initialize to exclude the base types.

diff --git a/Engine/Plugins/Runtime/MassEntity/Source/MassEntity/Private/MassEntitySubsystem.cpp b/Engine/Plugins/Runtime/MassEntity/Source/MassEntity/Private/MassEntitySubsystem.cpp
--- a/Engine/Plugins/Runtime/MassEntity/Source/MassEntity/Private/MassEntitySubsystem.cpp
+++ b/Engine/Plugins/Runtime/MassEntity/Source/MassEntity/Private/MassEntitySubsystem.cpp
@@ -71,25 +71,41 @@ void UMassEntitySubsystem::Initialize(FSubsystemCollectionBase& Collection)
 
 	DeferredCommandBuffer = MakeShareable(new FMassCommandBuffer());
 
-	// initialize the bitsets
-
+	// creating these bitset instances to populate respective bitset types' StructTrackers
 	FMassFragmentBitSet Fragments;
 	FMassTagBitSet Tags;
 	FMassChunkFragmentBitSet ChunkFragments;
+	FMassSharedFragmentBitSet LocalSharedFragments;
 
-	for (TObjectIterator<UStruct> StructIt; StructIt; ++StructIt)
+	for (TObjectIterator<UScriptStruct> StructIt; StructIt; ++StructIt)
 	{
 		if (StructIt->IsChildOf(FMassFragment::StaticStruct()))
 		{
-			Fragments.Add(*CastChecked<UScriptStruct>(*StructIt));
+			if (*StructIt != FMassFragment::StaticStruct())
+			{
+				Fragments.Add(**StructIt);
+			}
 		}
 		else if (StructIt->IsChildOf(FMassTag::StaticStruct()))
 		{
-			Tags.Add(*CastChecked<UScriptStruct>(*StructIt));
+			if (*StructIt != FMassTag::StaticStruct())
+			{
+				Tags.Add(**StructIt);
+			}
 		}
 		else if (StructIt->IsChildOf(FMassChunkFragment::StaticStruct()))
 		{
-			ChunkFragments.Add(*CastChecked<UScriptStruct>(*StructIt));
+			if (*StructIt != FMassChunkFragment::StaticStruct())
+			{
+				ChunkFragments.Add(**StructIt);
+			}
+		}
+		else if (StructIt->IsChildOf(FMassSharedFragment::StaticStruct()))
+		{
+			if (*StructIt != FMassSharedFragment::StaticStruct())
+			{
+				LocalSharedFragments.Add(**StructIt);
+			}
 		}
 	}
 }
