Enriched MassEntityQuery with capability to declare USubsystem requirements, that then can be fetched from the MassExecutionContext.

MassEntityQuery requirements are declared just like other requirements with a dedicated AddSystemRequirement call.
MassExecutionContext supplies functions for fetching subsystem instances checking if active query declared them as dependency in a given mode (ReadOnly or ReadWrite).
Temporarily added a way to fetch a subsystem outside of a query, then all restrictions are lifted. This will be replaced in further work on multithreading mass work.

diff --git a/Engine/Plugins/Experimental/StructUtils/Source/StructUtils/Public/StructTypeBitSet.h b/Engine/Plugins/Experimental/StructUtils/Source/StructUtils/Public/StructTypeBitSet.h
--- a/Engine/Plugins/Experimental/StructUtils/Source/StructUtils/Public/StructTypeBitSet.h
+++ b/Engine/Plugins/Experimental/StructUtils/Source/StructUtils/Public/StructTypeBitSet.h
@@ -64,6 +64,8 @@ struct FStructTracker
 		return StructTypesList.IsValidIndex(StructTypeIndex) ? StructTypesList[StructTypeIndex].Get() : nullptr;
 	}
 
+	int32 Num() const {	return StructTypeToIndexSet.Num(); }
+
 #if WITH_STRUCTUTILS_DEBUG
 	/**
 	* @return index identifying given tag or INDEX_NONE if it has never been used/seen before.
@@ -478,6 +480,11 @@ public:
 		return StructTypesBitArray.Contains(BitIndex);
 	}
 
+	static int32 GetMaxNum() 
+	{		
+		return TStructTrackerWrapper::StructTracker.Num();
+	}
+
 	FORCEINLINE bool operator==(const TStructTypeBitSet& Other) const { return StructTypesBitArray == Other.StructTypesBitArray; }
 	FORCEINLINE bool operator!=(const TStructTypeBitSet& Other) const { return !(*this == Other); }
 
diff --git a/Engine/Plugins/Runtime/MassEntity/Source/MassEntity/Private/MassEntityDebug.cpp b/Engine/Plugins/Runtime/MassEntity/Source/MassEntity/Private/MassEntityDebug.cpp
--- a/Engine/Plugins/Runtime/MassEntity/Source/MassEntity/Private/MassEntityDebug.cpp
+++ b/Engine/Plugins/Runtime/MassEntity/Source/MassEntity/Private/MassEntityDebug.cpp
@@ -137,7 +137,7 @@ FAutoConsoleCommandWithWorldArgsAndOutputDevice LogFragmentSizes(
 		{
 			for (const TWeakObjectPtr<const UScriptStruct>& WeakStruct : FMassFragmentBitSet::DebugGetAllStructTypes())
 			{
-				if (const UStruct* StructType = WeakStruct.Get())
+				if (const UScriptStruct* StructType = WeakStruct.Get())
 				{
 					Ar.Logf(ELogVerbosity::Log, TEXT("%s, size: %d"), *StructType->GetName(), StructType->GetStructureSize());
 				}
@@ -166,7 +166,7 @@ FAutoConsoleCommandWithOutputDevice LogFragments(
 {
 	auto PrintKnownTypes = [&OutputDevice](TConstArrayView<TWeakObjectPtr<const UScriptStruct>> AllStructs) {
 		int i = 0;
-		for (TWeakObjectPtr<const UStruct> Struct : AllStructs)
+		for (TWeakObjectPtr<const UScriptStruct> Struct : AllStructs)
 		{
 			if (Struct.IsValid())
 			{
diff --git a/Engine/Plugins/Runtime/MassEntity/Source/MassEntity/Private/MassEntityQuery.cpp b/Engine/Plugins/Runtime/MassEntity/Source/MassEntity/Private/MassEntityQuery.cpp
--- a/Engine/Plugins/Runtime/MassEntity/Source/MassEntity/Private/MassEntityQuery.cpp
+++ b/Engine/Plugins/Runtime/MassEntity/Source/MassEntity/Private/MassEntityQuery.cpp
@@ -138,6 +138,8 @@ void FMassEntityQuery::ForEachEntityChunk(UMassEntitySubsystem& EntitySubsystem,
 	int32 NumEntitiesToProcess = 0;
 #endif
 
+	ExecutionContext.SetSystemRequirements(RequiredConstSystems, RequiredMutableSystems);
+
 	// if there's a chunk collection set by the external code - use that
 	if (ExecutionContext.GetEntityCollection().IsSet())
 	{
diff --git a/Engine/Plugins/Runtime/MassEntity/Source/MassEntity/Private/MassEntityTypes.cpp b/Engine/Plugins/Runtime/MassEntity/Source/MassEntity/Private/MassEntityTypes.cpp
--- a/Engine/Plugins/Runtime/MassEntity/Source/MassEntity/Private/MassEntityTypes.cpp
+++ b/Engine/Plugins/Runtime/MassEntity/Source/MassEntity/Private/MassEntityTypes.cpp
@@ -7,6 +7,7 @@ DEFINE_TYPEBITSET(FMassFragmentBitSet);
 DEFINE_TYPEBITSET(FMassTagBitSet);
 DEFINE_TYPEBITSET(FMassChunkFragmentBitSet);
 DEFINE_TYPEBITSET(FMassSharedFragmentBitSet);
+DEFINE_TYPEBITSET(FMassExternalSubystemBitSet);
 
 uint32 FMassArchetypeSharedFragmentValues::CalculateHash() const
 {
diff --git a/Engine/Plugins/Runtime/MassEntity/Source/MassEntity/Private/MassExecutionContext.cpp b/Engine/Plugins/Runtime/MassEntity/Source/MassEntity/Private/MassExecutionContext.cpp
--- a/Engine/Plugins/Runtime/MassEntity/Source/MassEntity/Private/MassExecutionContext.cpp
+++ b/Engine/Plugins/Runtime/MassEntity/Source/MassEntity/Private/MassExecutionContext.cpp
@@ -17,6 +17,16 @@ void FMassExecutionContext::ClearExecutionData()
 	FragmentViews.Reset();
 	CurrentArchetypesTagBitSet.Reset();
 	ChunkSerialModificationNumber = -1;
+	ConstSubsystemsBitSet.Reset();
+	MutableSubsystemsBitSet.Reset();
+	bSubsystemRequirementsSet = false;
+}
+
+void FMassExecutionContext::SetSystemRequirements(const FMassExternalSubystemBitSet& RequiredConstSystems, const FMassExternalSubystemBitSet& RequiredMutableSystems)
+{
+	ConstSubsystemsBitSet = RequiredConstSystems;
+	MutableSubsystemsBitSet = RequiredMutableSystems;
+	bSubsystemRequirementsSet = true;
 }
 
 void FMassExecutionContext::SetEntityCollection(const FMassArchetypeEntityCollection& InEntityCollection)
diff --git a/Engine/Plugins/Runtime/MassEntity/Source/MassEntity/Public/MassEntityQuery.h b/Engine/Plugins/Runtime/MassEntity/Source/MassEntity/Public/MassEntityQuery.h
--- a/Engine/Plugins/Runtime/MassEntity/Source/MassEntity/Public/MassEntityQuery.h
+++ b/Engine/Plugins/Runtime/MassEntity/Source/MassEntity/Public/MassEntityQuery.h
@@ -15,7 +15,6 @@ struct FMassArchetypeData;
 struct FMassExecutionContext;
 struct FMassFragment;
 struct FMassArchetypeHandle;
-struct FMassExecutionContext;
 
 enum class EMassFragmentAccess : uint8
 {
@@ -309,6 +308,26 @@ public:
 		return *this;
 	}
 
+	template<typename T>
+	FMassEntityQuery& AddSystemRequirement(const EMassFragmentAccess AccessMode)
+	{
+		check(AccessMode != EMassFragmentAccess::None && AccessMode != EMassFragmentAccess::MAX);
+
+		switch (AccessMode)
+		{
+		case EMassFragmentAccess::ReadOnly:
+			RequiredConstSystems.Add<T>();
+			break;
+		case EMassFragmentAccess::ReadWrite:
+			RequiredMutableSystems.Add<T>();
+			break;
+		default:
+			check(false);
+		}
+		
+		return *this;
+	}
+
 	void Clear()
 	{
 		Requirements.Empty();
@@ -409,6 +428,9 @@ protected:
 	FMassSharedFragmentBitSet RequiredOptionalSharedFragments;
 	FMassSharedFragmentBitSet RequiredNoneSharedFragments;
 
+	FMassExternalSubystemBitSet RequiredConstSystems;
+	FMassExternalSubystemBitSet RequiredMutableSystems;
+
 private:
 	/** 
 	 * This function represents a condition that will be called for every chunk to be processed before the actual 
diff --git a/Engine/Plugins/Runtime/MassEntity/Source/MassEntity/Public/MassEntitySubsystem.h b/Engine/Plugins/Runtime/MassEntity/Source/MassEntity/Public/MassEntitySubsystem.h
--- a/Engine/Plugins/Runtime/MassEntity/Source/MassEntity/Public/MassEntitySubsystem.h
+++ b/Engine/Plugins/Runtime/MassEntity/Source/MassEntity/Public/MassEntitySubsystem.h
@@ -9,6 +9,7 @@
 #include "MassEntityQuery.h"
 #include "StructUtilsTypes.h"
 #include "MassObserverManager.h"
+#include "MassExecutionContext.h"
 #include "Containers/MpscQueue.h"
 #include "MassExecutionContext.h"
 #include "MassEntitySubsystem.generated.h"
diff --git a/Engine/Plugins/Runtime/MassEntity/Source/MassEntity/Public/MassEntityTypes.h b/Engine/Plugins/Runtime/MassEntity/Source/MassEntity/Public/MassEntityTypes.h
--- a/Engine/Plugins/Runtime/MassEntity/Source/MassEntity/Public/MassEntityTypes.h
+++ b/Engine/Plugins/Runtime/MassEntity/Source/MassEntity/Public/MassEntityTypes.h
@@ -5,6 +5,7 @@
 #include "StructTypeBitSet.h"
 #include "MassProcessingTypes.h"
 #include "StructArrayView.h"
+#include "Subsystems/Subsystem.h"
 #include "MassEntityTypes.generated.h"
 
 
@@ -106,6 +107,7 @@ DECLARE_STRUCTTYPEBITSET_EXPORTED(MASSENTITY_API, FMassFragmentBitSet, FMassFrag
 DECLARE_STRUCTTYPEBITSET_EXPORTED(MASSENTITY_API, FMassTagBitSet, FMassTag);
 DECLARE_STRUCTTYPEBITSET_EXPORTED(MASSENTITY_API, FMassChunkFragmentBitSet, FMassChunkFragment);
 DECLARE_STRUCTTYPEBITSET_EXPORTED(MASSENTITY_API, FMassSharedFragmentBitSet, FMassSharedFragment);
+DECLARE_CLASSTYPEBITSET_EXPORTED(MASSENTITY_API, FMassExternalSubystemBitSet, USubsystem);
 
 /** The type summarily describing a composition of an entity or an archetype. It contains information on both the
  *  fragments as well as tags */
diff --git a/Engine/Plugins/Runtime/MassEntity/Source/MassEntity/Public/MassExecutionContext.h b/Engine/Plugins/Runtime/MassEntity/Source/MassEntity/Public/MassExecutionContext.h
--- a/Engine/Plugins/Runtime/MassEntity/Source/MassEntity/Public/MassExecutionContext.h
+++ b/Engine/Plugins/Runtime/MassEntity/Source/MassEntity/Public/MassExecutionContext.h
@@ -3,6 +3,7 @@
 #pragma once
 
 #include "MassEntityTypes.h"
+#include "MassExternalSubsystemTraits.h"
 #include "MassEntityQuery.h"
 
 
@@ -32,6 +33,10 @@ private:
 
 	using FSharedFragmentView = TFragmentView<FStructView>;
 	TArray<FSharedFragmentView, TInlineAllocator<4>> SharedFragmentViews;
+
+	FMassExternalSubystemBitSet ConstSubsystemsBitSet;
+	FMassExternalSubystemBitSet MutableSubsystemsBitSet;
+	TArray<TObjectPtr<USubsystem>> Subsystems;
 	
 	// mz@todo make this shared ptr thread-safe and never auto-flush in MT environment. 
 	TSharedPtr<FMassCommandBuffer> DeferredCommandBuffer;
@@ -56,6 +61,11 @@ private:
 	 *  calling MassEntitySubsystem.FlushCommands() */
 	bool bFlushDeferredCommands = true;
 
+	/** A temporary variable letting us access subsystems via a FMassExecutionContext outside of MassEntityQuery's 
+	 * execution. Will go away once entity-independent calculations as well as query-bound calculations are being 
+	 * performed with a common tasks interface */
+	bool bSubsystemRequirementsSet = false;
+
 	TArrayView<FFragmentView> GetMutableRequirements() { return FragmentViews; }
 	TArrayView<FChunkFragmentView> GetMutableChunkRequirements() { return ChunkFragmentViews; }
 	TArrayView<FConstSharedFragmentView> GetMutableConstSharedRequirements() { return ConstSharedFragmentViews; }
@@ -65,10 +75,15 @@ private:
 	friend FMassEntityQuery;
 
 public:
-	FMassExecutionContext() = default;
 	explicit FMassExecutionContext(const float InDeltaTimeSeconds, const bool bInFlushDeferredCommands = true)
 		: DeltaTimeSeconds(InDeltaTimeSeconds)
 		, bFlushDeferredCommands(bInFlushDeferredCommands)
+	{
+		Subsystems.AddZeroed(FMassExternalSubystemBitSet::GetMaxNum());
+	}
+
+	FMassExecutionContext()
+		: FMassExecutionContext(/*InDeltaTimeSeconds=*/0.f)
 	{}
 
 #if WITH_MASSENTITY_DEBUG
@@ -228,6 +243,49 @@ public:
 		return View->FragmentView;
 	}
 
+	template<typename T>
+	T* GetMutableSubsystem(const UWorld* World)
+	{
+		// @todo consider getting this directly from entity subsystem - it should cache all the used system
+		const uint32 SystemIndex = FMassExternalSubystemBitSet::GetTypeIndex<T>();
+		if (bSubsystemRequirementsSet == false || ensure(MutableSubsystemsBitSet.IsBitSet(SystemIndex)))
+		{
+			return GetSystemInternal<T>(World, SystemIndex);
+		}
+
+		return nullptr;
+	}
+
+	template<typename T>
+	T& GetMutableSubsystemChecked(const UWorld* World)
+	{
+		T* InstancePtr = GetMutableSubsystem<T>(World);
+		check(InstancePtr);
+		return *InstancePtr;
+	}
+
+	template<typename T>
+	const T* GetSubsystem(const UWorld* World)
+	{
+		const uint32 SystemIndex = FMassExternalSubystemBitSet::GetTypeIndex<T>();
+		if (bSubsystemRequirementsSet == false || ensure(ConstSubsystemsBitSet.IsBitSet(SystemIndex) || MutableSubsystemsBitSet.IsBitSet(SystemIndex)))
+		{
+			return GetSystemInternal<T>(World, SystemIndex);
+		}
+		return nullptr;
+	}
+
+	template<typename T>
+	const T& GetSubsystemChecked(const UWorld* World)
+	{
+		const T* InstancePtr = GetSubsystem<T>(World);
+		check(InstancePtr);
+		return *InstancePtr;
+	}
+	
+
+	void SetSystemRequirements(const FMassExternalSubystemBitSet& RequiredConstSystems, const FMassExternalSubystemBitSet& RequiredMutableSystems);
+
 	/** Sparse chunk related operation */
 	const FMassArchetypeEntityCollection& GetEntityCollection() const { return EntityCollection; }
 
@@ -274,4 +332,21 @@ protected:
 			View.FragmentView.Reset();
 		}
 	}
+
+	template<typename T>
+	T* GetSystemInternal(const UWorld* World, const uint32 SystemIndex)
+	{
+		if (UNLIKELY(Subsystems.IsValidIndex(SystemIndex) == false))
+		{
+			Subsystems.AddZeroed(Subsystems.Num() - SystemIndex + 1);
+		}
+
+		T* SystemInstance = (T*)Subsystems[SystemIndex].Get();
+		if (SystemInstance == nullptr)
+		{
+			SystemInstance = FMassExternalSubsystemTraits::GetInstance<typename TRemoveConst<T>::Type>(World);
+			Subsystems[SystemIndex] = SystemInstance;
+		}
+		return SystemInstance;
+	}
 };
diff --git a/Engine/Plugins/Runtime/MassEntity/Source/MassEntity/Public/MassExternalSubsystemTraits.h b/Engine/Plugins/Runtime/MassEntity/Source/MassEntity/Public/MassExternalSubsystemTraits.h
new file mode 100644
--- /dev/null
+++ b/Engine/Plugins/Runtime/MassEntity/Source/MassEntity/Public/MassExternalSubsystemTraits.h
@@ -0,0 +1,49 @@
+// Copyright Epic Games, Inc. All Rights Reserved.
+
+#pragma once
+
+#include "Engine/World.h"
+#include "Subsystems/WorldSubsystem.h"
+
+/**
+ * Traits describing how a given piece of code can be used by Mass. By default everything is treated as
+ * needing to be accessed only on game thread (implying single-thread access). To change the behavior for a given
+ * class just specialize this template, like so:
+ *
+ * template<>
+ * struct TMassExternalSubsystemTraits<UMyCustomManager>
+ * {
+ *		enum { GameThreadOnly = false; }
+ * }
+ *
+ * this will let Mass know it can access UMyCustomManager on any thread.
+ *
+ * This information is being used to calculate processor and query dependencies as well as appropriate distribution of
+ * calculations across threads.
+ */
+template <typename T>
+struct TMassExternalSubsystemTraits final
+{
+	enum
+	{
+		GameThreadOnly = true,
+		ThreadSafeRead = false,
+		ThreadSafeWrite = false,
+	};
+};
+
+namespace FMassExternalSubsystemTraits
+{
+	/** 
+	Every TMassExternalSubsystemTraits specialization needs to implement the following. Not supplying default implementations
+	to be able to catch missing implementations and header inclusion at compilation time.
+	
+	This is a getter function that given a UWorld* fetches an instance.
+	*/
+	template<typename T, typename = typename TEnableIf<TIsDerivedFrom<T, UWorldSubsystem>::IsDerived>::Type>
+	static T* GetInstance(const UWorld* World)
+	{ 
+		// note that the default implementation works only for UWorldSubsystems
+		return UWorld::GetSubsystem<T>(World);
+	}
+}
