Mass Subsystem requirements unittests

diff --git a/Engine/Plugins/Runtime/MassEntity/Source/MassEntityTestSuite/Private/MassSystemRequirementTest.cpp b/Engine/Plugins/Runtime/MassEntity/Source/MassEntityTestSuite/Private/MassSystemRequirementTest.cpp
new file mode 100644
--- /dev/null
+++ b/Engine/Plugins/Runtime/MassEntity/Source/MassEntityTestSuite/Private/MassSystemRequirementTest.cpp
@@ -0,0 +1,99 @@
+// Copyright Epic Games, Inc. All Rights Reserved.
+
+#include "CoreMinimal.h"
+#include "AITestsCommon.h"
+
+#include "MassProcessorDependencySolver.h"
+#include "MassEntityTestTypes.h"
+
+// D:\p4\Starship_review\Engine\Source\Runtime\NavigationSystem\Public\NavigationSystem.h
+//#include "MassNavigationSubsystem.h"
+
+#define LOCTEXT_NAMESPACE "MassTest"
+
+PRAGMA_DISABLE_OPTIMIZATION
+
+//template<>
+//struct TMassExternalSystemGetter<UMassEntitySubsystem>
+//{
+//	
+//};
+
+namespace FMassSystemRequirementTest
+{
+
+//FEntityTestBase : FExecutionTestBase
+struct FSystemRequirementTestBase : FEntityTestBase//FExecutionTestBase
+{
+	using Super = FEntityTestBase;
+	virtual bool SetUp() override
+	{		
+		return Super::SetUp();
+	}
+};	
+
+struct FMutableRequirement : FSystemRequirementTestBase
+{
+	virtual bool InstantTest() override
+	{
+		EntitySubsystem->CreateEntity(FloatsArchetype);
+
+		FMassEntityQuery EntityQuery;
+		EntityQuery.AddRequirement<FTestFragment_Float>(EMassFragmentAccess::ReadOnly);
+		EntityQuery.AddSystemRequirement<UMassTestWorldSubsystem>(EMassFragmentAccess::ReadWrite);
+
+		check(EntitySubsystem->GetWorld());
+		UMassTestWorldSubsystem* TestSystemActual = EntitySubsystem->GetWorld()->GetSubsystem<UMassTestWorldSubsystem>();
+		UMassTestWorldSubsystem* TestSystem = nullptr;
+		const UMassTestWorldSubsystem* TestSystem2 = nullptr;
+		FMassExecutionContext ExecutionContext;
+		EntityQuery.ForEachEntityChunk(*EntitySubsystem, ExecutionContext, [this, &TestSystem, &TestSystem2, World = EntitySubsystem->GetWorld()](FMassExecutionContext& Context)
+		{
+			//UMassEntitySubsystem* NavSystem = Context.GetSubsystem<UMassEntitySubsystem>(World);
+			TestSystem = Context.GetMutableSubsystem<UMassTestWorldSubsystem>(World);
+			TestSystem2 = Context.GetSubsystem<UMassTestWorldSubsystem>(World);
+		});
+
+		AITEST_NOT_NULL(TEXT(""), TestSystem);
+		AITEST_EQUAL(TEXT(""), TestSystem, TestSystemActual);
+		AITEST_EQUAL(TEXT(""), TestSystem, TestSystem2);		
+
+		return true;
+	}
+};
+IMPLEMENT_AI_INSTANT_TEST(FMutableRequirement, "System.Mass.Query.System.TrivialMutable");
+
+struct FConstRequirement : FSystemRequirementTestBase
+{
+	virtual bool InstantTest() override
+	{
+		EntitySubsystem->CreateEntity(FloatsArchetype);
+
+		FMassEntityQuery EntityQuery;
+		EntityQuery.AddRequirement<FTestFragment_Float>(EMassFragmentAccess::ReadOnly);
+		EntityQuery.AddSystemRequirement<UMassTestWorldSubsystem>(EMassFragmentAccess::ReadOnly);
+
+		check(EntitySubsystem->GetWorld());
+		UMassTestWorldSubsystem* TestMutableSystem = nullptr;
+		const UMassTestWorldSubsystem* TestConstSystem = nullptr;
+		FMassExecutionContext ExecutionContext;
+		EntityQuery.ForEachEntityChunk(*EntitySubsystem, ExecutionContext, [this, &TestMutableSystem, &TestConstSystem, World = EntitySubsystem->GetWorld()](FMassExecutionContext& Context)
+		{
+			// commented out since there's no way to seamlessly handle failed ensures, and the line below should cause that.
+			//TestMutableSystem = Context.GetMutableSubsystem<UMassTestWorldSubsystem>(World);
+			TestConstSystem = Context.GetSubsystem<UMassTestWorldSubsystem>(World);
+		});
+
+		AITEST_NOT_NULL(TEXT("It should be possible to only fetch const system instance"), TestConstSystem);
+		AITEST_NULL(TEXT("Mutable system access should not be possible"), TestMutableSystem);
+
+		return true;
+	}
+};
+IMPLEMENT_AI_INSTANT_TEST(FConstRequirement, "System.Mass.Query.System.TrivialConst");
+
+} // FMassSystemRequirementTest
+
+PRAGMA_ENABLE_OPTIMIZATION
+
+#undef LOCTEXT_NAMESPACE
diff --git a/Engine/Plugins/Runtime/MassEntity/Source/MassEntityTestSuite/Private/MassTestTypes.cpp b/Engine/Plugins/Runtime/MassEntity/Source/MassEntityTestSuite/Private/MassTestTypes.cpp
--- a/Engine/Plugins/Runtime/MassEntity/Source/MassEntityTestSuite/Private/MassTestTypes.cpp
+++ b/Engine/Plugins/Runtime/MassEntity/Source/MassEntityTestSuite/Private/MassTestTypes.cpp
@@ -81,4 +81,20 @@ UMassTestProcessor_FloatsInts::UMassTestProcessor_FloatsInts()
 		EntityQuery.AddRequirement<FTestFragment_Float>(EMassFragmentAccess::ReadWrite);
 		EntityQuery.AddRequirement<FTestFragment_Int>(EMassFragmentAccess::ReadWrite);
 	};
-}
\ No newline at end of file
+}
+
+//----------------------------------------------------------------------//
+// UMassTestWorldSubsystem
+//----------------------------------------------------------------------//
+void UMassTestWorldSubsystem::Write(int32 InNumber)
+{
+	UE_MT_SCOPED_WRITE_ACCESS(AccessDetector);
+	Number = InNumber;
+}
+
+int32 UMassTestWorldSubsystem::Read() const
+{
+	UE_MT_SCOPED_READ_ACCESS(AccessDetector);
+	return Number;
+}
+
diff --git a/Engine/Plugins/Runtime/MassEntity/Source/MassEntityTestSuite/Public/MassEntityTestTypes.h b/Engine/Plugins/Runtime/MassEntity/Source/MassEntityTestSuite/Public/MassEntityTestTypes.h
--- a/Engine/Plugins/Runtime/MassEntity/Source/MassEntityTestSuite/Public/MassEntityTestTypes.h
+++ b/Engine/Plugins/Runtime/MassEntity/Source/MassEntityTestSuite/Public/MassEntityTestTypes.h
@@ -6,6 +6,9 @@
 #include "MassEntityTypes.h"
 #include "MassEntitySubsystem.h"
 #include "AITestsCommon.h"
+#include "Subsystems/WorldSubsystem.h"
+#include "Misc/MTAccessDetector.h"
+#include "MassExternalSubsystemTraits.h"
 #include "MassEntityTestTypes.generated.h"
 
 
@@ -205,3 +208,27 @@ struct FEntityTestBase : FExecutionTestBase
 
 	virtual bool SetUp() override;
 };
+
+UCLASS()
+class UMassTestWorldSubsystem : public UWorldSubsystem
+{
+	GENERATED_BODY()
+public:
+	void Write(int32 InNumber);
+	int32 Read() const;
+
+private:
+	UE_MT_DECLARE_RW_ACCESS_DETECTOR(AccessDetector);
+	int Number = 0;
+};
+
+template<>
+struct TMassExternalSubsystemTraits<UMassTestWorldSubsystem>
+{
+	enum
+	{
+		GameThreadOnly = false,
+		ThreadSafeRead = true,
+		ThreadSafeWrite = false,
+	};
+};
