Added a MASS_DO_PARALLEL define to have a one-stop-shop compilation switch for mass' multithreaded mode (WIP).

diff --git a/Engine/Plugins/Runtime/MassEntity/Source/MassEntity/Private/MassProcessingPhase.cpp b/Engine/Plugins/Runtime/MassEntity/Source/MassEntity/Private/MassProcessingPhase.cpp
--- a/Engine/Plugins/Runtime/MassEntity/Source/MassEntity/Private/MassProcessingPhase.cpp
+++ b/Engine/Plugins/Runtime/MassEntity/Source/MassEntity/Private/MassProcessingPhase.cpp
@@ -17,7 +17,7 @@ DECLARE_CYCLE_STAT(TEXT("Mass Phase Done"), STAT_MassPhaseDone, STATGROUP_TaskGr
 
 namespace FMassTweakables
 {
-	bool bFullyParallel = false;
+	bool bFullyParallel = MASS_DO_PARALLEL;
 
 	FAutoConsoleVariableRef CVars[] = {
 		{TEXT("mass.FullyParallel"), bFullyParallel, TEXT("Enables mass processing distribution to all available thread (via the task graph)")},
diff --git a/Engine/Plugins/Runtime/MassEntity/Source/MassEntity/Private/MassProcessor.cpp b/Engine/Plugins/Runtime/MassEntity/Source/MassEntity/Private/MassProcessor.cpp
--- a/Engine/Plugins/Runtime/MassEntity/Source/MassEntity/Private/MassProcessor.cpp
+++ b/Engine/Plugins/Runtime/MassEntity/Source/MassEntity/Private/MassProcessor.cpp
@@ -9,7 +9,7 @@
 
 DECLARE_CYCLE_STAT(TEXT("MassProcessor Group Completed"), Mass_GroupCompletedTask, STATGROUP_TaskGraphTasks);
 
-#define PARALLELIZED_TRAFFIC_HACK 1
+#define PARALLELIZED_TRAFFIC_HACK !MASS_DO_PARALLEL
 
 #if PARALLELIZED_TRAFFIC_HACK
 namespace UE::MassTraffic
@@ -37,7 +37,7 @@ namespace UE::Mass::Debug
 
 namespace FMassTweakables
 {
-	bool bParallelGroups = false;
+	bool bParallelGroups = MASS_DO_PARALLEL;
 	float PostponedTaskWaitTimeWarningLevel = 0.002f;
 
 	FAutoConsoleVariableRef CVarsMassProcessor[] = {
diff --git a/Engine/Plugins/Runtime/MassEntity/Source/MassEntity/Public/MassEntitySettings.h b/Engine/Plugins/Runtime/MassEntity/Source/MassEntity/Public/MassEntitySettings.h
--- a/Engine/Plugins/Runtime/MassEntity/Source/MassEntity/Public/MassEntitySettings.h
+++ b/Engine/Plugins/Runtime/MassEntity/Source/MassEntity/Public/MassEntitySettings.h
@@ -8,6 +8,8 @@
 #include "InstancedStruct.h"
 #include "MassEntitySettings.generated.h"
 
+#define MASS_DO_PARALLEL 0
+
 #define GET_MASS_CONFIG_VALUE(a) (GetMutableDefault<UMassEntitySettings>()->a)
 
 class UMassProcessingPhaseManager;
