Exposed parts of MassEntityTestSuite for reuse in other test-focused plugins

diff --git a/Engine/Plugins/Runtime/MassEntity/Source/MassEntityTestSuite/Private/MassChunkCollectionTest.cpp b/Engine/Plugins/Runtime/MassEntity/Source/MassEntityTestSuite/Private/MassChunkCollectionTest.cpp
--- a/Engine/Plugins/Runtime/MassEntity/Source/MassEntityTestSuite/Private/MassChunkCollectionTest.cpp
+++ b/Engine/Plugins/Runtime/MassEntity/Source/MassEntityTestSuite/Private/MassChunkCollectionTest.cpp
@@ -14,15 +14,6 @@ PRAGMA_DISABLE_OPTIMIZATION
 namespace FMassEntityCollectionTest
 {
 
-void Shuffle(FRandomStream& Rand, TArray<FMassEntityHandle>& Entities)
-{
-	for (int i = 0; i < Entities.Num(); ++i)
-	{
-		const int32 NewIndex = Rand.RandRange(0, Entities.Num() - 1);
-		Entities.Swap(i, NewIndex);
-	}
-}
-
 struct FEntityCollectionTestBase : FEntityTestBase
 {
 	TArray<FMassEntityHandle> Entities;
@@ -95,7 +86,7 @@ struct FEntityCollection_CreateOrderInvariant : FEntityCollectionTestBase
 		FMassArchetypeEntityCollection CollectionFromOrdered(FloatsArchetype, EntitiesSubSet, FMassArchetypeEntityCollection::NoDuplicates);
 
 		FRandomStream Rand(0);
-		Shuffle(Rand, EntitiesSubSet);
+		ShuffleDataWithRandomStream(Rand, EntitiesSubSet);
 
 		FMassArchetypeEntityCollection CollectionFromRandom(FloatsArchetype, EntitiesSubSet, FMassArchetypeEntityCollection::NoDuplicates);
 
diff --git a/Engine/Plugins/Runtime/MassEntity/Source/MassEntityTestSuite/Public/MassEntityTestTypes.h b/Engine/Plugins/Runtime/MassEntity/Source/MassEntityTestSuite/Public/MassEntityTestTypes.h
--- a/Engine/Plugins/Runtime/MassEntity/Source/MassEntityTestSuite/Public/MassEntityTestTypes.h
+++ b/Engine/Plugins/Runtime/MassEntity/Source/MassEntityTestSuite/Public/MassEntityTestTypes.h
@@ -6,6 +6,7 @@
 #include "MassEntityTypes.h"
 #include "MassEntitySubsystem.h"
 #include "AITestsCommon.h"
+#include "Math/RandomStream.h"
 #include "Subsystems/WorldSubsystem.h"
 #include "Misc/MTAccessDetector.h"
 #include "MassExternalSubsystemTraits.h"
@@ -99,7 +100,7 @@ struct FTestTag_D : public FMassTag
 
 
 UCLASS()
-class UMassTestProcessorBase : public UMassProcessor
+class MASSENTITYTESTSUITE_API UMassTestProcessorBase : public UMassProcessor
 {
 	GENERATED_BODY()
 public:
@@ -124,37 +125,37 @@ protected:
 };
 
 UCLASS()
-class UMassTestProcessor_A : public UMassTestProcessorBase
+class MASSENTITYTESTSUITE_API UMassTestProcessor_A : public UMassTestProcessorBase
 {
 	GENERATED_BODY()
 };
 
 UCLASS()
-class UMassTestProcessor_B : public UMassTestProcessorBase
+class MASSENTITYTESTSUITE_API UMassTestProcessor_B : public UMassTestProcessorBase
 {
 	GENERATED_BODY()
 };
 
 UCLASS()
-class UMassTestProcessor_C : public UMassTestProcessorBase
+class MASSENTITYTESTSUITE_API UMassTestProcessor_C : public UMassTestProcessorBase
 {
 	GENERATED_BODY()
 };
 
 UCLASS()
-class UMassTestProcessor_D : public UMassTestProcessorBase
+class MASSENTITYTESTSUITE_API UMassTestProcessor_D : public UMassTestProcessorBase
 {
 	GENERATED_BODY()
 };
 
 UCLASS()
-class UMassTestProcessor_E : public UMassTestProcessorBase
+class MASSENTITYTESTSUITE_API UMassTestProcessor_E : public UMassTestProcessorBase
 {
 	GENERATED_BODY()
 };
 
 UCLASS()
-class UMassTestProcessor_F : public UMassTestProcessorBase
+class MASSENTITYTESTSUITE_API UMassTestProcessor_F : public UMassTestProcessorBase
 {
 	GENERATED_BODY()
 };
@@ -187,7 +188,7 @@ public:
 	UMassTestProcessor_FloatsInts();
 };
 
-struct FExecutionTestBase : FAITestBase
+struct MASSENTITYTESTSUITE_API FExecutionTestBase : FAITestBase
 {
 	UMassEntitySubsystem* EntitySubsystem = nullptr;
 	UWorld* World = nullptr;
@@ -197,7 +198,7 @@ struct FExecutionTestBase : FAITestBase
 
 const int TestIntValue = 123456;
 
-struct FEntityTestBase : FExecutionTestBase
+struct MASSENTITYTESTSUITE_API FEntityTestBase : FExecutionTestBase
 {
 	FMassArchetypeHandle EmptyArchetype;
 	FMassArchetypeHandle FloatsArchetype;
@@ -209,6 +210,16 @@ struct FEntityTestBase : FExecutionTestBase
 	virtual bool SetUp() override;
 };
 
+template<typename T>
+void ShuffleDataWithRandomStream(FRandomStream& Rand, TArray<T>& Data)
+{
+	for (int i = 0; i < Data.Num(); ++i)
+	{
+		const int32 NewIndex = Rand.RandRange(0, Data.Num() - 1);
+		Data.Swap(i, NewIndex);
+	}
+}
+
 UCLASS()
 class UMassTestWorldSubsystem : public UWorldSubsystem
 {
