Setting default value of MassEntitySubsystem's atomic counters

diff --git a/Engine/Plugins/Runtime/MassEntity/Source/MassEntity/Public/MassEntitySubsystem.h b/Engine/Plugins/Runtime/MassEntity/Source/MassEntity/Public/MassEntitySubsystem.h
--- a/Engine/Plugins/Runtime/MassEntity/Source/MassEntity/Public/MassEntitySubsystem.h
+++ b/Engine/Plugins/Runtime/MassEntity/Source/MassEntity/Public/MassEntitySubsystem.h
@@ -395,8 +395,8 @@ private:
 	TSharedPtr<FMassCommandBuffer> DeferredCommandBuffer;
 	TMpscQueue<TSharedPtr<FMassCommandBuffer>> FlushedCommandBufferQueue;
 
-	std::atomic<int32> SerialNumberGenerator;
-	std::atomic<int32> ProcessingScopeCount;
+	std::atomic<int32> SerialNumberGenerator = 0;
+	std::atomic<int32> ProcessingScopeCount = 0;
 
 	// the "version" number increased every time an archetype gets added
 	uint32 ArchetypeDataVersion = 0;
