Added a couple of missing stats to Mass code

diff --git a/Engine/Plugins/Runtime/MassEntity/Source/MassEntity/Private/MassCommandBuffer.cpp b/Engine/Plugins/Runtime/MassEntity/Source/MassEntity/Private/MassCommandBuffer.cpp
--- a/Engine/Plugins/Runtime/MassEntity/Source/MassEntity/Private/MassCommandBuffer.cpp
+++ b/Engine/Plugins/Runtime/MassEntity/Source/MassEntity/Private/MassCommandBuffer.cpp
@@ -163,6 +163,8 @@ void FMassCommandBuffer::CleanUp()
 
 void FMassCommandBuffer::MoveAppend(FMassCommandBuffer& Other)
 {
+	TRACE_CPUPROFILER_EVENT_SCOPE(MassCommandBuffer_MoveAppend);
+
 	// @todo optimize, there surely a way to do faster then this.
 	UE_MT_SCOPED_READ_ACCESS(Other.PendingBatchCommandsDetector);
 	if (Other.HasPendingCommands())
diff --git a/Engine/Plugins/Runtime/MassEntity/Source/MassEntity/Private/MassEntitySubsystem.cpp b/Engine/Plugins/Runtime/MassEntity/Source/MassEntity/Private/MassEntitySubsystem.cpp
--- a/Engine/Plugins/Runtime/MassEntity/Source/MassEntity/Private/MassEntitySubsystem.cpp
+++ b/Engine/Plugins/Runtime/MassEntity/Source/MassEntity/Private/MassEntitySubsystem.cpp
@@ -398,6 +398,8 @@ void UMassEntitySubsystem::BatchBuildEntities(const FMassArchetypeEntityCollecti
 
 void UMassEntitySubsystem::BatchBuildEntities(const FMassArchetypeEntityCollectionWithPayload& EncodedEntitiesWithPayload, FMassArchetypeCompositionDescriptor&& Composition, const FMassArchetypeSharedFragmentValues& SharedFragmentValues)
 {
+	TRACE_CPUPROFILER_EVENT_SCOPE(Mass_BatchBuildEntities);
+
 	FMassArchetypeEntityCollection::FEntityRangeArray TargetArchetypeEntityRanges;
 
 	// "built" entities case, this is verified during FMassArchetypeEntityCollectionWithPayload construction
@@ -443,6 +445,8 @@ void UMassEntitySubsystem::BatchBuildEntities(const FMassArchetypeEntityCollecti
 
 TSharedRef<UMassEntitySubsystem::FEntityCreationContext> UMassEntitySubsystem::BatchCreateEntities(const FMassArchetypeHandle Archetype, const int32 Count, TArray<FMassEntityHandle>& OutEntities)
 {
+	TRACE_CPUPROFILER_EVENT_SCOPE(Mass_BatchCreateEntities);
+
 	FMassArchetypeData* ArchetypePtr = Archetype.DataPtr.Get();
 	check(ArchetypePtr);
 	check(Count > 0);
@@ -499,7 +503,7 @@ void UMassEntitySubsystem::DestroyEntity(FMassEntityHandle Entity)
 
 void UMassEntitySubsystem::BatchDestroyEntities(TConstArrayView<FMassEntityHandle> InEntities)
 {
-	TRACE_CPUPROFILER_EVENT_SCOPE_STR("BatchDestroyEntities");
+	TRACE_CPUPROFILER_EVENT_SCOPE(Mass_BatchDestroyEntities);
 
 	checkf(IsProcessing() == false, TEXT("Synchronous API function %s called during mass processing. Use asynchronous API instead."), ANSI_TO_TCHAR(__FUNCTION__));
 	
@@ -530,7 +534,7 @@ void UMassEntitySubsystem::BatchDestroyEntities(TConstArrayView<FMassEntityHandl
 
 void UMassEntitySubsystem::BatchDestroyEntityChunks(const FMassArchetypeEntityCollection& EntityCollection)
 {
-	TRACE_CPUPROFILER_EVENT_SCOPE_STR("BatchDestroyEntityChunks");
+	TRACE_CPUPROFILER_EVENT_SCOPE(Mass_BatchDestroyEntityChunks);
 
 	checkf(IsProcessing() == false, TEXT("Synchronous API function %s called during mass processing. Use asynchronous API instead."), ANSI_TO_TCHAR(__FUNCTION__));
 
@@ -835,6 +839,8 @@ void UMassEntitySubsystem::RemoveTagFromEntity(FMassEntityHandle Entity, const U
 
 void UMassEntitySubsystem::BatchChangeTagsForEntities(TConstArrayView<FMassArchetypeEntityCollection> EntityCollections, const FMassTagBitSet& TagsToAdd, const FMassTagBitSet& TagsToRemove)
 {
+	TRACE_CPUPROFILER_EVENT_SCOPE(Mass_BatchChangeTagsForEntities);
+
 	for (const FMassArchetypeEntityCollection& Collection : EntityCollections)
 	{
 		FMassArchetypeData* CurrentArchetype = Collection.GetArchetype().DataPtr.Get();
@@ -883,6 +889,8 @@ void UMassEntitySubsystem::BatchChangeTagsForEntities(TConstArrayView<FMassArche
 
 void UMassEntitySubsystem::BatchChangeFragmentCompositionForEntities(TConstArrayView<FMassArchetypeEntityCollection> EntityCollections, const FMassFragmentBitSet& FragmentsToAdd, const FMassFragmentBitSet& FragmentsToRemove)
 {
+	TRACE_CPUPROFILER_EVENT_SCOPE(Mass_BatchChangeFragmentCompositionForEntities);
+
 	for (const FMassArchetypeEntityCollection& Collection : EntityCollections)
 	{
 		FMassArchetypeData* CurrentArchetype = Collection.GetArchetype().DataPtr.Get();
@@ -934,6 +942,8 @@ void UMassEntitySubsystem::BatchChangeFragmentCompositionForEntities(TConstArray
 
 void UMassEntitySubsystem::BatchAddFragmentInstancesForEntities(TConstArrayView<FMassArchetypeEntityCollectionWithPayload> EntityCollections, const FMassFragmentBitSet& FragmentsAffected)
 {
+	TRACE_CPUPROFILER_EVENT_SCOPE(Mass_BatchAddFragmentInstancesForEntities);
+
 	// here's the scenario:
 	// * we get entities from potentially different archetypes
 	// * adding a fragment instance consists of two operations: A) add fragment type & B) set fragment value
@@ -1030,6 +1040,8 @@ void UMassEntitySubsystem::SetEntityFragmentsValues(FMassEntityHandle Entity, TA
 
 void UMassEntitySubsystem::BatchSetEntityFragmentsValues(const FMassArchetypeEntityCollection& SparseEntities, TArrayView<const FInstancedStruct> FragmentInstanceList)
 {
+	TRACE_CPUPROFILER_EVENT_SCOPE(Mass_BatchSetEntityFragmentsValues);
+
 	FMassArchetypeData* Archetype = SparseEntities.GetArchetype().DataPtr.Get();
 	check(Archetype);
 
diff --git a/Engine/Plugins/Runtime/MassEntity/Source/MassEntity/Private/MassObserverManager.cpp b/Engine/Plugins/Runtime/MassEntity/Source/MassEntity/Private/MassObserverManager.cpp
--- a/Engine/Plugins/Runtime/MassEntity/Source/MassEntity/Private/MassObserverManager.cpp
+++ b/Engine/Plugins/Runtime/MassEntity/Source/MassEntity/Private/MassObserverManager.cpp
@@ -212,6 +212,8 @@ void FMassObserverManager::HandleFragmentsImpl(FMassProcessingContext& Processin
 	, TArrayView<const UScriptStruct*> ObservedTypes
 	/*, const FMassFragmentBitSet& FragmentsBitSet*/, FMassObserversMap& HandlersContainer)
 {	
+	TRACE_CPUPROFILER_EVENT_SCOPE(MassObserver_HandleFragmentsImpl);
+
 	check(ObservedTypes.Num() > 0);
 
 	for (const UScriptStruct* Type : ObservedTypes)
@@ -225,6 +227,8 @@ void FMassObserverManager::HandleFragmentsImpl(FMassProcessingContext& Processin
 
 void FMassObserverManager::HandleSingleEntityImpl(const UScriptStruct& FragmentType, const FMassArchetypeEntityCollection& EntityCollection, FMassObserversMap& HandlersContainer)
 {
+	TRACE_CPUPROFILER_EVENT_SCOPE(MassObserver_HandleSingleEntityImpl);
+
 	FMassProcessingContext ProcessingContext(EntitySubsystem, /*DeltaSeconds=*/0.f);
 	ProcessingContext.bFlushCommandBuffer = false;
 	ProcessingContext.AuxData.InitializeAs(&FragmentType);
diff --git a/Engine/Plugins/Runtime/MassEntity/Source/MassEntity/Private/MassProcessingPhase.cpp b/Engine/Plugins/Runtime/MassEntity/Source/MassEntity/Private/MassProcessingPhase.cpp
--- a/Engine/Plugins/Runtime/MassEntity/Source/MassEntity/Private/MassProcessingPhase.cpp
+++ b/Engine/Plugins/Runtime/MassEntity/Source/MassEntity/Private/MassProcessingPhase.cpp
@@ -55,6 +55,8 @@ void FMassProcessingPhase::ExecuteTick(float DeltaTime, ELevelTick TickType, ENa
 
 	checkf(Manager, TEXT("Manager is null which is not a supported case. Either this FMassProcessingPhase has not been initialized properly or it's been left dangling after the FMassProcessingPhase owner got destroyed."));
 
+	TRACE_CPUPROFILER_EVENT_SCOPE_STR(*FString::Printf(TEXT("FMassProcessingPhase::ExecuteTick %s"), *UEnum::GetValueAsString(Phase)));
+
 	Manager->OnPhaseStart(*this);
 	
 	OnPhaseStart.Broadcast(DeltaTime);
diff --git a/Engine/Plugins/Runtime/MassEntity/Source/MassEntity/Public/MassCommands.h b/Engine/Plugins/Runtime/MassEntity/Source/MassEntity/Public/MassCommands.h
--- a/Engine/Plugins/Runtime/MassEntity/Source/MassEntity/Public/MassCommands.h
+++ b/Engine/Plugins/Runtime/MassEntity/Source/MassEntity/Public/MassCommands.h
@@ -119,7 +119,9 @@ struct FMassCommandDestroyEntities : public FMassBatchedEntityCommand
 protected:
 	virtual void Execute(UMassEntitySubsystem& System) const override
 	{
-		TArray<FMassArchetypeEntityCollection > EntityCollectionsToDestroy;
+		TRACE_CPUPROFILER_EVENT_SCOPE(MassCommandDestroyEntities_Execute);
+
+		TArray<FMassArchetypeEntityCollection> EntityCollectionsToDestroy;
 		UE::Mass::Utils::CreateEntityCollections(System, TargetEntities, FMassArchetypeEntityCollection::FoldDuplicates, EntityCollectionsToDestroy);
 		for (FMassArchetypeEntityCollection& Collection : EntityCollectionsToDestroy)
 		{
@@ -142,6 +144,7 @@ struct FMassCommandAddFragments : public FMassBatchedEntityCommand
 protected:
 	virtual void Execute(UMassEntitySubsystem& System) const override
 	{
+		TRACE_CPUPROFILER_EVENT_SCOPE(MassCommandAddFragments_Execute);
 		TArray<FMassArchetypeEntityCollection> EntityCollections;
 		UE::Mass::Utils::CreateEntityCollections(System, TargetEntities, FMassArchetypeEntityCollection::FoldDuplicates, EntityCollections);
 		System.BatchChangeFragmentCompositionForEntities(EntityCollections, FragmentsAffected, FMassFragmentBitSet());
@@ -160,6 +163,7 @@ struct FMassCommandRemoveFragments : public FMassBatchedEntityCommand
 protected:
 	virtual void Execute(UMassEntitySubsystem& System) const override
 	{
+		TRACE_CPUPROFILER_EVENT_SCOPE(MassCommandRemoveFragments_Execute);
 		TArray<FMassArchetypeEntityCollection> EntityCollections;
 		UE::Mass::Utils::CreateEntityCollections(System, TargetEntities, FMassArchetypeEntityCollection::FoldDuplicates, EntityCollections);
 		System.BatchChangeFragmentCompositionForEntities(EntityCollections, FMassFragmentBitSet(), FragmentsAffected);
@@ -182,6 +186,7 @@ struct FMassCommandChangeTags : public FMassBatchedEntityCommand
 protected:
 	virtual void Execute(UMassEntitySubsystem& System) const override
 	{
+		TRACE_CPUPROFILER_EVENT_SCOPE(MassCommandChangeTags_Execute);
 		TArray<FMassArchetypeEntityCollection> EntityCollections;
 		UE::Mass::Utils::CreateEntityCollections(System, TargetEntities, FMassArchetypeEntityCollection::FoldDuplicates, EntityCollections);
 
@@ -268,6 +273,8 @@ protected:
 
 	virtual void Execute(UMassEntitySubsystem& System) const override
 	{
+		TRACE_CPUPROFILER_EVENT_SCOPE(MassCommandAddFragmentInstances_Execute);
+
 		TArray<FStructArrayView> GenericMultiArray;
 		GenericMultiArray.Reserve(Fragments.GetNumArrays());
 		Fragments.GetAsGenericMultiArray(GenericMultiArray);
@@ -299,6 +306,8 @@ protected:
 
 	virtual void Execute(UMassEntitySubsystem& System) const override
 	{
+		TRACE_CPUPROFILER_EVENT_SCOPE(MassCommandBuildEntity_Execute);
+
 		TArray<FStructArrayView> GenericMultiArray;
 		GenericMultiArray.Reserve(Super::Fragments.GetNumArrays());
 		Super::Fragments.GetAsGenericMultiArray(GenericMultiArray);
@@ -364,6 +373,8 @@ protected:
 
 	virtual void Execute(UMassEntitySubsystem& System) const 
 	{
+		TRACE_CPUPROFILER_EVENT_SCOPE(MassCommandBuildEntityWithSharedFragments_Execute);
+
 		constexpr int FragmentTypesCount = UE::Mass::TMultiTypeList<TOthers...>::Ordinal + 1;
 		TArray<FStructArrayView> GenericMultiArray;
 		GenericMultiArray.Reserve(FragmentTypesCount);
@@ -463,6 +474,8 @@ protected:
 
 	virtual void Execute(UMassEntitySubsystem& System) const override
 	{
+		TRACE_CPUPROFILER_EVENT_SCOPE(MassDeferredCommand_Execute);
+
 		for (const FExecFunction& ExecFunction : DeferredFunctions)
 		{
 			ExecFunction(System);
