Massaging existing Mass parallelization code in preparation for its revival as benchmarking baseline.

diff --git a/Engine/Plugins/Runtime/MassEntity/Source/MassEntity/Private/MassProcessingPhase.cpp b/Engine/Plugins/Runtime/MassEntity/Source/MassEntity/Private/MassProcessingPhase.cpp
--- a/Engine/Plugins/Runtime/MassEntity/Source/MassEntity/Private/MassProcessingPhase.cpp
+++ b/Engine/Plugins/Runtime/MassEntity/Source/MassEntity/Private/MassProcessingPhase.cpp
@@ -63,32 +63,39 @@ void FMassProcessingPhase::ExecuteTick(float DeltaTime, ELevelTick TickType, ENa
 
 	check(PhaseProcessor);
 	
-	FMassProcessingContext Context(Manager->GetEntitySubsystemRef(), DeltaTime);
+	UMassEntitySubsystem& EntitySubsystem = Manager->GetEntitySubsystemRef();
+	FMassProcessingContext Context(EntitySubsystem, DeltaTime);
+
+	bIsDuringMassProcessing = true;
 
 	if (bRunInParallelMode)
 	{
-		bIsDuringMassProcessing = true;
-		const FGraphEventRef PipelineCompletionEvent = UE::Mass::Executor::TriggerParallelTasks(*PhaseProcessor, Context, [this, DeltaTime]()
-			{
-				OnParallelExecutionDone(DeltaTime);
-			});
-
-		if (PipelineCompletionEvent.IsValid())
+		bool bWorkRequested = false;
+		if (PhaseProcessor->IsEmpty() == false)
 		{
-			MyCompletionGraphEvent->DontCompleteUntil(PipelineCompletionEvent);
+			const FGraphEventRef PipelineCompletionEvent = UE::Mass::Executor::TriggerParallelTasks(*PhaseProcessor, Context, [this, DeltaTime]()
+				{
+					OnParallelExecutionDone(DeltaTime);
+				});
+
+			if (PipelineCompletionEvent.IsValid())
+			{
+				MyCompletionGraphEvent->DontCompleteUntil(PipelineCompletionEvent);
+				bWorkRequested = true;
+			}
 		}
-		else
+		if (bWorkRequested == false)
 		{
 			OnParallelExecutionDone(DeltaTime);
 		}
 	}
 	else
 	{
-		TGuardValue<bool> MassRunningGuard(bIsDuringMassProcessing, true);
 		UE::Mass::Executor::Run(*PhaseProcessor, Context);
 
 		OnPhaseEnd.Broadcast(DeltaTime);
 		Manager->OnPhaseEnd(*this);
+		bIsDuringMassProcessing = false;
 	}
 }
 
diff --git a/Engine/Plugins/Runtime/MassEntity/Source/MassEntity/Public/MassProcessingPhase.h b/Engine/Plugins/Runtime/MassEntity/Source/MassEntity/Public/MassProcessingPhase.h
--- a/Engine/Plugins/Runtime/MassEntity/Source/MassEntity/Public/MassProcessingPhase.h
+++ b/Engine/Plugins/Runtime/MassEntity/Source/MassEntity/Public/MassProcessingPhase.h
@@ -56,7 +56,7 @@ protected:
 
 private:
 	bool bRunInParallelMode = false;
-	bool bIsDuringMassProcessing = false;
+	std::atomic<bool> bIsDuringMassProcessing = false;
 };
 
 // It is unsafe to copy FTickFunctions and any subclasses of FTickFunction should specify the type trait WithCopy = false
