Implemented a helper, debug-time tool for detecting MT access to requirements declared by MassEntityQueries

diff --git a/Engine/Plugins/Experimental/StructUtils/Source/StructUtils/Public/StructTypeBitSet.h b/Engine/Plugins/Experimental/StructUtils/Source/StructUtils/Public/StructTypeBitSet.h
--- a/Engine/Plugins/Experimental/StructUtils/Source/StructUtils/Public/StructTypeBitSet.h
+++ b/Engine/Plugins/Experimental/StructUtils/Source/StructUtils/Public/StructTypeBitSet.h
@@ -97,6 +97,8 @@ struct FStructTracker
 template<typename TBaseStruct, typename TStructTrackerWrapper, typename TUStructType = UScriptStruct>
 struct TStructTypeBitSet
 {
+	using FStructTrackerWrapper = TStructTrackerWrapper;
+
 private:
 	struct FBitArrayExt : TBitArray<>
 	{
diff --git a/Engine/Plugins/Runtime/MassEntity/Source/MassEntity/Private/MassEntityQuery.cpp b/Engine/Plugins/Runtime/MassEntity/Source/MassEntity/Private/MassEntityQuery.cpp
--- a/Engine/Plugins/Runtime/MassEntity/Source/MassEntity/Private/MassEntityQuery.cpp
+++ b/Engine/Plugins/Runtime/MassEntity/Source/MassEntity/Private/MassEntityQuery.cpp
@@ -136,6 +136,8 @@ void FMassEntityQuery::ForEachEntityChunk(UMassEntitySubsystem& EntitySubsystem,
 {
 #if WITH_MASSENTITY_DEBUG
 	int32 NumEntitiesToProcess = 0;
+
+	EntitySubsystem.GetRequirementAccessDetector().RequireAccess(*this);
 #endif
 
 	ExecutionContext.SetSubsystemRequirements(RequiredConstSubsystems, RequiredMutableSubsystems);
@@ -178,6 +180,8 @@ void FMassEntityQuery::ForEachEntityChunk(UMassEntitySubsystem& EntitySubsystem,
 	// Not using VLOG to be thread safe
 	UE_CLOG(!ExecutionContext.DebugGetExecutionDesc().IsEmpty(), LogMass, VeryVerbose,
 		TEXT("%s: %d entities sent for processing"), *ExecutionContext.DebugGetExecutionDesc(), NumEntitiesToProcess);
+
+	EntitySubsystem.GetRequirementAccessDetector().ReleaseAccess(*this);
 #endif
 
 	ExecutionContext.ClearExecutionData();
diff --git a/Engine/Plugins/Runtime/MassEntity/Source/MassEntity/Private/MassEntitySubsystem.cpp b/Engine/Plugins/Runtime/MassEntity/Source/MassEntity/Private/MassEntitySubsystem.cpp
--- a/Engine/Plugins/Runtime/MassEntity/Source/MassEntity/Private/MassEntitySubsystem.cpp
+++ b/Engine/Plugins/Runtime/MassEntity/Source/MassEntity/Private/MassEntitySubsystem.cpp
@@ -108,6 +108,9 @@ void UMassEntitySubsystem::Initialize(FSubsystemCollectionBase& Collection)
 			}
 		}
 	}
+#if WITH_MASSENTITY_DEBUG
+	RequirementAccessDetector.Initialize();	
+#endif // WITH_MASSENTITY_DEBUG
 }
 
 void UMassEntitySubsystem::PostInitialize()
diff --git a/Engine/Plugins/Runtime/MassEntity/Source/MassEntity/Private/MassRequirementAccessDetector.cpp b/Engine/Plugins/Runtime/MassEntity/Source/MassEntity/Private/MassRequirementAccessDetector.cpp
new file mode 100644
--- /dev/null
+++ b/Engine/Plugins/Runtime/MassEntity/Source/MassEntity/Private/MassRequirementAccessDetector.cpp
@@ -0,0 +1,63 @@
+// Copyright Epic Games, Inc. All Rights Reserved.
+
+#include "MassRequirementAccessDetector.h"
+#include "MassEntityQuery.h"
+
+
+#if WITH_MASSENTITY_DEBUG
+namespace UE::Mass::Private
+{
+	bool bTrackRequirementsAccess = false;
+	
+	FAutoConsoleVariableRef CVarTrackRequirementsAccess(TEXT("mass.debug.TrackRequirementsAccess"), bTrackRequirementsAccess
+		, TEXT("Enables Mass processing debugging mode where we monitor thread-safety of query requirements access."));
+}
+
+void FMassRequirementAccessDetector::Initialize()
+{
+	check(IsInGameThread());
+	AddDetectors(FMassFragmentBitSet::FStructTrackerWrapper::StructTracker);
+	AddDetectors(FMassChunkFragmentBitSet::FStructTrackerWrapper::StructTracker);
+	AddDetectors(FMassSharedFragmentBitSet::FStructTrackerWrapper::StructTracker);
+	AddDetectors(FMassExternalSubystemBitSet::FStructTrackerWrapper::StructTracker);
+}
+
+void FMassRequirementAccessDetector::AddDetectors(const FStructTracker& StructTracker)
+{
+	TConstArrayView<TWeakObjectPtr<const UStruct>> Types = StructTracker.DebugGetAllStructTypes<UStruct>();
+	for (TWeakObjectPtr<const UStruct> Type : Types)
+	{
+		check(Type.Get());
+		Detectors.Add(Type.Get(), MakeShareable(new FRWAccessDetector()));
+	}
+}
+
+void FMassRequirementAccessDetector::RequireAccess(const FMassEntityQuery& Query)
+{
+	if (UE::Mass::Private::bTrackRequirementsAccess)
+	{
+		Operation(Query.RequiredConstSubsystems, &FRWAccessDetector::AcquireReadAccess);
+		Operation(Query.RequiredMutableSubsystems, &FRWAccessDetector::AcquireWriteAccess);
+		
+		Aquire(Query.Requirements);
+		Aquire(Query.ChunkRequirements);
+		Aquire(Query.ConstSharedRequirements);
+		Aquire(Query.SharedRequirements);
+	}
+}
+
+void FMassRequirementAccessDetector::ReleaseAccess(const FMassEntityQuery& Query)
+{
+	if (UE::Mass::Private::bTrackRequirementsAccess)
+	{
+		Operation(Query.RequiredConstSubsystems, &FRWAccessDetector::ReleaseReadAccess);
+		Operation(Query.RequiredMutableSubsystems, &FRWAccessDetector::ReleaseWriteAccess);
+
+		Release(Query.Requirements);
+		Release(Query.ChunkRequirements);
+		Release(Query.ConstSharedRequirements);
+		Release(Query.SharedRequirements);
+	}
+}
+
+#endif // WITH_MASSENTITY_DEBUG
\ No newline at end of file
diff --git a/Engine/Plugins/Runtime/MassEntity/Source/MassEntity/Public/MassEntityQuery.h b/Engine/Plugins/Runtime/MassEntity/Source/MassEntity/Public/MassEntityQuery.h
--- a/Engine/Plugins/Runtime/MassEntity/Source/MassEntity/Public/MassEntityQuery.h
+++ b/Engine/Plugins/Runtime/MassEntity/Source/MassEntity/Public/MassEntityQuery.h
@@ -16,6 +16,8 @@ struct FMassExecutionContext;
 struct FMassFragment;
 struct FMassArchetypeHandle;
 
+struct FMassRequirementAccessDetector;
+
 enum class EMassFragmentAccess : uint8
 {
 	// no binding required
@@ -410,6 +412,8 @@ protected:
 	void ReadCommandlineParams();
 
 protected:
+	friend struct FMassRequirementAccessDetector;
+
 	TArray<FMassFragmentRequirement> Requirements;
 	TArray<FMassFragmentRequirement> ChunkRequirements;
 	TArray<FMassFragmentRequirement> ConstSharedRequirements;
@@ -480,4 +484,4 @@ FORCEINLINE FMassEntityQuery& FMassEntityQuery::AddTagRequirements<EMassFragment
 	// force recaching the next time this query is used or the following CacheArchetypes call.
 	DirtyCachedData();
 	return *this;
-}
+}
\ No newline at end of file
diff --git a/Engine/Plugins/Runtime/MassEntity/Source/MassEntity/Public/MassEntitySubsystem.h b/Engine/Plugins/Runtime/MassEntity/Source/MassEntity/Public/MassEntitySubsystem.h
--- a/Engine/Plugins/Runtime/MassEntity/Source/MassEntity/Public/MassEntitySubsystem.h
+++ b/Engine/Plugins/Runtime/MassEntity/Source/MassEntity/Public/MassEntitySubsystem.h
@@ -12,6 +12,7 @@
 #include "MassExecutionContext.h"
 #include "Containers/MpscQueue.h"
 #include "MassExecutionContext.h"
+#include "MassRequirementAccessDetector.h"
 #include "MassEntitySubsystem.generated.h"
 
 
@@ -356,6 +357,8 @@ public:
 	void DebugForceArchetypeDataVersionBump() { ++ArchetypeDataVersion; }
 	void DebugGetArchetypeStrings(const FMassArchetypeHandle& Archetype, TArray<FName>& OutFragmentNames, TArray<FName>& OutTagNames);
 	FMassEntityHandle DebugGetEntityIndexHandle(const int32 EntityIndex) const { return Entities.IsValidIndex(EntityIndex) ? FMassEntityHandle(EntityIndex, Entities[EntityIndex].SerialNumber) : FMassEntityHandle(); }
+
+	FMassRequirementAccessDetector& GetRequirementAccessDetector() { return RequirementAccessDetector; }
 #endif // WITH_MASSENTITY_DEBUG
 
 protected:
@@ -420,4 +423,8 @@ private:
 
 	UPROPERTY(Transient)
 	FMassObserverManager ObserverManager;
+
+#if WITH_MASSENTITY_DEBUG
+	FMassRequirementAccessDetector RequirementAccessDetector;
+#endif // WITH_MASSENTITY_DEBUG
 };
diff --git a/Engine/Plugins/Runtime/MassEntity/Source/MassEntity/Public/MassRequirementAccessDetector.h b/Engine/Plugins/Runtime/MassEntity/Source/MassEntity/Public/MassRequirementAccessDetector.h
new file mode 100644
--- /dev/null
+++ b/Engine/Plugins/Runtime/MassEntity/Source/MassEntity/Public/MassRequirementAccessDetector.h
@@ -0,0 +1,75 @@
+// Copyright Epic Games, Inc. All Rights Reserved.
+
+#pragma once
+
+#include "MassProcessingTypes.h"
+#if WITH_MASSENTITY_DEBUG
+#include "Containers/ArrayView.h"
+#include "MassEntityQuery.h"
+#include "Misc/MTAccessDetector.h"
+
+struct MASSENTITY_API FMassRequirementAccessDetector
+{
+	void Initialize();
+	void RequireAccess(const FMassEntityQuery& Query);
+	void ReleaseAccess(const FMassEntityQuery& Query);
+
+private:
+	using FDetectorMethod = bool (FRWAccessDetector::*)() const;
+	template<typename TBitSet>
+	void Operation(const TBitSet& BitSet, FDetectorMethod Op)
+	{
+		TArray<const UStruct*> Types;
+		BitSet.ExportTypes(Types);
+		for (const UStruct* Type : Types)
+		{
+			if (TSharedRef<FRWAccessDetector>* Detector = Detectors.Find(Type))
+			{
+				FRWAccessDetector& DetectorRef = Detector->Get();
+				(DetectorRef.*Op)();
+			}
+		}
+	}
+
+	void Aquire(TConstArrayView<FMassFragmentRequirement> Requirements)
+	{
+		for (const FMassFragmentRequirement& Req : Requirements)
+		{
+			if (Req.Presence != EMassFragmentPresence::None)
+			{
+				if (Req.AccessMode == EMassFragmentAccess::ReadWrite)
+				{
+					Detectors.Find(Req.StructType)->Get().AcquireWriteAccess();
+				}
+				else if (Req.AccessMode == EMassFragmentAccess::ReadOnly)
+				{
+					Detectors.Find(Req.StructType)->Get().AcquireReadAccess();
+				}
+			}
+		}
+	}
+
+	void Release(TConstArrayView<FMassFragmentRequirement> Requirements)
+	{
+		for (const FMassFragmentRequirement& Req : Requirements)
+		{
+			if (Req.Presence != EMassFragmentPresence::None)
+			{
+				if (Req.AccessMode == EMassFragmentAccess::ReadWrite)
+				{
+					Detectors.Find(Req.StructType)->Get().ReleaseWriteAccess();
+				}
+				else if (Req.AccessMode == EMassFragmentAccess::ReadOnly)
+				{
+					Detectors.Find(Req.StructType)->Get().ReleaseReadAccess();
+				}
+			}
+		}
+	}
+
+	/** @Note the function is not thread-safe and meant to be only called internally on game thread (see FMassRequirementAccessDetector::Initialize) */
+	void AddDetectors(const FStructTracker& StructTracker);
+
+	TMap<const UStruct*, TSharedRef<FRWAccessDetector>> Detectors;
+};
+#endif // WITH_MASSENTITY_DEBUG
\ No newline at end of file
