Fixed missing header in MassRequirementAccessDetector.cpp

diff --git a/Engine/Plugins/Runtime/MassEntity/Source/MassEntity/Private/MassRequirementAccessDetector.cpp b/Engine/Plugins/Runtime/MassEntity/Source/MassEntity/Private/MassRequirementAccessDetector.cpp
--- a/Engine/Plugins/Runtime/MassEntity/Source/MassEntity/Private/MassRequirementAccessDetector.cpp
+++ b/Engine/Plugins/Runtime/MassEntity/Source/MassEntity/Private/MassRequirementAccessDetector.cpp
@@ -1,10 +1,10 @@
 // Copyright Epic Games, Inc. All Rights Reserved.
 
 #include "MassRequirementAccessDetector.h"
+#if WITH_MASSENTITY_DEBUG
 #include "MassEntityQuery.h"
+#include "HAL/IConsoleManager.h"
 
-
-#if WITH_MASSENTITY_DEBUG
 namespace UE::Mass::Private
 {
 	bool bTrackRequirementsAccess = false;
