Exporting mass processor requirements

diff --git a/Engine/Plugins/Runtime/MassEntity/Source/MassEntity/Private/MassEntityQuery.cpp b/Engine/Plugins/Runtime/MassEntity/Source/MassEntity/Private/MassEntityQuery.cpp
--- a/Engine/Plugins/Runtime/MassEntity/Source/MassEntity/Private/MassEntityQuery.cpp
+++ b/Engine/Plugins/Runtime/MassEntity/Source/MassEntity/Private/MassEntityQuery.cpp
@@ -9,16 +9,45 @@
 #include "Async/ParallelFor.h"
 #include "Containers/UnrealString.h"
 #include "MassProcessor.h"
+#include "MassProcessorDependencySolver.h"
 #if WITH_MASSENTITY_DEBUG
 #include "MassRequirementAccessDetector.h"
 #endif // WITH_MASSENTITY_DEBUG
 
+
+namespace UE::Mass::Private
+{
+	template<typename TContainer>
+	void ExportRequirements(TConstArrayView<FMassFragmentRequirement> Requirements, TContainer(&Out)[EMassAccessOperation::MAX])
+	{
+		for (const FMassFragmentRequirement& Requirement : Requirements)
+		{
+			if (Requirement.Presence != EMassFragmentPresence::None)
+			{
+				check(Requirement.StructType);
+				if (Requirement.AccessMode == EMassFragmentAccess::ReadOnly)
+				{
+					Out[EMassAccessOperation::Read].Add(*Requirement.StructType);
+				}
+				else if (Requirement.AccessMode == EMassFragmentAccess::ReadWrite)
+				{
+					Out[EMassAccessOperation::Write].Add(*Requirement.StructType);
+				}
+			}
+		}
+	}
+}
+
 //////////////////////////////////////////////////////////////////////
 // FMassEntityQuery
 
 
 FMassEntityQuery::FMassEntityQuery()
 {
+	bAllowParallelExecution = false;
+	bRequiresGameThreadExecution = false;
+	bRequiresMutatingWorldAccess = false;
+
 	ReadCommandlineParams();
 }
 
@@ -304,6 +333,18 @@ bool FMassEntityQuery::HasMatchingEntities(UMassEntitySubsystem& InEntitySubsyst
 	return false;
 }
 
+void FMassEntityQuery::ExportRequirements(FMassExecutionRequirements& OutRequirements) const
+{
+	OutRequirements.RequiredSubsystems[EMassAccessOperation::Read] += RequiredConstSubsystems;
+	OutRequirements.RequiredSubsystems[EMassAccessOperation::Write] += RequiredMutableSubsystems;
+
+	using UE::Mass::Private::ExportRequirements;
+	ExportRequirements(Requirements, OutRequirements.Fragments);
+	ExportRequirements(ChunkRequirements, OutRequirements.ChunkFragments);
+	ExportRequirements(ConstSharedRequirements, OutRequirements.SharedFragments);
+	ExportRequirements(SharedRequirements, OutRequirements.SharedFragments);
+}
+
 FString FMassEntityQuery::DebugGetDescription() const
 {
 #if WITH_MASSENTITY_DEBUG
diff --git a/Engine/Plugins/Runtime/MassEntity/Source/MassEntity/Private/MassProcessor.cpp b/Engine/Plugins/Runtime/MassEntity/Source/MassEntity/Private/MassProcessor.cpp
--- a/Engine/Plugins/Runtime/MassEntity/Source/MassEntity/Private/MassProcessor.cpp
+++ b/Engine/Plugins/Runtime/MassEntity/Source/MassEntity/Private/MassProcessor.cpp
@@ -142,6 +142,18 @@ void UMassProcessor::PostInitProperties()
 	if (HasAnyFlags(RF_ClassDefaultObject) == false) 
 	{
 		ConfigureQueries();
+
+		bool bNeedsGameThread = false;
+		for (FMassEntityQuery* QueryPtr : OwnedQueries)
+		{
+			CA_ASSUME(QueryPtr);
+			bNeedsGameThread = (bNeedsGameThread || QueryPtr->DoesRequireGameThreadExecution());
+		}
+		UE_CLOG(bRequiresGameThreadExecution != bNeedsGameThread, LogMass, Warning, TEXT("%s is marked bRequiresGameThreadExecution = %s, while the registered quries' requirement indicate the opposite")
+			, *GetProcessorName(), bRequiresGameThreadExecution ? TEXT("TRUE") : TEXT("FALSE"));
+
+		// better safe than sorry - if queries indicate the game thread execution is required then we marked the whole processor as such
+		bRequiresGameThreadExecution = bRequiresGameThreadExecution || bNeedsGameThread;
 	}
 #if CPUPROFILERTRACE_ENABLED
 	StatId = GetProcessorName();
@@ -157,6 +169,15 @@ void UMassProcessor::CallExecute(UMassEntitySubsystem& EntitySubsystem, FMassExe
 	Execute(EntitySubsystem, Context);
 }
 
+void UMassProcessor::ExportRequirements(FMassExecutionRequirements& OutRequirements) const
+{
+	for (FMassEntityQuery* Query : OwnedQueries)
+	{
+		CA_ASSUME(Query);
+		Query->ExportRequirements(OutRequirements);
+	}
+}
+
 void UMassProcessor::RegisterQuery(FMassEntityQuery& Query)
 {
 	const uintptr_t ThisStart = (uintptr_t)this;
diff --git a/Engine/Plugins/Runtime/MassEntity/Source/MassEntity/Public/MassEntityQuery.h b/Engine/Plugins/Runtime/MassEntity/Source/MassEntity/Public/MassEntityQuery.h
--- a/Engine/Plugins/Runtime/MassEntity/Source/MassEntity/Public/MassEntityQuery.h
+++ b/Engine/Plugins/Runtime/MassEntity/Source/MassEntity/Public/MassEntityQuery.h
@@ -7,6 +7,7 @@
 #include "Containers/UnrealString.h"
 #include "MassEntityTypes.h"
 #include "MassArchetypeTypes.h"
+#include "MassExternalSubsystemTraits.h"
 #include "MassEntityQuery.generated.h"
 
 
@@ -16,7 +17,7 @@ struct FMassArchetypeData;
 struct FMassExecutionContext;
 struct FMassFragment;
 struct FMassArchetypeHandle;
-
+struct FMassExecutionRequirements;
 struct FMassRequirementAccessDetector;
 
 enum class EMassFragmentAccess : uint8
@@ -301,10 +302,18 @@ public:
 		case EMassFragmentPresence::All:
 			RequiredAllSharedFragments.Add<T>();
 			SharedRequirements.Emplace(T::StaticStruct(), AccessMode, Presence);
+			if (AccessMode == EMassFragmentAccess::ReadWrite)
+			{
+				bRequiresGameThreadExecution |= TMassExternalSubsystemTraits<T>::GameThreadOnly;
+			}
 			break;
 		case EMassFragmentPresence::Optional:
 			RequiredOptionalSharedFragments.Add<T>();
 			SharedRequirements.Emplace(T::StaticStruct(), AccessMode, Presence);
+			if (AccessMode == EMassFragmentAccess::ReadWrite)
+			{
+				bRequiresGameThreadExecution |= TMassExternalSubsystemTraits<T>::GameThreadOnly;
+			}
 			break;
 		case EMassFragmentPresence::None:
 			RequiredNoneSharedFragments.Add<T>();
@@ -323,9 +332,11 @@ public:
 		{
 		case EMassFragmentAccess::ReadOnly:
 			RequiredConstSubsystems.Add<T>();
+			bRequiresGameThreadExecution |= TMassExternalSubsystemTraits<T>::GameThreadOnly;
 			break;
 		case EMassFragmentAccess::ReadWrite:
 			RequiredMutableSubsystems.Add<T>();
+			bRequiresGameThreadExecution |= TMassExternalSubsystemTraits<T>::GameThreadOnly;
 			break;
 		default:
 			check(false);
@@ -348,6 +359,9 @@ public:
 	
 	bool DoesArchetypeMatchRequirements(const FMassArchetypeHandle& ArchetypeHandle) const;
 
+	bool DoesRequireGameThreadExecution() const { return bRequiresGameThreadExecution || bRequiresMutatingWorldAccess; }
+	void RequireMutatingWorldAccess() { bRequiresMutatingWorldAccess = true; }
+
 	/** The function validates requirements we make for queries. See the FMassEntityQuery struct description for details.
 	 *  Note that this function is non-trivial and end users are not expected to need to use it. 
 	 *  @return whether this query's requirements follow the rules. */
@@ -411,6 +425,8 @@ public:
 
 	bool HasArchetypeFilter() const { return bool(ArchetypeCondition); }
 
+	void ExportRequirements(FMassExecutionRequirements& OutRequirements) const;
+
 protected:
 	void SortRequirements();
 	void ReadCommandlineParams();
@@ -460,7 +476,9 @@ private:
 	TArray<FMassArchetypeHandle> ValidArchetypes;
 	TArray<FMassQueryRequirementIndicesMapping> ArchetypeFragmentMapping;
 
-	bool bAllowParallelExecution = false;
+	uint8 bAllowParallelExecution : 1;
+	uint8 bRequiresGameThreadExecution : 1;
+	uint8 bRequiresMutatingWorldAccess : 1;
 
 	EMassExecutionContextType ExpectedContextType = EMassExecutionContextType::Local;
 
diff --git a/Engine/Plugins/Runtime/MassEntity/Source/MassEntity/Public/MassProcessor.h b/Engine/Plugins/Runtime/MassEntity/Source/MassEntity/Public/MassProcessor.h
--- a/Engine/Plugins/Runtime/MassEntity/Source/MassEntity/Public/MassProcessor.h
+++ b/Engine/Plugins/Runtime/MassEntity/Source/MassEntity/Public/MassProcessor.h
@@ -70,6 +70,10 @@ public:
 
 	TConstArrayView<const int32> GetPrerequisiteIndices() const { return DependencyIndices; }
 
+	/** By default fetches requirements declared entity queries registered via RegisterQuery. Processors can override 
+	 *	this function to supply additional requirements */
+	virtual void ExportRequirements(FMassExecutionRequirements& OutRequirements) const;
+
 	/** Adds Query to RegisteredQueries list. Query is required to be a member variable of this processor. Not meeting
 	 *  this requirement will cause check failure and the query won't be registered. */
 	void RegisterQuery(FMassEntityQuery& Query);
diff --git a/Engine/Plugins/Runtime/MassEntity/Source/MassEntity/Public/MassProcessorDependencySolver.h b/Engine/Plugins/Runtime/MassEntity/Source/MassEntity/Public/MassProcessorDependencySolver.h
--- a/Engine/Plugins/Runtime/MassEntity/Source/MassEntity/Public/MassProcessorDependencySolver.h
+++ b/Engine/Plugins/Runtime/MassEntity/Source/MassEntity/Public/MassProcessorDependencySolver.h
@@ -3,6 +3,7 @@
 #pragma once
 
 #include "MassProcessingTypes.h"
+#include "MassEntityTypes.h"
 
 
 class UMassProcessor;
@@ -15,6 +16,33 @@ enum class EDependencyNodeType : uint8
 	GroupEnd
 };
 
+namespace EMassAccessOperation
+{
+	constexpr int Read = 0;
+	constexpr int Write = 1;
+	constexpr int MAX = 2;
+};
+
+struct MASSENTITY_API FMassExecutionRequirements
+{
+	FMassExecutionRequirements& operator+=(const FMassExecutionRequirements& Other)
+	{
+		for (int i = 0; i < EMassAccessOperation::MAX; ++i)
+		{
+			Fragments[i] += Other.Fragments[i];
+			ChunkFragments[i] += Other.ChunkFragments[i];
+			SharedFragments[i] += Other.SharedFragments[i];
+			RequiredSubsystems[i] += Other.RequiredSubsystems[i];
+		}
+		return *this;
+	}
+
+	FMassFragmentBitSet Fragments[EMassAccessOperation::MAX];
+	FMassChunkFragmentBitSet ChunkFragments[EMassAccessOperation::MAX];
+	FMassSharedFragmentBitSet SharedFragments[EMassAccessOperation::MAX];
+	FMassExternalSubystemBitSet RequiredSubsystems[EMassAccessOperation::MAX];
+};
+
 struct MASSENTITY_API FProcessorDependencySolver
 {
 private:
