Requirement-based processor dependencies.
This version only expands on current declaration-based dependencies. More iterations to come.

diff --git a/Engine/Plugins/Runtime/MassEntity/Source/MassEntity/Private/MassEntityQuery.cpp b/Engine/Plugins/Runtime/MassEntity/Source/MassEntity/Private/MassEntityQuery.cpp
--- a/Engine/Plugins/Runtime/MassEntity/Source/MassEntity/Private/MassEntityQuery.cpp
+++ b/Engine/Plugins/Runtime/MassEntity/Source/MassEntity/Private/MassEntityQuery.cpp
@@ -18,7 +18,7 @@
 namespace UE::Mass::Private
 {
 	template<typename TContainer>
-	void ExportRequirements(TConstArrayView<FMassFragmentRequirement> Requirements, TContainer(&Out)[EMassAccessOperation::MAX])
+	void ExportRequirements(TConstArrayView<FMassFragmentRequirement> Requirements, TMassExecutionAccess<TContainer>& Out)
 	{
 		for (const FMassFragmentRequirement& Requirement : Requirements)
 		{
@@ -27,11 +27,11 @@ namespace UE::Mass::Private
 				check(Requirement.StructType);
 				if (Requirement.AccessMode == EMassFragmentAccess::ReadOnly)
 				{
-					Out[EMassAccessOperation::Read].Add(*Requirement.StructType);
+					Out.Read.Add(*Requirement.StructType);
 				}
 				else if (Requirement.AccessMode == EMassFragmentAccess::ReadWrite)
 				{
-					Out[EMassAccessOperation::Write].Add(*Requirement.StructType);
+					Out.Write.Add(*Requirement.StructType);
 				}
 			}
 		}
@@ -335,14 +335,14 @@ bool FMassEntityQuery::HasMatchingEntities(UMassEntitySubsystem& InEntitySubsyst
 
 void FMassEntityQuery::ExportRequirements(FMassExecutionRequirements& OutRequirements) const
 {
-	OutRequirements.RequiredSubsystems[EMassAccessOperation::Read] += RequiredConstSubsystems;
-	OutRequirements.RequiredSubsystems[EMassAccessOperation::Write] += RequiredMutableSubsystems;
+	OutRequirements.RequiredSubsystems.Read += RequiredConstSubsystems;
+	OutRequirements.RequiredSubsystems.Write += RequiredMutableSubsystems;
 
 	using UE::Mass::Private::ExportRequirements;
-	ExportRequirements(Requirements, OutRequirements.Fragments);
-	ExportRequirements(ChunkRequirements, OutRequirements.ChunkFragments);
-	ExportRequirements(ConstSharedRequirements, OutRequirements.SharedFragments);
-	ExportRequirements(SharedRequirements, OutRequirements.SharedFragments);
+	ExportRequirements<FMassFragmentBitSet>(Requirements, OutRequirements.Fragments);
+	ExportRequirements<FMassChunkFragmentBitSet>(ChunkRequirements, OutRequirements.ChunkFragments);
+	ExportRequirements<FMassSharedFragmentBitSet>(ConstSharedRequirements, OutRequirements.SharedFragments);
+	ExportRequirements<FMassSharedFragmentBitSet>(SharedRequirements, OutRequirements.SharedFragments);
 }
 
 FString FMassEntityQuery::DebugGetDescription() const
diff --git a/Engine/Plugins/Runtime/MassEntity/Source/MassEntity/Private/MassProcessorDependencySolver.cpp b/Engine/Plugins/Runtime/MassEntity/Source/MassEntity/Private/MassProcessorDependencySolver.cpp
--- a/Engine/Plugins/Runtime/MassEntity/Source/MassEntity/Private/MassProcessorDependencySolver.cpp
+++ b/Engine/Plugins/Runtime/MassEntity/Source/MassEntity/Private/MassProcessorDependencySolver.cpp
@@ -38,6 +38,100 @@ bool FProcessorDependencySolver::FNode::HasDependencies() const
 	return ExecuteBefore.Num() > 0 || ExecuteAfter.Num() > 0;
 }
 
+//----------------------------------------------------------------------//
+//  FProcessorDependencySolver::FResourceUsage
+//----------------------------------------------------------------------//
+FProcessorDependencySolver::FResourceUsage::FResourceUsage()
+{
+	for (int i = 0; i < EMassAccessOperation::MAX; ++i)
+	{
+		FragmentsAccess[i].Access.AddZeroed(FMassFragmentBitSet::GetMaxNum());
+		ChunkFragmentsAccess[i].Access.AddZeroed(FMassChunkFragmentBitSet::GetMaxNum());
+		SharedFragmentsAccess[i].Access.AddZeroed(FMassSharedFragmentBitSet::GetMaxNum());
+		RequiredSubsystemsAccess[i].Access.AddZeroed(FMassExternalSubystemBitSet::GetMaxNum());
+	}
+}
+
+template<typename TBitSet>
+void FProcessorDependencySolver::FResourceUsage::HandleElementType(TMassExecutionAccess<FResourceAccess>& ElementAccess
+	, const TMassExecutionAccess<TBitSet>& TestedRequirements, FProcessorDependencySolver::FNode& InOutNode, const int32 NodeIndex)
+{
+	// for every bit set in TestedRequirements we do the following
+	// 1. For every read only requirement we make InOutNode depend on the currently stored Writer of this resource
+	//    - note that this operation is not destructive, meaning we don't destructively consume the data, since all 
+	//      subsequent read access to the given resource will also depend on the Writer
+	// 2. For every read-write requirement we make InOutNode depend on all the readers and writers currently stored. 
+	//    - once that's done we clean currently stored Readers and Writers since every subsequent operation on this 
+	//      resource will be blocked by currently considered InOutNode (as the new Writer)
+	// 3. For all accessed resources we store information that InOutNode is accessing it
+
+	// 1. For every read only requirement we make InOutNode depend on the currently stored Writer of this resource
+	for (auto It = TestedRequirements.Read.GetIndexIterator(); It; ++It)
+	{
+		InOutNode.OriginalDependencies.Append(ElementAccess.Write.Access[*It].Users);
+	}
+
+	// 2. For every read-write requirement we make InOutNode depend on all the readers and writers currently stored. 
+	for (auto It = TestedRequirements.Write.GetIndexIterator(); It; ++It)
+	{
+		InOutNode.OriginalDependencies.Append(ElementAccess.Read.Access[*It].Users);
+		ElementAccess.Read.Access[*It].Users.Reset();
+
+		// with the algorithm described above we expect there to only ever be at most one Writer
+		ensure(ElementAccess.Write.Access[*It].Users.Num() <= 1);
+		InOutNode.OriginalDependencies.Append(ElementAccess.Write.Access[*It].Users);
+		ElementAccess.Write.Access[*It].Users.Reset();
+	}
+
+	// 3. For all accessed resources we store information that InOutNode is accessing it
+	for (auto It = TestedRequirements.Read.GetIndexIterator(); It; ++It)
+	{
+		// mark Element at index indicated by It as being used in mode EMassAccessOperation(i) by NodeIndex
+		ElementAccess.Read.Access[*It].Users.Add(NodeIndex);
+	}
+	for (auto It = TestedRequirements.Write.GetIndexIterator(); It; ++It)
+	{
+		// mark Element at index indicated by It as being used in mode EMassAccessOperation(i) by NodeIndex
+		ElementAccess.Write.Access[*It].Users.Add(NodeIndex);
+	}
+}
+
+template<typename TBitSet>
+bool FProcessorDependencySolver::FResourceUsage::CanAccess(const TMassExecutionAccess<TBitSet>& StoredElements, const TMassExecutionAccess<TBitSet>& TestedElements)
+{
+	// see if there's an overlap of tested write operations with existing read & write operations, as well as 
+	// tested read operations with existing write operations
+	
+	return !(
+		// if someone's already writing to what I want to write
+		TestedElements.Write.HasAny(StoredElements.Write)
+		// or if someone's already reading what I want to write
+		|| TestedElements.Write.HasAny(StoredElements.Read)
+		// or if someone's already writing what I want to read
+		|| TestedElements.Read.HasAny(StoredElements.Write)
+	);
+}
+
+bool FProcessorDependencySolver::FResourceUsage::CanAccessRequirements(const FMassExecutionRequirements& TestedRequirements) const
+{
+	bool bCanAccess = CanAccess<FMassFragmentBitSet>(Requirements.Fragments, TestedRequirements.Fragments)
+		&& CanAccess<FMassChunkFragmentBitSet>(Requirements.ChunkFragments, TestedRequirements.ChunkFragments)
+		&& CanAccess<FMassSharedFragmentBitSet>(Requirements.SharedFragments, TestedRequirements.SharedFragments)
+		&& CanAccess<FMassExternalSubystemBitSet>(Requirements.RequiredSubsystems, TestedRequirements.RequiredSubsystems);
+
+	return bCanAccess;
+}
+
+void FProcessorDependencySolver::FResourceUsage::SubmitNode(const int32 NodeIndex, FNode& InOutNode)
+{
+	HandleElementType<FMassFragmentBitSet>(FragmentsAccess, InOutNode.Requirements.Fragments, InOutNode, NodeIndex);
+	HandleElementType<FMassChunkFragmentBitSet>(ChunkFragmentsAccess, InOutNode.Requirements.ChunkFragments, InOutNode, NodeIndex);
+	HandleElementType<FMassSharedFragmentBitSet>(SharedFragmentsAccess, InOutNode.Requirements.SharedFragments, InOutNode, NodeIndex);
+	HandleElementType<FMassExternalSubystemBitSet>(RequiredSubsystemsAccess, InOutNode.Requirements.RequiredSubsystems, InOutNode, NodeIndex);
+
+	Requirements += InOutNode.Requirements;
+}
+
 //----------------------------------------------------------------------//
 //  FProcessorDependencySolver
 //----------------------------------------------------------------------//
@@ -61,48 +155,44 @@ FString FProcessorDependencySolver::NameViewToString(TConstArrayView<FName> View
 	return ReturnVal + TEXT("]");
 }
 
-int32 FProcessorDependencySolver::PerformSolverStep(FNode& RootNode, TArray<int32>& InOutIndicesRemaining, TArray<int32>& OutNodeIndices)
+bool FProcessorDependencySolver::PerformSolverStep(FResourceUsage& ResourceUsage, FNode& RootNode, TArray<int32>& InOutIndicesRemaining, TArray<int32>& OutNodeIndices)
 {
-	const int32 StartingCount = OutNodeIndices.Num();
-
-	for (int32 i = InOutIndicesRemaining.Num() - 1; i >= 0; --i)
+	int32 AcceptedNodeIndex = INDEX_NONE;
+	int32 FallbackAcceptedNodeIndex = INDEX_NONE;
+	// note that InOutIndicesRemaining contains indices valid for RootNode.SubNodes
+	// but that's ok, since in this step we only care about elements on this "level" (i.e. direct children of RootNode).
+	for (int32 i = 0; i < InOutIndicesRemaining.Num(); ++i)
 	{
 		const int32 NodeIndex = InOutIndicesRemaining[i];
 		if (RootNode.SubNodes[NodeIndex].TransientDependencies.Num() == 0)
 		{
-			InOutIndicesRemaining.RemoveAtSwap(i, 1, /*bAllowShrinking=*/false);
-			OutNodeIndices.Add(NodeIndex);
+			if (ResourceUsage.CanAccessRequirements(RootNode.SubNodes[NodeIndex].Requirements))
+			{
+				AcceptedNodeIndex = NodeIndex;
+				break;
+			}
+			else if (FallbackAcceptedNodeIndex == INDEX_NONE)
+			{
+				// if none of the nodes left can "cleanly" execute (i.e. without conflicting with already stored nodes)
+				// we'll just pick this one up and go with it. 
+				FallbackAcceptedNodeIndex = NodeIndex;
+			}
 		}
 	}
 
-	for (const int32 NodeIndex : OutNodeIndices)
+	if (AcceptedNodeIndex != INDEX_NONE || FallbackAcceptedNodeIndex != INDEX_NONE)
 	{
+		const int32 NodeIndex = AcceptedNodeIndex != INDEX_NONE ? AcceptedNodeIndex : FallbackAcceptedNodeIndex;
+		ResourceUsage.SubmitNode(NodeIndex, RootNode.SubNodes[NodeIndex]);
+		InOutIndicesRemaining.RemoveSingle(NodeIndex);
+		OutNodeIndices.Add(NodeIndex);
+
 		for (const int32 RemainingNodeIndex : InOutIndicesRemaining)
 		{
 			RootNode.SubNodes[RemainingNodeIndex].TransientDependencies.RemoveSingleSwap(NodeIndex, /*bAllowShrinking=*/false);
 		}
-	}
-
-	return OutNodeIndices.Num() - StartingCount;
-}
-
-bool FProcessorDependencySolver::PerformPrioritySolverStep(FNode& RootNode, TArray<int32>& InOutPriorityIndicesRemaining, TArray<int32>& OutNodeIndices)
-{
-	for (int32 i = 0; i < InOutPriorityIndicesRemaining.Num(); ++i)
-	{
-		const int32 NodeIndex = InOutPriorityIndicesRemaining[i];
-		if (RootNode.SubNodes[NodeIndex].TransientDependencies.Num() == 0)
-		{
-			InOutPriorityIndicesRemaining.RemoveAt(i, 1, /*bAllowShrinking=*/false);
-			OutNodeIndices.Add(NodeIndex);
-
-			for (const int32 RemainingNodeIndex : InOutPriorityIndicesRemaining)
-			{
-				RootNode.SubNodes[RemainingNodeIndex].TransientDependencies.RemoveSingleSwap(NodeIndex, /*bAllowShrinking=*/false);
-			}
-
-			return true;
-		}
+		
+		return true;
 	}
 
 	return false;
@@ -206,10 +296,13 @@ void FProcessorDependencySolver::BuildDependencies(FNode& RootNode)
 		{
 			ExecuteBefore = Node.Processor->GetExecutionOrder().ExecuteBefore;
 			ExecuteAfter = Node.Processor->GetExecutionOrder().ExecuteAfter;
+			Node.Processor->ExportRequirements(Node.Requirements);
+			RootNode.Requirements += Node.Requirements;
 		}
 		else
 		{
 			BuildDependencies(Node);
+			RootNode.Requirements += Node.Requirements;
 			ExecuteBefore = Node.ExecuteBefore;
 			ExecuteAfter = Node.ExecuteAfter;
 		}
@@ -744,17 +837,16 @@ void FProcessorDependencySolver::Solve(FNode& RootNode, TConstArrayView<const FN
 		}
 	}
 	
+	FResourceUsage ResourceUsage;
+	
 	TArray<int32> SortedNodeIndices;
 	SortedNodeIndices.Reserve(RootNode.SubNodes.Num());
 
-	bool bCyclesDetected = false;
-	int32 Steps = 0;
 	while (IndicesRemaining.Num())
 	{
-		bCyclesDetected = (PerformPrioritySolverStep(RootNode, IndicesRemaining, SortedNodeIndices) == false);
-		++Steps;
+		const bool bStepSuccessful = PerformSolverStep(ResourceUsage, RootNode, IndicesRemaining, SortedNodeIndices);
 
-		if (bCyclesDetected)
+		if (bStepSuccessful == false)
 		{
 			bAnyCyclesDetected = true;
 
@@ -792,7 +884,7 @@ void FProcessorDependencySolver::Solve(FNode& RootNode, TConstArrayView<const FN
 		TArray<FName> DependencyNames;
 		for (const int32 DependencyIndex : RootNode.SubNodes[NodeIndex].OriginalDependencies)
 		{
-			DependencyNames.Add(RootNode.SubNodes[DependencyIndex].Name);
+			DependencyNames.AddUnique(RootNode.SubNodes[DependencyIndex].Name);
 		}
 
 		if (RootNode.SubNodes[NodeIndex].Processor != nullptr)
diff --git a/Engine/Plugins/Runtime/MassEntity/Source/MassEntity/Public/MassProcessorDependencySolver.h b/Engine/Plugins/Runtime/MassEntity/Source/MassEntity/Public/MassProcessorDependencySolver.h
--- a/Engine/Plugins/Runtime/MassEntity/Source/MassEntity/Public/MassProcessorDependencySolver.h
+++ b/Engine/Plugins/Runtime/MassEntity/Source/MassEntity/Public/MassProcessorDependencySolver.h
@@ -4,6 +4,7 @@
 
 #include "MassProcessingTypes.h"
 #include "MassEntityTypes.h"
+#include "Containers/StaticArray.h"
 
 
 class UMassProcessor;
@@ -18,9 +19,28 @@ enum class EDependencyNodeType : uint8
 
 namespace EMassAccessOperation
 {
-	constexpr int Read = 0;
-	constexpr int Write = 1;
-	constexpr int MAX = 2;
+	constexpr uint32 Read = 0;
+	constexpr uint32 Write = 1;
+	constexpr uint32 MAX = 2;
+};
+
+template<typename T>
+struct MASSENTITY_API TMassExecutionAccess
+{
+	T Read;
+	T Write;
+
+	T& operator[](const uint32 OpIndex)
+	{
+		check(OpIndex <= EMassAccessOperation::MAX);
+		return OpIndex == EMassAccessOperation::Read ? Read : Write;
+	}
+
+	const T& operator[](const uint32 OpIndex) const
+	{
+		check(OpIndex <= EMassAccessOperation::MAX);
+		return OpIndex == EMassAccessOperation::Read ? Read : Write;
+	}
 };
 
 struct MASSENTITY_API FMassExecutionRequirements
@@ -37,10 +57,10 @@ struct MASSENTITY_API FMassExecutionRequirements
 		return *this;
 	}
 
-	FMassFragmentBitSet Fragments[EMassAccessOperation::MAX];
-	FMassChunkFragmentBitSet ChunkFragments[EMassAccessOperation::MAX];
-	FMassSharedFragmentBitSet SharedFragments[EMassAccessOperation::MAX];
-	FMassExternalSubystemBitSet RequiredSubsystems[EMassAccessOperation::MAX];
+	TMassExecutionAccess<FMassFragmentBitSet> Fragments;
+	TMassExecutionAccess<FMassChunkFragmentBitSet> ChunkFragments;
+	TMassExecutionAccess<FMassSharedFragmentBitSet> SharedFragments;
+	TMassExecutionAccess<FMassExternalSubystemBitSet> RequiredSubsystems;
 };
 
 struct MASSENTITY_API FProcessorDependencySolver
@@ -60,12 +80,45 @@ private:
 		TArray<int32> TransientDependencies;
 		TArray<FName> ExecuteBefore;
 		TArray<FName> ExecuteAfter;
+		FMassExecutionRequirements Requirements;
 
 		int32 FindOrAddGroupNodeIndex(const FString& GroupName);
 		int32 FindNodeIndex(FName InNodeName) const;
 		bool HasDependencies() const;
 	};
 
+	struct FResourceUsage
+	{
+		FResourceUsage();
+
+		bool CanAccessRequirements(const FMassExecutionRequirements& TestedRequirements) const;
+		void SubmitNode(const int32 NodeIndex, FNode& InOutNode);
+
+	private:
+		struct FResourceUsers
+		{
+			TArray<int32> Users;
+		};
+		
+		struct FResourceAccess
+		{
+			TArray<FResourceUsers> Access;
+		};
+		
+		FMassExecutionRequirements Requirements;
+		TMassExecutionAccess<FResourceAccess> FragmentsAccess;
+		TMassExecutionAccess<FResourceAccess> ChunkFragmentsAccess;
+		TMassExecutionAccess<FResourceAccess> SharedFragmentsAccess;
+		TMassExecutionAccess<FResourceAccess> RequiredSubsystemsAccess;
+
+		template<typename TBitSet>
+		static void HandleElementType(TMassExecutionAccess<FResourceAccess>& ElementAccess
+			, const TMassExecutionAccess<TBitSet>& TestedRequirements, FProcessorDependencySolver::FNode& InOutNode, const int32 NodeIndex);
+
+		template<typename TBitSet>
+		static bool CanAccess(const TMassExecutionAccess<TBitSet>& StoredElements, const TMassExecutionAccess<TBitSet>& TestedElements);
+	};
+
 public:
 	struct FOrderInfo
 	{
@@ -83,22 +136,13 @@ public:
 protected:
 	// note that internals are protected rather than private to support unit testing
 
-	/**
-	 * Traverses RootNode's child nodes indicated by InOutIndicesRemaining and appends to OutNodeIndices the ones that 
-	 * have no dependencies. The indices added to OutNodeIndices also get removed from remaining nodes' outstanding 
-	 * dependencies.
-	 * Note that the whole InOutIndicesRemaining gets tested in sequence, which means nodes can get their dependencies 
-	 * emptied and added to OutNodeIndices within one call (as opposed to PerformPrioritySolverStep).
-	 */
-	static int32 PerformSolverStep(FNode& RootNode, TArray<int32>& InOutIndicesRemaining, TArray<int32>& OutNodeIndices);
-	
 	/**
 	 * Traverses InOutIndicesRemaining in search of the first RootNode's node that has no dependencies left. Once found 
 	 * the node's index gets added to OutNodeIndices, removed from dependency lists from all other nodes and the function 
 	 * quits.
 	 * @return 'true' if a dependency-less node has been found and added to OutNodeIndices; 'false' otherwise.
 	 */
-	static bool PerformPrioritySolverStep(FNode& RootNode, TArray<int32>& InOutIndicesRemaining, TArray<int32>& OutNodeIndices);
+	static bool PerformSolverStep(FResourceUsage& ResourceUsage, FNode& RootNode, TArray<int32>& InOutIndicesRemaining, TArray<int32>& OutNodeIndices);
 	
 	static FString NameViewToString(TConstArrayView<FName> View);
 
