Marked a number of Mass or Mass-used systems as save to be accessed off-game-thread.
Also added a dedicated traits template for shared fragments.

As part of implementation added a dedicated MassGameplay module for storing TMassExternalSubsystemTraits declared for subsystems external to Mass code that the MassGameplay plugin's modules are using.

diff --git a/Engine/Plugins/AI/MassAI/Source/MassAIBehavior/MassAIBehavior.Build.cs b/Engine/Plugins/AI/MassAI/Source/MassAIBehavior/MassAIBehavior.Build.cs
--- a/Engine/Plugins/AI/MassAI/Source/MassAIBehavior/MassAIBehavior.Build.cs
+++ b/Engine/Plugins/AI/MassAI/Source/MassAIBehavior/MassAIBehavior.Build.cs
@@ -44,7 +44,8 @@ namespace UnrealBuildTool.Rules
 					"StateTreeModule",
 					"StructUtils",
 					"ZoneGraph",
-					"ZoneGraphAnnotations"
+					"ZoneGraphAnnotations",
+					"MassGameplayExternalTraits"
 				}
 			);
 
diff --git a/Engine/Plugins/AI/MassAI/Source/MassAIBehavior/Private/MassLookAtProcessors.cpp b/Engine/Plugins/AI/MassAI/Source/MassAIBehavior/Private/MassLookAtProcessors.cpp
--- a/Engine/Plugins/AI/MassAI/Source/MassAIBehavior/Private/MassLookAtProcessors.cpp
+++ b/Engine/Plugins/AI/MassAI/Source/MassAIBehavior/Private/MassLookAtProcessors.cpp
@@ -18,6 +18,7 @@
 #include "Algo/RandomShuffle.h"
 #include "Engine/World.h"
 #include "MassLODFragments.h"
+#include "MassGameplayExternalTraits.h"
 
 #define UNSAFE_FOR_MT 1
 
diff --git a/Engine/Plugins/AI/MassAI/Source/MassNavigation/Public/MassNavigationSubsystem.h b/Engine/Plugins/AI/MassAI/Source/MassNavigation/Public/MassNavigationSubsystem.h
--- a/Engine/Plugins/AI/MassAI/Source/MassNavigation/Public/MassNavigationSubsystem.h
+++ b/Engine/Plugins/AI/MassAI/Source/MassNavigation/Public/MassNavigationSubsystem.h
@@ -45,3 +45,12 @@ protected:
 
 	FNavigationObstacleHashGrid2D AvoidanceObstacleGrid;
 };
+
+template<>
+struct TMassExternalSubsystemTraits<UMassNavigationSubsystem> final
+{
+	enum
+	{
+		GameThreadOnly = false
+	};
+};
diff --git a/Engine/Plugins/AI/MassAI/Source/MassZoneGraphNavigation/MassZoneGraphNavigation.Build.cs b/Engine/Plugins/AI/MassAI/Source/MassZoneGraphNavigation/MassZoneGraphNavigation.Build.cs
--- a/Engine/Plugins/AI/MassAI/Source/MassZoneGraphNavigation/MassZoneGraphNavigation.Build.cs
+++ b/Engine/Plugins/AI/MassAI/Source/MassZoneGraphNavigation/MassZoneGraphNavigation.Build.cs
@@ -1,4 +1,4 @@
-﻿// Copyright Epic Games, Inc. All Rights Reserved.
+// Copyright Epic Games, Inc. All Rights Reserved.
 
 namespace UnrealBuildTool.Rules
 {
@@ -38,7 +38,8 @@ namespace UnrealBuildTool.Rules
 					"StructUtils",
 					"ZoneGraph",
 					"ZoneGraphAnnotations",
-					"DeveloperSettings"
+					"DeveloperSettings",
+					"MassGameplayExternalTraits"
 				}
 			);
 
diff --git a/Engine/Plugins/AI/MassAI/Source/MassZoneGraphNavigation/Private/MassZoneGraphNavigationProcessors.cpp b/Engine/Plugins/AI/MassAI/Source/MassZoneGraphNavigation/Private/MassZoneGraphNavigationProcessors.cpp
--- a/Engine/Plugins/AI/MassAI/Source/MassZoneGraphNavigation/Private/MassZoneGraphNavigationProcessors.cpp
+++ b/Engine/Plugins/AI/MassAI/Source/MassZoneGraphNavigation/Private/MassZoneGraphNavigationProcessors.cpp
@@ -10,6 +10,7 @@
 #include "MassSignalSubsystem.h"
 #include "ZoneGraphSubsystem.h"
 #include "ZoneGraphQuery.h"
+#include "MassGameplayExternalTraits.h"
 #include "VisualLogger/VisualLogger.h"
 #include "MassSimulationLOD.h"
 #include "Engine/World.h"
diff --git a/Engine/Plugins/AI/MassCrowd/Source/MassCrowd/Public/MassCrowdSubsystem.h b/Engine/Plugins/AI/MassCrowd/Source/MassCrowd/Public/MassCrowdSubsystem.h
--- a/Engine/Plugins/AI/MassCrowd/Source/MassCrowd/Public/MassCrowdSubsystem.h
+++ b/Engine/Plugins/AI/MassCrowd/Source/MassCrowd/Public/MassCrowdSubsystem.h
@@ -243,3 +243,12 @@ protected:
 	/** Tag mask that represents all possible lane density tags. Built from MassCrowdSettings. */
 	FZoneGraphTagMask DensityMask;
 };
+
+template<>
+struct TMassExternalSubsystemTraits<UMassCrowdSubsystem> final
+{
+	enum
+	{
+		GameThreadOnly = false
+	};
+};
diff --git a/Engine/Plugins/Runtime/MassEntity/Source/MassEntity/Public/MassEntityQuery.h b/Engine/Plugins/Runtime/MassEntity/Source/MassEntity/Public/MassEntityQuery.h
--- a/Engine/Plugins/Runtime/MassEntity/Source/MassEntity/Public/MassEntityQuery.h
+++ b/Engine/Plugins/Runtime/MassEntity/Source/MassEntity/Public/MassEntityQuery.h
@@ -304,7 +304,7 @@ public:
 			SharedRequirements.Emplace(T::StaticStruct(), AccessMode, Presence);
 			if (AccessMode == EMassFragmentAccess::ReadWrite)
 			{
-				bRequiresGameThreadExecution |= TMassExternalSubsystemTraits<T>::GameThreadOnly;
+				bRequiresGameThreadExecution |= TMassSharedFragmentTraits<T>::GameThreadOnly;
 			}
 			break;
 		case EMassFragmentPresence::Optional:
@@ -312,7 +312,7 @@ public:
 			SharedRequirements.Emplace(T::StaticStruct(), AccessMode, Presence);
 			if (AccessMode == EMassFragmentAccess::ReadWrite)
 			{
-				bRequiresGameThreadExecution |= TMassExternalSubsystemTraits<T>::GameThreadOnly;
+				bRequiresGameThreadExecution |= TMassSharedFragmentTraits<T>::GameThreadOnly;
 			}
 			break;
 		case EMassFragmentPresence::None:
diff --git a/Engine/Plugins/Runtime/MassEntity/Source/MassEntity/Public/MassEntityTypes.h b/Engine/Plugins/Runtime/MassEntity/Source/MassEntity/Public/MassEntityTypes.h
--- a/Engine/Plugins/Runtime/MassEntity/Source/MassEntity/Public/MassEntityTypes.h
+++ b/Engine/Plugins/Runtime/MassEntity/Source/MassEntity/Public/MassEntityTypes.h
@@ -6,6 +6,7 @@
 #include "MassProcessingTypes.h"
 #include "StructArrayView.h"
 #include "Subsystems/Subsystem.h"
+#include "MassExternalSubsystemTraits.h"
 #include "MassEntityTypes.generated.h"
 
 
diff --git a/Engine/Plugins/Runtime/MassEntity/Source/MassEntity/Public/MassExternalSubsystemTraits.h b/Engine/Plugins/Runtime/MassEntity/Source/MassEntity/Public/MassExternalSubsystemTraits.h
--- a/Engine/Plugins/Runtime/MassEntity/Source/MassEntity/Public/MassExternalSubsystemTraits.h
+++ b/Engine/Plugins/Runtime/MassEntity/Source/MassEntity/Public/MassExternalSubsystemTraits.h
@@ -6,9 +6,8 @@
 #include "Subsystems/WorldSubsystem.h"
 
 /**
- * Traits describing how a given piece of code can be used by Mass. By default everything is treated as
- * needing to be accessed only on game thread (implying single-thread access). To change the behavior for a given
- * class just specialize this template, like so:
+ * Traits describing how a given piece of code can be used by Mass. We require author or user of a given subsystem to 
+ * define its traits. To do it add the following in an accessible location. 
  *
  * template<>
  * struct TMassExternalSubsystemTraits<UMyCustomManager>
@@ -26,7 +25,10 @@ struct TMassExternalSubsystemTraits final
 {
 	enum
 	{
-		GameThreadOnly = true,
+		// Note that we're not supplying a default value for this property to be able to statically catch code that tries
+		// to access given subsystem without including the appropriate headers. See the comment above if you want to use
+		// a USubsystem that has not had its traits defined before. 
+		// GameThreadOnly = true,
 		ThreadSafeRead = false,
 		ThreadSafeWrite = false,
 	};
@@ -47,3 +49,16 @@ namespace FMassExternalSubsystemTraits
 		return UWorld::GetSubsystem<T>(World);
 	}
 }
+
+/** 
+ * Shared Fragments' traits.
+ * @see TMassExternalSubsystemTraits
+ */
+template <typename T>
+struct TMassSharedFragmentTraits final
+{
+	enum
+	{
+		GameThreadOnly = false,
+	};
+};
diff --git a/Engine/Plugins/Runtime/MassGameplay/MassGameplay.uplugin b/Engine/Plugins/Runtime/MassGameplay/MassGameplay.uplugin
--- a/Engine/Plugins/Runtime/MassGameplay/MassGameplay.uplugin
+++ b/Engine/Plugins/Runtime/MassGameplay/MassGameplay.uplugin
@@ -68,6 +68,10 @@
 			"Name": "MassMovementEditor",
 			"Type": "Editor",
 			"LoadingPhase": "PostEngineInit"
+		},
+		{
+			"Name": "MassGameplayExternalTraits",
+			"Type": "Runtime"
 		}
 	],
 	"Plugins": [
diff --git a/Engine/Plugins/Runtime/MassGameplay/Source/MassActors/Public/MassActorSubsystem.h b/Engine/Plugins/Runtime/MassGameplay/Source/MassActors/Public/MassActorSubsystem.h
--- a/Engine/Plugins/Runtime/MassGameplay/Source/MassActors/Public/MassActorSubsystem.h
+++ b/Engine/Plugins/Runtime/MassGameplay/Source/MassActors/Public/MassActorSubsystem.h
@@ -110,3 +110,12 @@ protected:
 	UPROPERTY()
 	UMassEntitySubsystem* EntitySystem;
 };
+
+template<>
+struct TMassExternalSubsystemTraits<UMassActorSubsystem> final
+{
+	enum
+	{
+		GameThreadOnly = false
+	};
+};
diff --git a/Engine/Plugins/Runtime/MassGameplay/Source/MassGameplayExternalTraits/MassGameplayExternalTraits.Build.cs b/Engine/Plugins/Runtime/MassGameplay/Source/MassGameplayExternalTraits/MassGameplayExternalTraits.Build.cs
new file mode 100644
--- /dev/null
+++ b/Engine/Plugins/Runtime/MassGameplay/Source/MassGameplayExternalTraits/MassGameplayExternalTraits.Build.cs
@@ -0,0 +1,25 @@
+// Copyright Epic Games, Inc. All Rights Reserved.
+
+namespace UnrealBuildTool.Rules
+{
+	public class MassGameplayExternalTraits : ModuleRules
+	{
+		public MassGameplayExternalTraits(ReadOnlyTargetRules Target) : base(Target)
+		{
+			PCHUsage = ModuleRules.PCHUsageMode.UseExplicitOrSharedPCHs;
+
+			PublicIncludePaths.AddRange(
+				new string[] { 
+					ModuleDirectory + "/Public" 
+				}
+			);
+
+			PublicDependencyModuleNames.AddRange(
+				new string[] {
+					"Core",
+					"MassEntity",
+				}
+			);
+		}
+	}
+}
diff --git a/Engine/Plugins/Runtime/MassGameplay/Source/MassGameplayExternalTraits/Private/MassGameplayExternalTraitsModule.cpp b/Engine/Plugins/Runtime/MassGameplay/Source/MassGameplayExternalTraits/Private/MassGameplayExternalTraitsModule.cpp
new file mode 100644
--- /dev/null
+++ b/Engine/Plugins/Runtime/MassGameplay/Source/MassGameplayExternalTraits/Private/MassGameplayExternalTraitsModule.cpp
@@ -0,0 +1,10 @@
+// Copyright Epic Games, Inc. All Rights Reserved.
+
+#include "IMassGameplayExternalTraitsModule.h"
+
+class FMassGameplayExternalTraitsModule : public IMassGameplayExternalTraitsModule
+{
+
+};
+
+IMPLEMENT_MODULE(FMassGameplayExternalTraitsModule, MassGameplayExternalTraits)
diff --git a/Engine/Plugins/Runtime/MassGameplay/Source/MassGameplayExternalTraits/Public/IMassGameplayExternalTraitsModule.h b/Engine/Plugins/Runtime/MassGameplay/Source/MassGameplayExternalTraits/Public/IMassGameplayExternalTraitsModule.h
new file mode 100644
--- /dev/null
+++ b/Engine/Plugins/Runtime/MassGameplay/Source/MassGameplayExternalTraits/Public/IMassGameplayExternalTraitsModule.h
@@ -0,0 +1,39 @@
+// Copyright Epic Games, Inc. All Rights Reserved.
+
+#pragma once
+
+#include "Modules/ModuleInterface.h"
+#include "Modules/ModuleManager.h"
+
+
+/**
+ * The public interface to this module.  In most cases, this interface is only public to sibling modules 
+ * within this plugin.
+ */
+class IMassGameplayExternalTraitsModule : public IModuleInterface
+{
+
+public:
+
+	/**
+	 * Singleton-like access to this module's interface.  This is just for convenience!
+	 * Beware of calling this during the shutdown phase, though.  Your module might have been unloaded already.
+	 *
+	 * @return Returns singleton instance, loading the module on demand if needed
+	 */
+	static inline IMassGameplayExternalTraitsModule& Get()
+	{
+		return FModuleManager::LoadModuleChecked<IMassGameplayExternalTraitsModule>( "MassGameplayExternalTraits" );
+	}
+
+	/**
+	 * Checks to see if this module is loaded and ready.  It is only valid to call Get() if IsAvailable() returns true.
+	 *
+	 * @return True if the module is loaded and ready to use
+	 */
+	static inline bool IsAvailable()
+	{
+		return FModuleManager::Get().IsModuleLoaded( "MassGameplayExternalTraits" );
+	}
+};
+
diff --git a/Engine/Plugins/Runtime/MassGameplay/Source/MassGameplayExternalTraits/Public/MassGameplayExternalTraits.h b/Engine/Plugins/Runtime/MassGameplay/Source/MassGameplayExternalTraits/Public/MassGameplayExternalTraits.h
new file mode 100644
--- /dev/null
+++ b/Engine/Plugins/Runtime/MassGameplay/Source/MassGameplayExternalTraits/Public/MassGameplayExternalTraits.h
@@ -0,0 +1,32 @@
+// Copyright Epic Games, Inc. All Rights Reserved.
+
+#pragma once
+
+#include "MassExternalSubsystemTraits.h"
+
+/**
+ * The point of this module and header is to have a common place for all the non-Mass subsystems MassGameplay modules
+ * are using. This is the most convenient way to ensure we define these only once. The Mass code declaring the system 
+ * below as requirements (for example MassEntityQuery's requirements) need to include this header. Otherwise 
+ * a compilation error concerning TMassExternalSubsystemTraits will be generated by the compiler.
+ */
+
+class USmartObjectSubsystem;
+template<>
+struct TMassExternalSubsystemTraits<USmartObjectSubsystem> final
+{
+	enum
+	{
+		GameThreadOnly = false
+	};
+};
+
+class UZoneGraphSubsystem;
+template<>
+struct TMassExternalSubsystemTraits<UZoneGraphSubsystem> final
+{
+	enum
+	{
+		GameThreadOnly = false
+	};
+};
diff --git a/Engine/Plugins/Runtime/MassGameplay/Source/MassReplication/Public/MassReplicationSubsystem.h b/Engine/Plugins/Runtime/MassGameplay/Source/MassReplication/Public/MassReplicationSubsystem.h
--- a/Engine/Plugins/Runtime/MassGameplay/Source/MassReplication/Public/MassReplicationSubsystem.h
+++ b/Engine/Plugins/Runtime/MassGameplay/Source/MassReplication/Public/MassReplicationSubsystem.h
@@ -337,4 +337,13 @@ protected:
 	uint64 LastSynchronizedFrame = 0;
 
 	FReplicationHashGrid2D ReplicationGrid;
-};
\ No newline at end of file
+};
+
+template<>
+struct TMassExternalSubsystemTraits<UMassReplicationSubsystem> final
+{
+	enum
+	{
+		GameThreadOnly = false
+	};
+};
diff --git a/Engine/Plugins/Runtime/MassGameplay/Source/MassRepresentation/Public/MassRepresentationFragments.h b/Engine/Plugins/Runtime/MassGameplay/Source/MassRepresentation/Public/MassRepresentationFragments.h
--- a/Engine/Plugins/Runtime/MassGameplay/Source/MassRepresentation/Public/MassRepresentationFragments.h
+++ b/Engine/Plugins/Runtime/MassGameplay/Source/MassRepresentation/Public/MassRepresentationFragments.h
@@ -65,7 +65,7 @@ struct FMassRepresentationSubsystemSharedFragment : public FMassSharedFragment
 };
 
 template<>
-struct TMassExternalSubsystemTraits<FMassRepresentationSubsystemSharedFragment> final
+struct TMassSharedFragmentTraits<FMassRepresentationSubsystemSharedFragment> final
 {
 	enum
 	{
@@ -120,7 +120,7 @@ struct FMassRepresentationParameters : public FMassSharedFragment
 };
 
 template<>
-struct TMassExternalSubsystemTraits<FMassRepresentationParameters> final
+struct TMassSharedFragmentTraits<FMassRepresentationParameters> final
 {
 	enum
 	{
diff --git a/Engine/Plugins/Runtime/MassGameplay/Source/MassSmartObjects/MassSmartObjects.Build.cs b/Engine/Plugins/Runtime/MassGameplay/Source/MassSmartObjects/MassSmartObjects.Build.cs
--- a/Engine/Plugins/Runtime/MassGameplay/Source/MassSmartObjects/MassSmartObjects.Build.cs
+++ b/Engine/Plugins/Runtime/MassGameplay/Source/MassSmartObjects/MassSmartObjects.Build.cs
@@ -36,7 +36,8 @@ namespace UnrealBuildTool.Rules
 					"SmartObjectsModule",
 					"StructUtils",
 					"ZoneGraph",
-					"ZoneGraphAnnotations"
+					"ZoneGraphAnnotations",
+					"MassGameplayExternalTraits"
 				}
 			);
 		}
diff --git a/Engine/Plugins/Runtime/MassGameplay/Source/MassSmartObjects/Private/MassSmartObjectProcessor.cpp b/Engine/Plugins/Runtime/MassGameplay/Source/MassSmartObjects/Private/MassSmartObjectProcessor.cpp
--- a/Engine/Plugins/Runtime/MassGameplay/Source/MassSmartObjects/Private/MassSmartObjectProcessor.cpp
+++ b/Engine/Plugins/Runtime/MassGameplay/Source/MassSmartObjects/Private/MassSmartObjectProcessor.cpp
@@ -16,6 +16,7 @@
 #include "VisualLogger/VisualLogger.h"
 #include "ZoneGraphAnnotationSubsystem.h"
 #include "ZoneGraphSubsystem.h"
+#include "MassGameplayExternalTraits.h"
 
 //----------------------------------------------------------------------//
 // UMassSmartObjectCandidatesFinderProcessor
