Made MassEntityQuery.ForEachEntityChunk's path for entity collection utilize stored requirements mapping if given archetype is "known" by the query (i.e. its one of the cached archetypes). This is the most common case and reusing mapping save nontrivial of work.

diff --git a/Engine/Plugins/Runtime/MassEntity/Source/MassEntity/Private/MassEntityQuery.cpp b/Engine/Plugins/Runtime/MassEntity/Source/MassEntity/Private/MassEntityQuery.cpp
--- a/Engine/Plugins/Runtime/MassEntity/Source/MassEntity/Private/MassEntityQuery.cpp
+++ b/Engine/Plugins/Runtime/MassEntity/Source/MassEntity/Private/MassEntityQuery.cpp
@@ -195,15 +195,19 @@ void FMassEntityQuery::ForEachEntityChunk(UMassEntitySubsystem& EntitySubsystem,
 	// if there's a chunk collection set by the external code - use that
 	if (ExecutionContext.GetEntityCollection().IsSet())
 	{
+		const FMassArchetypeHandle& ArchetypeHandle = ExecutionContext.GetEntityCollection().GetArchetype();
 		// verify the archetype matches requirements
-		if (DoesArchetypeMatchRequirements(ExecutionContext.GetEntityCollection().GetArchetype()) == false)
+		if (DoesArchetypeMatchRequirements(ArchetypeHandle) == false)
 		{
 			UE_VLOG_UELOG(&EntitySubsystem, LogMass, Log, TEXT("Attempted to execute FMassEntityQuery with an incompatible Archetype: %s")
-				, *DebugGetArchetypeCompatibilityDescription(ExecutionContext.GetEntityCollection().GetArchetype()));
+				, *DebugGetArchetypeCompatibilityDescription(ArchetypeHandle));
 			return;
 		}
 		ExecutionContext.SetRequirements(Requirements, ChunkRequirements, ConstSharedRequirements, SharedRequirements);
-		ExecutionContext.GetEntityCollection().GetArchetype().DataPtr->ExecuteFunction(ExecutionContext, ExecuteFunction, {}, ExecutionContext.GetEntityCollection().GetRanges());
+		
+		ArchetypeHandle.DataPtr->ExecuteFunction(ExecutionContext, ExecuteFunction
+			, GetFragmentMappingForArchetype(ArchetypeHandle)
+			, ExecutionContext.GetEntityCollection().GetRanges());
 #if WITH_MASSENTITY_DEBUG
 		NumEntitiesToProcess = ExecutionContext.GetNumEntities();
 #endif
@@ -333,6 +337,13 @@ bool FMassEntityQuery::HasMatchingEntities(UMassEntitySubsystem& InEntitySubsyst
 	return false;
 }
 
+const FMassQueryRequirementIndicesMapping& FMassEntityQuery::GetFragmentMappingForArchetype(const FMassArchetypeHandle ArchetypeHandle) const
+{
+	static const FMassQueryRequirementIndicesMapping FallbackEmptyMapping;
+	const int32 ArchetypeIndex = ValidArchetypes.Find(ArchetypeHandle);
+	return ArchetypeIndex != INDEX_NONE ? ArchetypeFragmentMapping[ArchetypeIndex] : FallbackEmptyMapping;
+}
+
 void FMassEntityQuery::ExportRequirements(FMassExecutionRequirements& OutRequirements) const
 {
 	OutRequirements.RequiredSubsystems.Read += RequiredConstSubsystems;
diff --git a/Engine/Plugins/Runtime/MassEntity/Source/MassEntity/Public/MassEntityQuery.h b/Engine/Plugins/Runtime/MassEntity/Source/MassEntity/Public/MassEntityQuery.h
--- a/Engine/Plugins/Runtime/MassEntity/Source/MassEntity/Public/MassEntityQuery.h
+++ b/Engine/Plugins/Runtime/MassEntity/Source/MassEntity/Public/MassEntityQuery.h
@@ -425,6 +425,12 @@ public:
 
 	bool HasArchetypeFilter() const { return bool(ArchetypeCondition); }
 
+	/** 
+	 * If ArchetypeHandle is among ValidArchetypes then the function retrieves requirements mapping cached for it,
+	 * otherwise an empty mapping will be returned (and the requirements binding will be done the slow way).
+	 */
+	const FMassQueryRequirementIndicesMapping& GetFragmentMappingForArchetype(const FMassArchetypeHandle ArchetypeHandle) const;
+
 	void ExportRequirements(FMassExecutionRequirements& OutRequirements) const;
 
 protected:
