Revived MASS_DO_PARALLEL implementation. Disabled by default.

diff --git a/Engine/Plugins/Runtime/MassEntity/Source/MassEntity/MassEntity.Build.cs b/Engine/Plugins/Runtime/MassEntity/Source/MassEntity/MassEntity.Build.cs
--- a/Engine/Plugins/Runtime/MassEntity/Source/MassEntity/MassEntity.Build.cs
+++ b/Engine/Plugins/Runtime/MassEntity/Source/MassEntity/MassEntity.Build.cs
@@ -35,6 +35,9 @@ namespace UnrealBuildTool.Rules
 			{
 				PrivateDependencyModuleNames.Add("MessageLog");
 			}
+
+			// @todo will be enables once we're confident in MASS_DO_PARALLEL implementation 
+			// PublicDefinitions.Add("MASS_DO_PARALLEL=" + ((Target.Type != TargetType.Server) ? 1 : 0));
 		}
 	}
 }
\ No newline at end of file
diff --git a/Engine/Plugins/Runtime/MassEntity/Source/MassEntity/Private/MassEntityModule.cpp b/Engine/Plugins/Runtime/MassEntity/Source/MassEntity/Private/MassEntityModule.cpp
--- a/Engine/Plugins/Runtime/MassEntity/Source/MassEntity/Private/MassEntityModule.cpp
+++ b/Engine/Plugins/Runtime/MassEntity/Source/MassEntity/Private/MassEntityModule.cpp
@@ -2,6 +2,7 @@
 
 #include "CoreMinimal.h"
 #include "IMassEntityModule.h"
+#include "MassEntityTypes.h"
 #if WITH_UNREAL_DEVELOPER_TOOLS
 #include "MessageLogModule.h"
 #include "Engine/World.h"
@@ -39,6 +40,12 @@ void FMassEntityModuleModule::StartupModule()
 		OnWorldCleanupHandle = FWorldDelegates::OnWorldCleanup.AddStatic(&FMassEntityModuleModule::OnWorldCleanup);
 	}
 #endif // WITH_UNREAL_DEVELOPER_TOOLS
+
+#if MASS_DO_PARALLEL
+	UE_LOG(LogMass, Log, TEXT("MassEntity running with MULTITHREADING support."));
+#else
+	UE_LOG(LogMass, Log, TEXT("MassEntity running in game thread."));
+#endif // MASS_DO_PARALLEL
 }
 
 
diff --git a/Engine/Plugins/Runtime/MassEntity/Source/MassEntity/Private/MassExecutor.cpp b/Engine/Plugins/Runtime/MassEntity/Source/MassEntity/Private/MassExecutor.cpp
--- a/Engine/Plugins/Runtime/MassEntity/Source/MassEntity/Private/MassExecutor.cpp
+++ b/Engine/Plugins/Runtime/MassEntity/Source/MassEntity/Private/MassExecutor.cpp
@@ -161,7 +161,7 @@ struct FMassExecutorDoneTask
 			ExecutionContext.Defer().MoveAppend(EntitySubsystem.Defer());
 		}
 
-		UE_LOG(LogMass, Log, TEXT("MassExecutor %s tasks DONE"), *DebugName);
+		UE_LOG(LogMass, Verbose, TEXT("MassExecutor %s tasks DONE"), *DebugName);
 		ExecutionContext.SetFlushDeferredCommands(true);
 		ExecutionContext.FlushDeferred(EntitySubsystem);
 
@@ -193,6 +193,7 @@ FGraphEventRef TriggerParallelTasks(UMassProcessor& Processor, FMassProcessingCo
 	ExecutionContext.SetDeferredCommandBuffer(CommandBuffer);
 	ExecutionContext.SetFlushDeferredCommands(false);
 	ExecutionContext.SetAuxData(ProcessingContext.AuxData);
+	ExecutionContext.SetExecutionType(EMassExecutionContextType::Processor);
 
 	FGraphEventRef CompletionEvent;
 	{
diff --git a/Engine/Plugins/Runtime/MassEntity/Source/MassEntity/Private/MassProcessor.cpp b/Engine/Plugins/Runtime/MassEntity/Source/MassEntity/Private/MassProcessor.cpp
--- a/Engine/Plugins/Runtime/MassEntity/Source/MassEntity/Private/MassProcessor.cpp
+++ b/Engine/Plugins/Runtime/MassEntity/Source/MassEntity/Private/MassProcessor.cpp
@@ -30,22 +30,11 @@ namespace UE::Mass::Debug
 
 // change to && 1 to enable more detailed processing tasks logging
 #if WITH_MASSENTITY_DEBUG && 0
-#define PROCESSOR_LOG(Fmt, ...) UE_LOG(LogMass, Log, Fmt, ##__VA_ARGS__)
+#define PROCESSOR_LOG(Fmt, ...) UE_LOG(LogMass, Verbose, Fmt, ##__VA_ARGS__)
 #else // WITH_MASSENTITY_DEBUG
 #define PROCESSOR_LOG(...) 
 #endif // WITH_MASSENTITY_DEBUG
 
-namespace FMassTweakables
-{
-	bool bParallelGroups = MASS_DO_PARALLEL;
-	float PostponedTaskWaitTimeWarningLevel = 0.002f;
-
-	FAutoConsoleVariableRef CVarsMassProcessor[] = {
-		{TEXT("mass.ParallelGroups"), bParallelGroups, TEXT("Enables mass processing groups distribution to all available threads (via the task graph)")},
-		{TEXT("mass.PostponedTaskWaitTimeWarningLevel"), PostponedTaskWaitTimeWarningLevel, TEXT("if waiting for postponed task\'s dependencies exceeds this number an error will be logged")},
-	};
-}
-
 class FMassProcessorTask
 {
 public:
@@ -316,8 +305,7 @@ FGraphEventRef UMassCompositeProcessor::DispatchProcessorTasks(UMassEntitySubsys
 void UMassCompositeProcessor::Execute(UMassEntitySubsystem& EntitySubsystem, FMassExecutionContext& Context)
 {
 #if PARALLELIZED_TRAFFIC_HACK
-	if (FMassTweakables::bParallelGroups == false
-		&& UE::MassTraffic::bParallelizeTraffic && GetProcessingPhase() == EMassProcessingPhase::PrePhysics)
+	if (UE::MassTraffic::bParallelizeTraffic && GetProcessingPhase() == EMassProcessingPhase::PrePhysics)
 	{
 		static FName TrafficGroup(TEXT("Traffic"));
 		FMassExecutionContext TrafficExecutionContext;
@@ -352,121 +340,8 @@ void UMassCompositeProcessor::Execute(UMassEntitySubsystem& EntitySubsystem, FMa
 
 		return;
 	}
-#endif // PARALLELIZED_TRAFFIC_HACK
-
-	if (FMassTweakables::bParallelGroups && bHasOffThreadSubGroups)
-	{
-		CompletionStatus.Reset();
-		CompletionStatus.AddDefaulted(ChildPipeline.Processors.Num());
-		TArray<int32> PostponedProcessors;
-
-		FMassExecutionContext SingleThreadContext = Context;
-		SingleThreadContext.SetDeferredCommandBuffer(MakeShareable(new FMassCommandBuffer()));
-
-		for (int32 NodeIndex = 0; NodeIndex < ChildPipeline.Processors.Num(); ++NodeIndex)
-		{
-			UMassProcessor* Proc = ChildPipeline.Processors[NodeIndex];
-			check(Proc);
-			UMassCompositeProcessor* CompositeProc = Cast<UMassCompositeProcessor>(Proc);
-			if (CompositeProc == nullptr || CompositeProc->bRunInSeparateThread == false)
-			{
-				// check if all dependencies have been processed already
-				Proc->TransientDependencyIndices.Reset();
-				for (const int32 DependencyIndex : Proc->DependencyIndices)
-				{
-					if (CompletionStatus[DependencyIndex].IsDone() == false)
-					{
-						Proc->TransientDependencyIndices.Add(DependencyIndex);
-					}
-				}
-
-				if (Proc->TransientDependencyIndices.Num() == 0)
-				{
-					PROCESSOR_LOG(TEXT("+--+ Instant Execution: %s.%s in %u")
-						, CompositeProc ? TEXT("") : *Proc->GetExecutionOrder().ExecuteInGroup.ToString()
-						, CompositeProc ? *CompositeProc->GetGroupName().ToString() : *Proc->GetProcessorName()
-						, FPlatformTLS::GetCurrentThreadId());
-
-					Proc->CallExecute(EntitySubsystem, SingleThreadContext);
-					CompletionStatus[NodeIndex].Status = EProcessorCompletionStatus::Done;
-				}
-				else
-				{
-					PROCESSOR_LOG(TEXT("+--+ POSTPONED Execution: %s.%s in %u")
-						, CompositeProc ? TEXT("") : *Proc->GetExecutionOrder().ExecuteInGroup.ToString()
-						, CompositeProc ? *CompositeProc->GetGroupName().ToString() : *Proc->GetProcessorName()
-						, FPlatformTLS::GetCurrentThreadId());
-
-					// postpone
-					// the graph event is needed in case we have off-thread processors depending on this one
-					CompletionStatus[NodeIndex].CompletionEvent = FGraphEvent::CreateGraphEvent();
-					CompletionStatus[NodeIndex].Status = EProcessorCompletionStatus::Postponed;
-					PostponedProcessors.Add(NodeIndex);
-				}
-			}
-			else
-			{
-				// gather prerequisites
-				FString DependenciesDesc;
-				FGraphEventArray Prerequisites;
-				for (const int32 Index : Proc->GetPrerequisiteIndices())
-				{
-					if (Index != INDEX_NONE && CompletionStatus[Index].IsDone() == false)
-					{
-						Prerequisites.Add(CompletionStatus[Index].CompletionEvent);
-						DependenciesDesc += FString::Printf(TEXT("%s, "), *ChildPipeline.Processors[Index]->GetProcessorName());
-					}
-				}
-
-				PROCESSOR_LOG(TEXT("+--+ Task %s created. %s%s"), *Proc->GetProcessorName()
-					, DependenciesDesc.Len() > 0 ? TEXT(" Dependencies: ") : TEXT(""), *DependenciesDesc);
-
-				// send off to another thread
-				CompletionStatus[NodeIndex].CompletionEvent = TGraphTask<FMassProcessorTask>::CreateTask(&Prerequisites).ConstructAndDispatchWhenReady(EntitySubsystem, Context, *Proc);
-				CompletionStatus[NodeIndex].Status = EProcessorCompletionStatus::Threaded;
-			}
-		}
-
-		{
-			TRACE_CPUPROFILER_EVENT_SCOPE_TEXT("Postponed MassProcessors");
-
-			for (int32 i = 0; i < PostponedProcessors.Num(); ++i)
-			{
-				const int32 PostponedIndex = PostponedProcessors[i];
-				UMassProcessor* Proc = ChildPipeline.Processors[PostponedIndex];
-				check(Proc);
-				for (int32 j = Proc->TransientDependencyIndices.Num() - 1; j >= 0; --j)
-				{
-					int32 DependencyIndex = Proc->TransientDependencyIndices[j];
-					ensureAlways(CompletionStatus[DependencyIndex].IsDone() || CompletionStatus[DependencyIndex].Status == EProcessorCompletionStatus::Threaded);
-						
-					CompletionStatus[DependencyIndex].Wait();
-					Proc->TransientDependencyIndices.RemoveAtSwap(j, 1, /*bAllowShrinking=*/false);
-				}
-
-				if (Proc->TransientDependencyIndices.Num() == 0)
-				{
-					Proc->CallExecute(EntitySubsystem, SingleThreadContext);
-					CompletionStatus[PostponedIndex].Status = EProcessorCompletionStatus::Done;
-					CompletionStatus[PostponedIndex].CompletionEvent->DispatchSubsequents();
-					PostponedProcessors.RemoveAt(i--, 1, /*bAllowShrinking=*/false);
-				}
-			}
-			ensureMsgf(PostponedProcessors.Num() == 0
-				, TEXT("Failed to execute all processors in one intermittent sequence - this indicates an issue with depdendency ordering"));
-		}
-
-		// wait for all events to complete
-		for (auto& Event : CompletionStatus)
-		{
-			if (Event.CompletionEvent.IsValid())
-			{
-				Event.CompletionEvent->Wait();
-			}
-		}
-		Context.Defer().MoveAppend(SingleThreadContext.Defer());
-	}
 	else
+#endif // PARALLELIZED_TRAFFIC_HACK
 	{
 		for (UMassProcessor* Proc : ChildPipeline.Processors)
 		{
@@ -508,9 +383,10 @@ void UMassCompositeProcessor::CopyAndSort(const FMassProcessingPhaseConfig& Phas
 	// with subsequent task only depending on the elements prior on the list
 	TMap<FName, int32> NameToDependencyIndex;
 	NameToDependencyIndex.Reserve(SortedProcessorsAndGroups.Num());
+	TArray<int32> SuperGroupDependency;
 	for (FProcessorDependencySolver::FOrderInfo& Element : SortedProcessorsAndGroups)
 	{
-		NameToDependencyIndex.Add(Element.Name, NameToDependencyIndex.Num());
+		NameToDependencyIndex.Add(Element.Name, ProcessingFlatGraph.Num());
 
 		FDependencyNode& Node = ProcessingFlatGraph.Add_GetRef({ Element.Name, Element.Processor });
 		Node.Dependencies.Reserve(Element.Dependencies.Num());
@@ -518,7 +394,50 @@ void UMassCompositeProcessor::CopyAndSort(const FMassProcessingPhaseConfig& Phas
 		{
 			Node.Dependencies.Add(NameToDependencyIndex.FindChecked(DependencyName));
 		}
+		switch (Element.NodeType)
+		{
+		case EDependencyNodeType::GroupStart:
+			Node.Name = FName(FString::Printf(TEXT("%s_Start"), *Node.Name.ToString()));
+			SuperGroupDependency.Add(ProcessingFlatGraph.Num() - 1);
+			break;
+		case EDependencyNodeType::GroupEnd:
+			Node.Name = FName(FString::Printf(TEXT("%s_End"), *Node.Name.ToString()));
+			SuperGroupDependency.Pop();
+			break;
+		default:
+			// this bit makes all processors withing a group depend on the group starting
+			if (SuperGroupDependency.Num())
+			{
+				Node.Dependencies.Add(SuperGroupDependency.Last());
+			}
+			break;
+		}
 	}
+
+#if WITH_MASSENTITY_DEBUG
+	UE_LOG(LogMass, Log, TEXT("%s flat processing graph:"), *GroupName.ToString());
+
+	for (int i = 0; i < ProcessingFlatGraph.Num(); ++i)
+	{
+		FDependencyNode& ProcessingNode = ProcessingFlatGraph[i];
+		FString DependenciesDesc;
+		for (const int32 DependencyIndex : ProcessingNode.Dependencies)
+		{
+			DependenciesDesc += FString::Printf(TEXT("%s, "), *ProcessingFlatGraph[DependencyIndex].Name.ToString());
+		}
+
+		if (ProcessingNode.Processor)
+		{
+			UE_LOG(LogMass, Log, TEXT("Task %s%s%s"), *ProcessingNode.Processor->GetProcessorName()
+				, DependenciesDesc.Len() > 0 ? TEXT(" depends on ") : TEXT(""), *DependenciesDesc);
+		}
+		else
+		{
+			UE_LOG(LogMass, Log, TEXT("Group %s%s%s"), *ProcessingNode.Name.ToString()
+				, DependenciesDesc.Len() > 0 ? TEXT(" depends on ") : TEXT(""), *DependenciesDesc);
+		}
+	}
+#endif // WITH_MASSENTITY_DEBUG
 }
 
 int32 UMassCompositeProcessor::Populate(TArray<FProcessorDependencySolver::FOrderInfo>& ProcessorsAndGroups, const int32 StartIndex)
@@ -559,10 +478,10 @@ int32 UMassCompositeProcessor::Populate(TArray<FProcessorDependencySolver::FOrde
 			UMassCompositeProcessor* GroupProcessor = NewObject<UMassCompositeProcessor>(GetOuter());
 			GroupProcessor->SetGroupName(Element.Name);
 			GroupProcessor->SetProcessingPhase(ProcessingPhase);
-			if (PhaseConfig.OffGameThreadGroupNames.Find(Element.Name) != INDEX_NONE)
-			{
-				GroupProcessor->bRunInSeparateThread = true;
-			}
+			// if there are no groups explicitly declared as "off game thread" we let everything go, subject to other 
+			// limitations (like some subsystems requiring being run on game thread via TMassExternalSubsystemTraits
+			GroupProcessor->bRunInSeparateThread = (PhaseConfig.OffGameThreadGroupNames.IsEmpty() == true)
+				|| (PhaseConfig.OffGameThreadGroupNames.Find(Element.Name) != INDEX_NONE);
 
 			GroupProcessor->DependencyIndices.Reset();
 			for (const FName DependencyName : Element.Dependencies)
diff --git a/Engine/Plugins/Runtime/MassEntity/Source/MassEntity/Public/MassEntitySettings.h b/Engine/Plugins/Runtime/MassEntity/Source/MassEntity/Public/MassEntitySettings.h
--- a/Engine/Plugins/Runtime/MassEntity/Source/MassEntity/Public/MassEntitySettings.h
+++ b/Engine/Plugins/Runtime/MassEntity/Source/MassEntity/Public/MassEntitySettings.h
@@ -8,7 +8,6 @@
 #include "InstancedStruct.h"
 #include "MassEntitySettings.generated.h"
 
-#define MASS_DO_PARALLEL 0
 
 #define GET_MASS_CONFIG_VALUE(a) (GetMutableDefault<UMassEntitySettings>()->a)
 
diff --git a/Engine/Plugins/Runtime/MassEntity/Source/MassEntity/Public/MassProcessingTypes.h b/Engine/Plugins/Runtime/MassEntity/Source/MassEntity/Public/MassProcessingTypes.h
--- a/Engine/Plugins/Runtime/MassEntity/Source/MassEntity/Public/MassProcessingTypes.h
+++ b/Engine/Plugins/Runtime/MassEntity/Source/MassEntity/Public/MassProcessingTypes.h
@@ -6,6 +6,10 @@
 #include "InstancedStruct.h"
 #include "MassProcessingTypes.generated.h"
 
+#ifndef MASS_DO_PARALLEL
+#define MASS_DO_PARALLEL 0
+#endif // MASS_DO_PARALLEL
+
 #define WITH_MASSENTITY_DEBUG (!(UE_BUILD_SHIPPING || UE_BUILD_SHIPPING_WITH_EDITOR || UE_BUILD_TEST) && WITH_STRUCTUTILS_DEBUG && 1)
 
 class UMassEntitySubsystem;
