Renamed UE::MassEntity namespace to UE::Mass

diff --git a/Engine/Plugins/Runtime/MassEntity/Source/MassEntity/Private/MassArchetypeData.h b/Engine/Plugins/Runtime/MassEntity/Source/MassEntity/Private/MassArchetypeData.h
--- a/Engine/Plugins/Runtime/MassEntity/Source/MassEntity/Private/MassArchetypeData.h
+++ b/Engine/Plugins/Runtime/MassEntity/Source/MassEntity/Private/MassArchetypeData.h
@@ -10,7 +10,7 @@ struct FMassExecutionContext;
 class FOutputDevice;
 struct FMassArchetypeEntityCollection;
 
-namespace UE::MassEntity
+namespace UE::Mass
 {
 	constexpr int32 ChunkSize = 128*1024;
 }
@@ -210,7 +210,7 @@ public:
 
 	int32 GetNumEntities() const { return EntityMap.Num(); }
 
-	int32 GetChunkAllocSize() const { return UE::MassEntity::ChunkSize; }
+	int32 GetChunkAllocSize() const { return UE::Mass::ChunkSize; }
 
 	int32 GetChunkCount() const { return Chunks.Num(); }
 
diff --git a/Engine/Plugins/Runtime/MassEntity/Source/MassEntity/Private/MassProcessingPhase.cpp b/Engine/Plugins/Runtime/MassEntity/Source/MassEntity/Private/MassProcessingPhase.cpp
--- a/Engine/Plugins/Runtime/MassEntity/Source/MassEntity/Private/MassProcessingPhase.cpp
+++ b/Engine/Plugins/Runtime/MassEntity/Source/MassEntity/Private/MassProcessingPhase.cpp
@@ -24,7 +24,7 @@ namespace FMassTweakables
 	};
 }
 
-namespace UE::MassEntity::Private
+namespace UE::Mass::Private
 {
 	ETickingGroup PhaseToTickingGroup[int(EMassProcessingPhase::MAX)]
 	{
@@ -35,7 +35,7 @@ namespace UE::MassEntity::Private
 		ETickingGroup::TG_PostPhysics,	// EMassProcessingPhase::PostPhysics
 		ETickingGroup::TG_LastDemotable, // EMassProcessingPhase::FrameEnd
 	};
-} // UE::MassEntity::Private
+} // UE::Mass::Private
 
 //----------------------------------------------------------------------//
 //  FMassProcessingPhase
@@ -302,7 +302,7 @@ void UMassProcessingPhaseManager::CreatePhases()
 	for (int i = 0; i < int(EMassProcessingPhase::MAX); ++i)
 	{
 		ProcessingPhases[i].Phase = EMassProcessingPhase(i);
-		ProcessingPhases[i].TickGroup = UE::MassEntity::Private::PhaseToTickingGroup[i];
+		ProcessingPhases[i].TickGroup = UE::Mass::Private::PhaseToTickingGroup[i];
 		UMassCompositeProcessor* PhaseProcessor = NewObject<UMassCompositeProcessor>(this, UMassCompositeProcessor::StaticClass()
 			, *FString::Printf(TEXT("ProcessingPhase_%s"), *EnumToString(EMassProcessingPhase(i))));
 		SetPhaseProcessor(EMassProcessingPhase(i), PhaseProcessor);
