Replaced FMassTweakables with UE::Mass::Tweakables in MassProcessingPhase.cpp

diff --git a/Engine/Plugins/Runtime/MassEntity/Source/MassEntity/Private/MassProcessingPhase.cpp b/Engine/Plugins/Runtime/MassEntity/Source/MassEntity/Private/MassProcessingPhase.cpp
--- a/Engine/Plugins/Runtime/MassEntity/Source/MassEntity/Private/MassProcessingPhase.cpp
+++ b/Engine/Plugins/Runtime/MassEntity/Source/MassEntity/Private/MassProcessingPhase.cpp
@@ -15,7 +15,7 @@
 
 DECLARE_CYCLE_STAT(TEXT("Mass Phase Done"), STAT_MassPhaseDone, STATGROUP_TaskGraphTasks);
 
-namespace FMassTweakables
+namespace UE::Mass::Tweakables
 {
 	bool bFullyParallel = MASS_DO_PARALLEL;
 
@@ -321,9 +321,9 @@ void UMassProcessingPhaseManager::OnPhaseEnd(FMassProcessingPhase& Phase)
 	CurrentPhase = EMassProcessingPhase::MAX;
 
 	// switch between parallel and single-thread versions only after a given batch of processing has been wrapped up	
-	if (Phase.IsConfiguredForParallelMode() != FMassTweakables::bFullyParallel)
+	if (Phase.IsConfiguredForParallelMode() != UE::Mass::Tweakables::bFullyParallel)
 	{
-		if (FMassTweakables::bFullyParallel)
+		if (UE::Mass::Tweakables::bFullyParallel)
 		{
 			Phase.ConfigureForParallelMode();
 		}
