Fixed Mass's test-time processors failing ensure in FMassEntityQuery::ForEachEntityChunk due to not having registered their entity queries

diff --git a/Engine/Plugins/Runtime/MassEntity/Source/MassEntityTestSuite/Private/MassTestTypes.cpp b/Engine/Plugins/Runtime/MassEntity/Source/MassEntityTestSuite/Private/MassTestTypes.cpp
--- a/Engine/Plugins/Runtime/MassEntity/Source/MassEntityTestSuite/Private/MassTestTypes.cpp
+++ b/Engine/Plugins/Runtime/MassEntity/Source/MassEntityTestSuite/Private/MassTestTypes.cpp
@@ -47,6 +47,7 @@ bool FEntityTestBase::SetUp()
 // Processors 
 //----------------------------------------------------------------------//
 UMassTestProcessorBase::UMassTestProcessorBase()
+	: EntityQuery(*this)
 {
 #if WITH_EDITORONLY_DATA
 	bCanShowUpInSettings = false;
