Fixed one execution path of MassEntityQuery::ForEachEntityChunk not releasing aquired resource access detectors

diff --git a/Engine/Plugins/Runtime/MassEntity/Source/MassEntity/Private/MassEntityQuery.cpp b/Engine/Plugins/Runtime/MassEntity/Source/MassEntity/Private/MassEntityQuery.cpp
--- a/Engine/Plugins/Runtime/MassEntity/Source/MassEntity/Private/MassEntityQuery.cpp
+++ b/Engine/Plugins/Runtime/MassEntity/Source/MassEntity/Private/MassEntityQuery.cpp
@@ -201,6 +201,10 @@ void FMassEntityQuery::ForEachEntityChunk(UMassEntitySubsystem& EntitySubsystem,
 		{
 			UE_VLOG_UELOG(&EntitySubsystem, LogMass, Log, TEXT("Attempted to execute FMassEntityQuery with an incompatible Archetype: %s")
 				, *DebugGetArchetypeCompatibilityDescription(ArchetypeHandle));
+
+#if WITH_MASSENTITY_DEBUG
+			EntitySubsystem.GetRequirementAccessDetector().ReleaseAccess(*this);
+#endif // WITH_MASSENTITY_DEBUG
 			return;
 		}
 		ExecutionContext.SetRequirements(Requirements, ChunkRequirements, ConstSharedRequirements, SharedRequirements);
