Improved resource-base dependencies.
The main change is that if a processor P in group A depends on a processor R in group B we no longer infer group A depending on group B (which was the case previously). In fact, groups are no longer represented outside of dependency solving step - once dependencies are solved we end up with a list of processors ordered in a way respecting both declared "execute after" and "execute before" dependencies but also dependencies resulting from requirements declared by Processors' queries (the latter part has practical meaning in single-threaded mode).
As part of this CL remnants of old approaches have been cut, along with the PARALLELIZED_TRAFFIC_HACK (which is not compatible with the how phase-composite-processor stored its child processors, i.e. in a flat list).

diff --git a/Engine/Plugins/Runtime/MassEntity/Source/MassEntity/Private/MassProcessor.cpp b/Engine/Plugins/Runtime/MassEntity/Source/MassEntity/Private/MassProcessor.cpp
--- a/Engine/Plugins/Runtime/MassEntity/Source/MassEntity/Private/MassProcessor.cpp
+++ b/Engine/Plugins/Runtime/MassEntity/Source/MassEntity/Private/MassProcessor.cpp
@@ -9,16 +9,6 @@
 
 DECLARE_CYCLE_STAT(TEXT("MassProcessor Group Completed"), Mass_GroupCompletedTask, STATGROUP_TaskGraphTasks);
 
-#define PARALLELIZED_TRAFFIC_HACK !MASS_DO_PARALLEL
-
-#if PARALLELIZED_TRAFFIC_HACK
-namespace UE::MassTraffic
-{
-	int32 bParallelizeTraffic = 1;
-	FAutoConsoleVariableRef CVarParallelizeTraffic(TEXT("ai.traffic.parallelize"), bParallelizeTraffic, TEXT("Whether to parallelize traffic or not"), ECVF_Cheat);
-}
-#endif // PARALLELIZED_TRAFFIC_HACK
-
 #if WITH_MASSENTITY_DEBUG
 namespace UE::Mass::Debug
 {
@@ -225,8 +215,6 @@ void UMassProcessor::PostEditChangeProperty(FPropertyChangedEvent& PropertyChang
 //----------------------------------------------------------------------//
 UMassCompositeProcessor::UMassCompositeProcessor()
 	: GroupName(TEXT("None"))
-	, bRunInSeparateThread(false)
-	, bHasOffThreadSubGroups(false)
 {
 	// not auto-registering composite processors since the idea of the global processors list is to indicate all 
 	// the processors doing the work while composite processors are just containers. Having said that subclasses 
@@ -258,15 +246,11 @@ FGraphEventRef UMassCompositeProcessor::DispatchProcessorTasks(UMassEntitySubsys
 			Prerequisites.Add(Events[DependencyIndex]);
 		}
 
-		if (ProcessingNode.Processor)
+		// we don't expect any group nodes at this point. If we get any there's a bug in dependencies solving
+		if (ensure(ProcessingNode.Processor))
 		{
 			Events.Add(ProcessingNode.Processor->DispatchProcessorTasks(EntitySubsystem, ExecutionContext, Prerequisites));
 		}
-		else
-		{
-			Events.Add(FFunctionGraphTask::CreateAndDispatchWhenReady([=](){}
-				, GET_STATID(Mass_GroupCompletedTask), &Prerequisites, ENamedThreads::AnyHiPriThreadHiPriTask));
-		}
 	}
 
 
@@ -282,16 +266,9 @@ FGraphEventRef UMassCompositeProcessor::DispatchProcessorTasks(UMassEntitySubsys
 				DependenciesDesc += FString::Printf(TEXT("%s, "), *ProcessingFlatGraph[DependencyIndex].Name.ToString());
 			}
 
-			if (ProcessingNode.Processor)
-			{
-				PROCESSOR_LOG(TEXT("Task %u %s%s%s"), Events[i]->GetTraceId(), *ProcessingNode.Processor->GetProcessorName()
-					, DependenciesDesc.Len() > 0 ? TEXT(" depends on ") : TEXT(""), *DependenciesDesc);
-			}
-			else
-			{
-				PROCESSOR_LOG(TEXT("Group %u %s%s%s"), Events[i]->GetTraceId(), *ProcessingNode.Name.ToString()
-					, DependenciesDesc.Len() > 0 ? TEXT(" depends on ") : TEXT(""), *DependenciesDesc);
-			}
+			check(ProcessingNode.Processor);
+			PROCESSOR_LOG(TEXT("Task %u %s%s%s"), Events[i]->GetTraceId(), *ProcessingNode.Processor->GetProcessorName()
+				, DependenciesDesc.Len() > 0 ? TEXT(" depends on ") : TEXT(""), *DependenciesDesc);
 		}
 	}
 #endif // WITH_MASSENTITY_DEBUG
@@ -304,50 +281,10 @@ FGraphEventRef UMassCompositeProcessor::DispatchProcessorTasks(UMassEntitySubsys
 
 void UMassCompositeProcessor::Execute(UMassEntitySubsystem& EntitySubsystem, FMassExecutionContext& Context)
 {
-#if PARALLELIZED_TRAFFIC_HACK
-	if (UE::MassTraffic::bParallelizeTraffic && GetProcessingPhase() == EMassProcessingPhase::PrePhysics)
-	{
-		static FName TrafficGroup(TEXT("Traffic"));
-		FMassExecutionContext TrafficExecutionContext;
-		FGraphEventRef TrafficCompletionEvent;
-
-		for (UMassProcessor* Proc : ChildPipeline.Processors)
-		{
-			check(Proc);
-				if (GetProcessingPhase() == EMassProcessingPhase::PrePhysics)
-			{
-				if (const UMassCompositeProcessor* CompProcessor = Cast<UMassCompositeProcessor>(Proc))
-				{
-					if (CompProcessor->GetGroupName() == TrafficGroup)
-					{
-						TrafficExecutionContext = Context;
-						// Needs its own command buffer as it will be in it own thread, will be merged after this loop
-						TrafficExecutionContext.SetDeferredCommandBuffer(MakeShareable(new FMassCommandBuffer()));
-						TrafficCompletionEvent = TGraphTask<FMassProcessorTask>::CreateTask().ConstructAndDispatchWhenReady(EntitySubsystem, TrafficExecutionContext, *Proc, /*bManageCommandBuffer=*/false);
-						continue;
-					}
-				}
-			}
-			Proc->CallExecute(EntitySubsystem, Context);
-		}
-
-		if (TrafficCompletionEvent)
-		{
-			// synchronize with the traffic "thread"
-			TrafficCompletionEvent->Wait();
-			Context.Defer().MoveAppend(TrafficExecutionContext.Defer());
-		}
-
-		return;
-	}
-	else
-#endif // PARALLELIZED_TRAFFIC_HACK
+	for (UMassProcessor* Proc : ChildPipeline.Processors)
 	{
-		for (UMassProcessor* Proc : ChildPipeline.Processors)
-		{
-			check(Proc);
-			Proc->CallExecute(EntitySubsystem, Context);
-		}
+		check(Proc);
+		Proc->CallExecute(EntitySubsystem, Context);
 	}
 }
 
@@ -375,7 +312,7 @@ void UMassCompositeProcessor::CopyAndSort(const FMassProcessingPhaseConfig& Phas
 	// figure out dependencies
 	FProcessorDependencySolver Solver(TmpPipeline.Processors, GroupName, DependencyGraphFileName);
 	TArray<FProcessorDependencySolver::FOrderInfo> SortedProcessorsAndGroups;
-	Solver.ResolveDependencies(SortedProcessorsAndGroups, PhaseConfig.OffGameThreadGroupNames);
+	Solver.ResolveDependencies(SortedProcessorsAndGroups);
 
 	Populate(SortedProcessorsAndGroups);
 
@@ -388,30 +325,15 @@ void UMassCompositeProcessor::CopyAndSort(const FMassProcessingPhaseConfig& Phas
 	{
 		NameToDependencyIndex.Add(Element.Name, ProcessingFlatGraph.Num());
 
+		// we don't expect to get any "group" nodes here. If it happens it indicates a bug in dependency solving
+		checkSlow(Element.Processor);
 		FDependencyNode& Node = ProcessingFlatGraph.Add_GetRef({ Element.Name, Element.Processor });
 		Node.Dependencies.Reserve(Element.Dependencies.Num());
 		for (FName DependencyName : Element.Dependencies)
 		{
+			checkSlow(DependencyName.IsNone() == false);
 			Node.Dependencies.Add(NameToDependencyIndex.FindChecked(DependencyName));
 		}
-		switch (Element.NodeType)
-		{
-		case EDependencyNodeType::GroupStart:
-			Node.Name = FName(FString::Printf(TEXT("%s_Start"), *Node.Name.ToString()));
-			SuperGroupDependency.Add(ProcessingFlatGraph.Num() - 1);
-			break;
-		case EDependencyNodeType::GroupEnd:
-			Node.Name = FName(FString::Printf(TEXT("%s_End"), *Node.Name.ToString()));
-			SuperGroupDependency.Pop();
-			break;
-		default:
-			// this bit makes all processors withing a group depend on the group starting
-			if (SuperGroupDependency.Num())
-			{
-				Node.Dependencies.Add(SuperGroupDependency.Last());
-			}
-			break;
-		}
 	}
 
 #if WITH_MASSENTITY_DEBUG
@@ -440,74 +362,21 @@ void UMassCompositeProcessor::CopyAndSort(const FMassProcessingPhaseConfig& Phas
 #endif // WITH_MASSENTITY_DEBUG
 }
 
-int32 UMassCompositeProcessor::Populate(TArray<FProcessorDependencySolver::FOrderInfo>& ProcessorsAndGroups, const int32 StartIndex)
+void UMassCompositeProcessor::Populate(TConstArrayView<FProcessorDependencySolver::FOrderInfo> OrderedProcessors)
 {
 	ChildPipeline.Processors.Reset();
 
-	// if processor -> Add
-	// if not processor 
-	//    if "start" -> Create group and pass the InOutProcessorsAndGroups in
-	//    else return (pop out)
-	int32 Index = StartIndex;
-	TMap<FName, int32> NameToIndexMap;
-	bool bOffThreadGroupsFound = false;
 	const FMassProcessingPhaseConfig& PhaseConfig = GET_MASS_CONFIG_VALUE(GetProcessingPhaseConfig(ProcessingPhase));
 
-	while (ProcessorsAndGroups.IsValidIndex(Index))
+	for (const FProcessorDependencySolver::FOrderInfo& ProcessorInfo : OrderedProcessors)
 	{
-		FProcessorDependencySolver::FOrderInfo& Element = ProcessorsAndGroups[Index];
-		ensure(Element.NodeType != EDependencyNodeType::Invalid);
-		
-		if (Element.NodeType == EDependencyNodeType::Processor)
+		if (ensureMsgf(ProcessorInfo.NodeType == EDependencyNodeType::Processor, TEXT("Encountered unexpected EDependencyNodeType while populating %s"), *GetGroupName().ToString()))
 		{
-			check(Element.Processor);
-			check(Element.Processor->GetExecutionOrder().ExecuteInGroup == GroupName || Element.Processor->GetExecutionOrder().ExecuteInGroup.IsNone());
-
-			Element.Processor->DependencyIndices.Reset();
-			for (const FName DependencyName : Element.Dependencies)
-			{
-				Element.Processor->DependencyIndices.Add(NameToIndexMap.FindChecked(DependencyName));
-			}
-
-			NameToIndexMap.Add(Element.Name, ChildPipeline.Processors.Num());
-			ChildPipeline.AppendProcessor(*Element.Processor);
-			++Index;
-		}
-		else if (Element.NodeType == EDependencyNodeType::GroupStart)
-		{
-			UMassCompositeProcessor* GroupProcessor = NewObject<UMassCompositeProcessor>(GetOuter());
-			GroupProcessor->SetGroupName(Element.Name);
-			GroupProcessor->SetProcessingPhase(ProcessingPhase);
-			// if there are no groups explicitly declared as "off game thread" we let everything go, subject to other 
-			// limitations (like some subsystems requiring being run on game thread via TMassExternalSubsystemTraits
-			GroupProcessor->bRunInSeparateThread = (PhaseConfig.OffGameThreadGroupNames.IsEmpty() == true)
-				|| (PhaseConfig.OffGameThreadGroupNames.Find(Element.Name) != INDEX_NONE);
-
-			GroupProcessor->DependencyIndices.Reset();
-			for (const FName DependencyName : Element.Dependencies)
-			{
-				const int32 DependencyIndex = NameToIndexMap.FindChecked(DependencyName);
-				GroupProcessor->DependencyIndices.Add(DependencyIndex);
-			}
-
-			NameToIndexMap.Add(Element.Name, ChildPipeline.Processors.Num());
-			ChildPipeline.AppendProcessor(*GroupProcessor);
-
-			Index = GroupProcessor->Populate(ProcessorsAndGroups, Index + 1) + 1;
-			bOffThreadGroupsFound = (bOffThreadGroupsFound || GroupProcessor->bRunInSeparateThread);
-		}
-		else 
-		{
-			check(Element.NodeType == EDependencyNodeType::GroupEnd);
-			check(Element.Name == GroupName);
-			// done
-			break;
+			checkSlow(ProcessorInfo.Processor);
+			
+			ChildPipeline.AppendProcessor(*ProcessorInfo.Processor);
 		}
 	}
-
-	bHasOffThreadSubGroups = bOffThreadGroupsFound;
-
-	return Index;
 }
 
 void UMassCompositeProcessor::DebugOutputDescription(FOutputDevice& Ar, int32 Indent) const
diff --git a/Engine/Plugins/Runtime/MassEntity/Source/MassEntity/Private/MassProcessorDependencySolver.cpp b/Engine/Plugins/Runtime/MassEntity/Source/MassEntity/Private/MassProcessorDependencySolver.cpp
--- a/Engine/Plugins/Runtime/MassEntity/Source/MassEntity/Private/MassProcessorDependencySolver.cpp
+++ b/Engine/Plugins/Runtime/MassEntity/Source/MassEntity/Private/MassProcessorDependencySolver.cpp
@@ -10,32 +10,21 @@
 
 #define LOCTEXT_NAMESPACE "Mass"
 
-//----------------------------------------------------------------------//
-//  FProcessorDependencySolver::FNode
-//----------------------------------------------------------------------//
-int32 FProcessorDependencySolver::FNode::FindOrAddGroupNodeIndex(const FString& GroupName)
+namespace UE::Mass::Private
 {
-	const FName GroupNodeName = FName(*GroupName);
-	const int32* GroupNodeIndex = Indices.Find(GroupNodeName);
-	if (GroupNodeIndex)
+	FString NameViewToString(TConstArrayView<FName> View)
 	{
-		return *GroupNodeIndex;
+		if (View.Num() == 0)
+		{
+			return TEXT("[]");
+		}
+		FString ReturnVal = FString::Printf(TEXT("[%s"), *View[0].ToString());
+		for (int i = 1; i < View.Num(); ++i)
+		{
+			ReturnVal += FString::Printf(TEXT(", %s"), *View[i].ToString());
+		}
+		return ReturnVal + TEXT("]");
 	}
-
-	const int32 Index = SubNodes.Add({ GroupNodeName, nullptr });
-	Indices.Add(GroupNodeName, Index);
-	return Index;
-}
-
-int32 FProcessorDependencySolver::FNode::FindNodeIndex(FName InNodeName) const
-{
-	const int32* GroupNodeIndex = Indices.Find(InNodeName);
-	return GroupNodeIndex ? *GroupNodeIndex : INDEX_NONE;
-}
-
-bool FProcessorDependencySolver::FNode::HasDependencies() const
-{
-	return ExecuteBefore.Num() > 0 || ExecuteAfter.Num() > 0;
 }
 
 //----------------------------------------------------------------------//
@@ -137,36 +126,21 @@ void FProcessorDependencySolver::FResourceUsage::SubmitNode(const int32 NodeInde
 //----------------------------------------------------------------------//
 FProcessorDependencySolver::FProcessorDependencySolver(TArrayView<UMassProcessor*> InProcessors, const FName Name, const FString& InDependencyGraphFileName /*= FString()*/)
 	: Processors(InProcessors)
-	, GroupRootNode({Name, nullptr})
 	, DependencyGraphFileName(InDependencyGraphFileName)
+	, CollectionName(Name)
 {}
 
-FString FProcessorDependencySolver::NameViewToString(TConstArrayView<FName> View)
-{
-	if (View.Num() == 0)
-	{
-		return TEXT("[]");
-	}
-	FString ReturnVal = FString::Printf(TEXT("[%s"), *View[0].ToString());
-	for (int i = 1; i < View.Num(); ++i)
-	{
-		ReturnVal += FString::Printf(TEXT(", %s"), *View[i].ToString());
-	}
-	return ReturnVal + TEXT("]");
-}
-
-bool FProcessorDependencySolver::PerformSolverStep(FResourceUsage& ResourceUsage, FNode& RootNode, TArray<int32>& InOutIndicesRemaining, TArray<int32>& OutNodeIndices)
+bool FProcessorDependencySolver::PerformSolverStep(FResourceUsage& ResourceUsage, TArray<int32>& InOutIndicesRemaining, TArray<int32>& OutNodeIndices)
 {
 	int32 AcceptedNodeIndex = INDEX_NONE;
 	int32 FallbackAcceptedNodeIndex = INDEX_NONE;
-	// note that InOutIndicesRemaining contains indices valid for RootNode.SubNodes
-	// but that's ok, since in this step we only care about elements on this "level" (i.e. direct children of RootNode).
+
 	for (int32 i = 0; i < InOutIndicesRemaining.Num(); ++i)
 	{
 		const int32 NodeIndex = InOutIndicesRemaining[i];
-		if (RootNode.SubNodes[NodeIndex].TransientDependencies.Num() == 0)
+		if (AllNodes[NodeIndex].TransientDependencies.Num() == 0)
 		{
-			if (ResourceUsage.CanAccessRequirements(RootNode.SubNodes[NodeIndex].Requirements))
+			if (ResourceUsage.CanAccessRequirements(AllNodes[NodeIndex].Requirements))
 			{
 				AcceptedNodeIndex = NodeIndex;
 				break;
@@ -183,13 +157,22 @@ bool FProcessorDependencySolver::PerformSolverStep(FResourceUsage& ResourceUsage
 	if (AcceptedNodeIndex != INDEX_NONE || FallbackAcceptedNodeIndex != INDEX_NONE)
 	{
 		const int32 NodeIndex = AcceptedNodeIndex != INDEX_NONE ? AcceptedNodeIndex : FallbackAcceptedNodeIndex;
-		ResourceUsage.SubmitNode(NodeIndex, RootNode.SubNodes[NodeIndex]);
+
+		// Note that this is not an unexpected event and will happen during every dependency solving. It's a part 
+		// of the algorithm. We initially look for all the things we can run without conflicting with anything else. 
+		// But that can't last forever, at some point we'll end up in a situation where every node left waits for 
+		// something that has been submitted already. Then we just pick one of the waiting ones (the one indicated by 
+		// FallbackAcceptedNodeIndex), "run it" and proceed.
+		UE_CLOG(AcceptedNodeIndex == INDEX_NONE, LogMass, Verbose, TEXT("No dependency-free node can be picked, due to resource requirements. Picking %s as the next node.")
+			, *AllNodes[NodeIndex].Name.ToString());
+
+		ResourceUsage.SubmitNode(NodeIndex, AllNodes[NodeIndex]);
 		InOutIndicesRemaining.RemoveSingle(NodeIndex);
 		OutNodeIndices.Add(NodeIndex);
 
 		for (const int32 RemainingNodeIndex : InOutIndicesRemaining)
 		{
-			RootNode.SubNodes[RemainingNodeIndex].TransientDependencies.RemoveSingleSwap(NodeIndex, /*bAllowShrinking=*/false);
+			AllNodes[RemainingNodeIndex].TransientDependencies.RemoveSingleSwap(NodeIndex, /*bAllowShrinking=*/false);
 		}
 		
 		return true;
@@ -219,174 +202,150 @@ void FProcessorDependencySolver::CreateSubGroupNames(FName InGroupName, TArray<F
 	}
 }
 
-void FProcessorDependencySolver::AddNode(FName InGroupName, UMassProcessor& Processor)
+void FProcessorDependencySolver::CreateNodes(UMassProcessor& Processor)
 {
 	check(Processor.GetClass());
 	const FName ProcName = Processor.GetClass()->GetFName();
 
-	if (InGroupName.IsNone())
-	{
-		if (GroupRootNode.Indices.Find(ProcName) == nullptr)
-		{
-			GroupRootNode.Indices.Add(ProcName, GroupRootNode.SubNodes.Num());
-			GroupRootNode.SubNodes.Add({ProcName, &Processor});
-		}
-		else
-		{
-			UE_LOG(LogMass, Warning, TEXT("%s Processor %s already registered. Duplicates are not supported.")
-				, ANSI_TO_TCHAR(__FUNCTION__), *ProcName.ToString());
-		}
-	}
-	else
-	{
-		// The idea here is that we create a tree-like structure to place the Processor as a leaf.
-		// The levels of InGroupName are used to create tree successive branch nodes, so a group name A.B.C
-		// will result in creating node A, then that node will have node A.B created for it and that one in 
-		// turn will get A.B.C; Every node created gets information about the subnodes it hosts for faster 
-		// lookup during dependency resolving
+	check(NodeIndexMap.Find(ProcName) == nullptr);
 
-		// first we need to figure out all the subgroups indicated by InGroupName
-		TArray<FString> AllGroupNames;
-		CreateSubGroupNames(InGroupName, AllGroupNames);
+	// first check if a node for this Processor has been added, due to duplicates or some other processor given this 
+	// processor as a dependency by name
+	const int32 NodeIndex = AllNodes.Num();
+	NodeIndexMap.Add(ProcName, NodeIndex);
+	AllNodes.Add({ ProcName, &Processor, NodeIndex });
 
-		// now drill down into the successive levels of the free and find or create the necessary nodes. 
-		FNode* CurrentNode = &GroupRootNode;
-		for (int GroupNameIndex = 0; GroupNameIndex < AllGroupNames.Num(); ++GroupNameIndex)
-		{
-			const int32 NodeIndex = CurrentNode->FindOrAddGroupNodeIndex(AllGroupNames[GroupNameIndex]);
-			CurrentNode->Indices.Add(ProcName, NodeIndex);
-			// we make every node aware of all the subgroups it contains
-			for (int SubgroupIndex = GroupNameIndex + 1; SubgroupIndex < AllGroupNames.Num(); ++SubgroupIndex)
+	const FMassProcessorExecutionOrder& ExecutionOrder = Processor.GetExecutionOrder();
+
+	AllNodes[NodeIndex].ExecuteAfter.Append(ExecutionOrder.ExecuteAfter);
+	AllNodes[NodeIndex].ExecuteBefore.Append(ExecutionOrder.ExecuteBefore);
+	Processor.ExportRequirements(AllNodes[NodeIndex].Requirements);
+
+	if (ExecutionOrder.ExecuteInGroup.IsNone() == false)
+	{
+		TArray<FString> AllGroupNames;
+		CreateSubGroupNames(ExecutionOrder.ExecuteInGroup, AllGroupNames);
+	
+		check(AllGroupNames.Num() > 0);
+		int32 ParentGroupNodeIndex = INDEX_NONE;
+		for (const FString& GroupName : AllGroupNames)
+		{
+			const FName GroupFName(GroupName);
+			int32* LocalGroupIndex = NodeIndexMap.Find(GroupFName);
+			// group name hasn't been encountered yet - create it
+			if (LocalGroupIndex == nullptr)
 			{
-				CurrentNode->Indices.Add(*AllGroupNames[SubgroupIndex], NodeIndex);
+				int32 NewGroupNodeIndex = AllNodes.Num();
+				NodeIndexMap.Add(GroupFName, NewGroupNodeIndex);
+				FNode& GroupNode = AllNodes.Add_GetRef({ GroupFName, nullptr, NewGroupNodeIndex });
+				// just ignore depending on the dummy "root" node
+				if (ParentGroupNodeIndex != INDEX_NONE)
+				{
+					GroupNode.OriginalDependencies.Add(ParentGroupNodeIndex);
+					AllNodes[ParentGroupNodeIndex].SubNodeIndices.Add(NewGroupNodeIndex);
+				}
+
+				ParentGroupNodeIndex = NewGroupNodeIndex;
+			}
+			else
+			{	
+				ParentGroupNodeIndex = *LocalGroupIndex;
 			}
 
-			CurrentNode = &CurrentNode->SubNodes[NodeIndex];
 		}
 
-		// at this point CurrentNode points at the leaf group where we need to add the processor itself
-		if (CurrentNode->Indices.Find(ProcName) == nullptr)
-		{
-			const int32 ProcNodeIndex = CurrentNode->SubNodes.Add({ProcName, &Processor});
-			CurrentNode->Indices.Add(ProcName, ProcNodeIndex);
-		}
-		else
-		{
-			UE_LOG(LogMass, Warning, TEXT("%s Processor %s already registered. Duplicates are not supported.")
-				, ANSI_TO_TCHAR(__FUNCTION__), *ProcName.ToString());
-		}
+		check(ParentGroupNodeIndex > 0);
+		AllNodes[ParentGroupNodeIndex].SubNodeIndices.Add(NodeIndex);
 	}
 }
 
-void FProcessorDependencySolver::BuildDependencies(FNode& RootNode)
+void FProcessorDependencySolver::BuildDependencies()
 {
-	// The main idea here is that for every subnode we convert its dependencies (expressed via ExecuteBefore and ExecuteAfter)
-	// into a single dependency list by reversing the order of ExecuteBefore relationship - so if A.ExecuteBefore contains B then
-	// we turn it into B.ExecuteAfter += A. The dependency names get converted to node indices. If given dependency name 
-	// is unknown at the RootNode's level (i.e. it's not present in its subtree) we make that name a dependency of RootNode. 
-	// Handling of this dependency will then be taken over by RootNode's parent.
-
-	for (int32 NodeIndex = 0; NodeIndex < RootNode.SubNodes.Num(); ++NodeIndex)
+	// at this point we have collected all the known processors and groups in AllNodes so we can transpose 
+	// A.ExecuteBefore(B) type of dependencies into B.ExecuteAfter(A)
+	for (FNode& Node : AllNodes)
 	{
-		FNode& Node = RootNode.SubNodes[NodeIndex];
-		TConstArrayView<FName> ExecuteBefore;
-		TConstArrayView<FName> ExecuteAfter;
-
-		if (Node.Processor)
-		{
-			ExecuteBefore = Node.Processor->GetExecutionOrder().ExecuteBefore;
-			ExecuteAfter = Node.Processor->GetExecutionOrder().ExecuteAfter;
-			Node.Processor->ExportRequirements(Node.Requirements);
-			RootNode.Requirements += Node.Requirements;
-		}
-		else
+		for (const FName& BeforeDependencyName : Node.ExecuteBefore)
 		{
-			BuildDependencies(Node);
-			RootNode.Requirements += Node.Requirements;
-			ExecuteBefore = Node.ExecuteBefore;
-			ExecuteAfter = Node.ExecuteAfter;
-		}
-		
-		for (const FName& DependencyName : ExecuteBefore)
-		{
-			const int32 DependencyIndex = RootNode.FindNodeIndex(DependencyName);
-			// RootNode doesn't know about DependencyName then add it to RootNode's dependencies
-			if (DependencyIndex == INDEX_NONE)
+			int32* DependentNodeIndex = NodeIndexMap.Find(BeforeDependencyName);
+			if (ensure(DependentNodeIndex))
 			{
-				RootNode.ExecuteBefore.AddUnique(DependencyName);
+				check(AllNodes.IsValidIndex(*DependentNodeIndex));
+				AllNodes[*DependentNodeIndex].ExecuteAfter.Add(Node.Name);
 			}
 			else
 			{
-				if (NodeIndex != DependencyIndex)
-				{
-					// just make DependencyIndex depend on Node
-					RootNode.SubNodes[DependencyIndex].OriginalDependencies.AddUnique(NodeIndex);
-				}
-				else
-				{
-					UE_LOG(LogMass, Warning, TEXT("%s a node (%s) tried to add a circular dependency on itself")
-						, ANSI_TO_TCHAR(__FUNCTION__), *Node.Name.ToString());
-				}
+				UE_LOG(LogMass, Log, TEXT("Unable to find dependency \"%s\" declared by %s"), *BeforeDependencyName.ToString()
+					, *Node.Name.ToString());
 			}
 		}
-		for (const FName& DependencyName : ExecuteAfter)
+		Node.ExecuteBefore.Reset();
+	}
+
+	// at this point all nodes contain:
+	// - single "original dependency" pointing at its parent group
+	// - ExecuteAfter populated with node names
+
+	// Now, for every Name in ExecuteAfter we do the following:
+	//	if Name represents a processor, add it as "original dependency"
+	//	else, if Name represents a group:
+	//		- append all group's child node names to ExecuteAfter
+	// 
+	for (FNode& Node : AllNodes)
+	{
+		for (int i = 0; i < Node.ExecuteAfter.Num(); ++i)
 		{
-			const int32 DependencyIndex = RootNode.FindNodeIndex(DependencyName);
-			// RootNode doesn't know about DependencyName then add it to RootNode's dependencies
-			if (DependencyIndex == INDEX_NONE)
-			{
-				RootNode.ExecuteAfter.AddUnique(DependencyName);
-			}
-			else 
+			const FName& AfterDependencyName = Node.ExecuteAfter[i];
+			int32* PrerequisiteNodeIndex = NodeIndexMap.Find(AfterDependencyName);
+			if (ensure(PrerequisiteNodeIndex))
 			{
-				if (NodeIndex != DependencyIndex)
+				const FNode& PrerequisiteNode = AllNodes[*PrerequisiteNodeIndex];
+
+				if (PrerequisiteNode.IsGroup())
 				{
-					// just make Node depend on DependencyIndex
-					Node.OriginalDependencies.AddUnique(DependencyIndex);
+					for (int32 SubNodeIndex : PrerequisiteNode.SubNodeIndices)
+					{
+						Node.ExecuteAfter.AddUnique(AllNodes[SubNodeIndex].Name);
+					}
 				}
 				else
 				{
-					UE_LOG(LogMass, Warning, TEXT("%s a node (%s) tried to add a circular dependency on itself")
-						, ANSI_TO_TCHAR(__FUNCTION__), *Node.Name.ToString());
+					Node.OriginalDependencies.AddUnique(*PrerequisiteNodeIndex);
 				}
 			}
+			else
+			{
+				UE_LOG(LogMass, Log, TEXT("Unable to find dependency \"%s\" declared by %s"), *AfterDependencyName.ToString()
+					, *Node.Name.ToString());
+			}
 		}
 	}
 }
 
-void FProcessorDependencySolver::LogNode(const FNode& RootNode, const FNode* ParentNode, int Indent)
+void FProcessorDependencySolver::LogNode(const FNode& Node, int Indent)
 {
-	FString Dependencies;
-	for (const int32 DependencyIndex : RootNode.OriginalDependencies)
-	{
-		if (ensure(ParentNode && ParentNode->SubNodes.IsValidIndex(DependencyIndex)))
-		{
-			Dependencies += ParentNode->SubNodes[DependencyIndex].Name.ToString() + TEXT(", ");
-		}
-	}
-
-	UE_LOG(LogMass, Log, TEXT("%*s%s %s%s"), Indent, TEXT(""), *RootNode.Name.ToString()
-		, Dependencies.Len() > 0 ? TEXT("after: ") : TEXT(""), *Dependencies);
+	using UE::Mass::Private::NameViewToString;
 
-	if (RootNode.Processor == nullptr)
+	if (Node.IsGroup())
 	{
-		UE_LOG(LogMass, Log, TEXT("%*s%s before:%s after:%s"), Indent, TEXT(""), *RootNode.Name.ToString()
-			, *NameViewToString(RootNode.ExecuteBefore)
-			, *NameViewToString(RootNode.ExecuteAfter));
+		UE_LOG(LogMass, Log, TEXT("%*s%s before:%s after:%s"), Indent, TEXT(""), *Node.Name.ToString()
+			, *NameViewToString(Node.ExecuteBefore)
+			, *NameViewToString(Node.ExecuteAfter));
 
-		for (const FNode& Node : RootNode.SubNodes)
+		for (const int32 NodeIndex : Node.SubNodeIndices)
 		{
-			LogNode(Node, &RootNode, Indent + 4);
+			LogNode(AllNodes[NodeIndex], Indent + 4);
 		}
 	}
 	else
 	{
-		UE_LOG(LogMass, Log, TEXT("%*s%s before:%s after:%s"), Indent, TEXT(""), *RootNode.Name.ToString()
-			, *NameViewToString(RootNode.Processor->GetExecutionOrder().ExecuteBefore)
-			, *NameViewToString(RootNode.Processor->GetExecutionOrder().ExecuteAfter));
+		UE_LOG(LogMass, Log, TEXT("%*s%s before:%s after:%s"), Indent, TEXT(""), *Node.Name.ToString()
+			, *NameViewToString(Node.Processor->GetExecutionOrder().ExecuteBefore)
+			, *NameViewToString(Node.Processor->GetExecutionOrder().ExecuteAfter));
 	}
 }
 
+#if 0 // disabled the graph building for now, leaving the old code here for reference
 struct FDumpGraphDependencyUtils
 {
 	static void DumpGraphNode(FArchive& LogFile, const FProcessorDependencySolver::FNode& Node, int Indent, TSet<const FProcessorDependencySolver::FNode*>& AllNodes, const bool bRoot)
@@ -774,77 +733,40 @@ void FProcessorDependencySolver::DumpGraph(FArchive& LogFile) const
 	LogFile.Logf(TEXT("}"));
 }
 
-void FProcessorDependencySolver::Solve(FNode& RootNode, TConstArrayView<const FName> PriorityNodes, TArray<FProcessorDependencySolver::FOrderInfo>& OutResult, int LoggingIndent)
+#endif // 0; disabled the graph building for now, leaving the old code here for reference
+
+void FProcessorDependencySolver::Solve(TArray<FProcessorDependencySolver::FOrderInfo>& OutResult)
 {
-	if (RootNode.SubNodes.Num() == 0)
+	using UE::Mass::Private::NameViewToString;
+
+	if (AllNodes.Num() == 0)
 	{
 		return;
 	}
 
-	for (FNode& SubNode : RootNode.SubNodes)
+	for (FNode& SubNode : AllNodes)
 	{
 		SubNode.TransientDependencies = SubNode.OriginalDependencies;
 	}
 
 	TArray<int32> IndicesRemaining;
-	// populate with priority nodes and their dependencies first
-	
-	if (PriorityNodes.Num())
-	{
-		IndicesRemaining.Reserve(RootNode.SubNodes.Num());
-		for (const FName PriorityNodeName : PriorityNodes)
-		{
-			// @todo complex names
-			const int32* PriorityNodeIndex = RootNode.Indices.Find(PriorityNodeName);
-			if (PriorityNodeIndex && RootNode.SubNodes.IsValidIndex(*PriorityNodeIndex)
-				// remember that RootNode.Indices points to a whole group's index for elements from a given group
-				&& RootNode.SubNodes[*PriorityNodeIndex].Name == PriorityNodeName)
-			{ 
-				const int32 StartIndex = IndicesRemaining.Add(*PriorityNodeIndex);
-				// in this loop we'll be potentially adding elements to IndicesRemaining and the idea is to process
-				// the values freshly added as well, so that we populate IndicesRemaining with all super-dependencies of 
-				// the priority node currently processed
-				for (int i = StartIndex; i < IndicesRemaining.Num(); ++i)
-				{
-					const int32 SubNodeIndex = IndicesRemaining[i];
-					for (const int32 DependencyIndex : RootNode.SubNodes[SubNodeIndex].OriginalDependencies)
-					{
-						// some dependencies may have already been added to IndicesRemaining by higher-priority nodes
-						// thus we're adding only unique items (IndicesRemaining is a 1-1 mapping of all RootNode's child nodes).
-						IndicesRemaining.AddUnique(DependencyIndex);
-					}
-				}
-			}
-		}
-	}
-
-	if (IndicesRemaining.Num())
-	{
-		// some node indices already added, fill up with missing ones
-		for (int32 i = 0; i < RootNode.SubNodes.Num(); ++i)
-		{
-			IndicesRemaining.AddUnique(i);
-		}
-	}
-	else
+	IndicesRemaining.Reserve(AllNodes.Num());
+	for (int32 i = 0; i < AllNodes.Num(); ++i)
 	{
-		IndicesRemaining.AddZeroed(RootNode.SubNodes.Num());
-		// no priority nodes, just use add all indices in sequence
-		// start from 1 since IndicesRemaining[0] = 0 already
-		for (int32 i = 1; i < RootNode.SubNodes.Num(); ++i)
+		// skip all the group nodes, all group dependencies have already been converted to individual processor dependencies
+		if (AllNodes[i].IsGroup() == false)
 		{
-			IndicesRemaining[i] = i;
+			IndicesRemaining.Add(i);
 		}
 	}
-	
 	FResourceUsage ResourceUsage;
-	
+
 	TArray<int32> SortedNodeIndices;
-	SortedNodeIndices.Reserve(RootNode.SubNodes.Num());
+	SortedNodeIndices.Reserve(AllNodes.Num());
 
 	while (IndicesRemaining.Num())
 	{
-		const bool bStepSuccessful = PerformSolverStep(ResourceUsage, RootNode, IndicesRemaining, SortedNodeIndices);
+		const bool bStepSuccessful = PerformSolverStep(ResourceUsage, IndicesRemaining, SortedNodeIndices);
 
 		if (bStepSuccessful == false)
 		{
@@ -853,7 +775,7 @@ void FProcessorDependencySolver::Solve(FNode& RootNode, TConstArrayView<const FN
 			UE_LOG(LogMass, Error, TEXT("Detected processing dependency cycle:"));
 			for (const int32 Index : IndicesRemaining)
 			{
-				UMassProcessor* Processor = RootNode.SubNodes[Index].Processor;
+				UMassProcessor* Processor = AllNodes[Index].Processor;
 				if (Processor)
 				{
 					UE_LOG(LogMass, Error, TEXT("\t%s, group: %s, before: %s, after %s")
@@ -865,54 +787,57 @@ void FProcessorDependencySolver::Solve(FNode& RootNode, TConstArrayView<const FN
 				else
 				{
 					// group
-					UE_LOG(LogMass, Error, TEXT("\tGroup %s"), *RootNode.SubNodes[Index].Name.ToString());
+					UE_LOG(LogMass, Error, TEXT("\tGroup %s"), *AllNodes[Index].Name.ToString());
 				}
 			}
 			UE_LOG(LogMass, Error, TEXT("Cutting the chain at an arbitrary location."));
 
 			// remove first dependency
 			// note that if we're in a cycle handling scenario every node does have some dependencies left
-			RootNode.SubNodes[IndicesRemaining[0]].TransientDependencies.Pop(/*bAllowShrinking=*/false);
+			AllNodes[IndicesRemaining[0]].TransientDependencies.Pop(/*bAllowShrinking=*/false);
 		}
 	}
 
-	// now we have the desired order in SortedNodeIndices. We have to traverse it recursively to add to OutResult
+	// now we have the desired order in SortedNodeIndices. We have to traverse it to add to OutResult
 	for (int i = 0; i < SortedNodeIndices.Num(); ++i)
 	{
 		const int32 NodeIndex = SortedNodeIndices[i];
 
 		TArray<FName> DependencyNames;
-		for (const int32 DependencyIndex : RootNode.SubNodes[NodeIndex].OriginalDependencies)
+		for (const int32 DependencyIndex : AllNodes[NodeIndex].OriginalDependencies)
 		{
-			DependencyNames.AddUnique(RootNode.SubNodes[DependencyIndex].Name);
+			DependencyNames.AddUnique(AllNodes[DependencyIndex].Name);
 		}
 
-		if (RootNode.SubNodes[NodeIndex].Processor != nullptr)
-		{
-			OutResult.Add({ RootNode.SubNodes[NodeIndex].Name, RootNode.SubNodes[NodeIndex].Processor, EDependencyNodeType::Processor, DependencyNames });
-		}
-		else if (RootNode.SubNodes[NodeIndex].SubNodes.Num() > 0)
+		// at this point we expect SortedNodeIndices to only point to regular processors (i.e. no groups)
+		if (ensure(AllNodes[NodeIndex].Processor != nullptr))
 		{
-			OutResult.Add({ RootNode.SubNodes[NodeIndex].Name, RootNode.SubNodes[NodeIndex].Processor, EDependencyNodeType::GroupStart, DependencyNames });
-			Solve(RootNode.SubNodes[NodeIndex], PriorityNodes, OutResult, LoggingIndent);
-			DependencyNames.Reset(RootNode.SubNodes[NodeIndex].SubNodes.Num());
-			for (const FNode& ChildNode : RootNode.SubNodes[NodeIndex].SubNodes)
-			{
-				DependencyNames.Add(ChildNode.Name);
-			}
-			OutResult.Add({ RootNode.SubNodes[NodeIndex].Name, RootNode.SubNodes[NodeIndex].Processor, EDependencyNodeType::GroupEnd, DependencyNames });
+			OutResult.Add({ AllNodes[NodeIndex].Name, AllNodes[NodeIndex].Processor, EDependencyNodeType::Processor, DependencyNames });
 		}
 	}
 }
 
-void FProcessorDependencySolver::ResolveDependencies(TArray<FProcessorDependencySolver::FOrderInfo>& OutResult, TConstArrayView<const FName> PriorityNodes)
+void FProcessorDependencySolver::ResolveDependencies(TArray<FProcessorDependencySolver::FOrderInfo>& OutResult)
 {
+	if (Processors.Num() == 0)
+	{
+		return;
+	}
+
 	FScopedCategoryAndVerbosityOverride LogOverride(TEXT("LogMass"), ELogVerbosity::Log);
 
 	bAnyCyclesDetected = false;
 
 	UE_LOG(LogMass, Log, TEXT("Gathering dependencies data:"));
 
+	AllNodes.Reset();
+	NodeIndexMap.Reset();
+	// as the very first node we add a "root" node that represents the "top level group" and also simplifies the rest
+	// of the lookup code - if a processor declares it's in group None or depends on Node it we don't need to check that 
+	// explicitly. 
+	AllNodes.Add(FNode(CollectionName, nullptr, 0));
+	NodeIndexMap.Add(FName(), 0);
+
 	// gather the processors information first
 	for (UMassProcessor* Processor : Processors)
 	{
@@ -921,74 +846,25 @@ void FProcessorDependencySolver::ResolveDependencies(TArray<FProcessorDependency
 			UE_LOG(LogMass, Warning, TEXT("%s nullptr found in Processors collection being processed"), ANSI_TO_TCHAR(__FUNCTION__));
 			continue;
 		}
-		AddNode(Processor->GetExecutionOrder().ExecuteInGroup, *Processor);
+		CreateNodes(*Processor);
 	}
 
-	BuildDependencies(GroupRootNode);
-	// @todo anything in GroupRootNode.Dependencies is an undefined symbol
-
-	// Any dependencies that are promoted to the root node means they are unresolved.
-    const bool bAnyUnresolvedDependencies = GroupRootNode.HasDependencies();
-	if (bAnyUnresolvedDependencies)
-	{
-		const UWorld* World = Processors.IsEmpty() ? nullptr : Processors[0]->GetWorld();
-	    for (const FName& DependencyName : GroupRootNode.ExecuteBefore)
-	    {
-		    UE_LOG(LogMass, Log, TEXT("(%s) %s found an unresolved execute before dependency (%s)"),
-		    	World != nullptr ? *ToString(World->GetNetMode()) : TEXT("unknown"), ANSI_TO_TCHAR(__FUNCTION__), *DependencyName.ToString());
-	    }
-    
-	    for (const FName& DependencyName : GroupRootNode.ExecuteAfter)
-	    {
-		    UE_LOG(LogMass, Log, TEXT("(%s) %s found an unresolved execute after dependency (%s)"),
-		    	World != nullptr ? *ToString(World->GetNetMode()) : TEXT("unknown"), ANSI_TO_TCHAR(__FUNCTION__), *DependencyName.ToString());
-	    }
-	}
+	check(AllNodes.Num());
+	LogNode(AllNodes[0]);
 
-	LogNode(GroupRootNode);
-	
-	if (!DependencyGraphFileName.IsEmpty())
-	{
-		const FString FileName = FString::Printf(TEXT("%s%s-%s.dot"), *FPaths::ProjectLogDir(), *DependencyGraphFileName, *FDateTime::Now().ToString());
-		if (FArchive* LogFile = IFileManager::Get().CreateFileWriter(*FileName))
-		{
-			DumpGraph(*LogFile);
+	BuildDependencies();
 
-			LogFile->Close();
-			delete LogFile;
-		}
-		else
-		{
-			UE_LOG(LogMass, Error, TEXT("%s Unable to dump dependency graph into filename %s"), ANSI_TO_TCHAR(__FUNCTION__), *DependencyGraphFileName);
-		}
-	}
+	// now none of the processor nodes depend on groups - we replaced these dependencies with depending directly 
+	// on individual processors. However, we keep the group nodes around since we store the dependencies via index, so 
+	// removing nodes would mess that up. Solve below ignores group nodes and OutResult will not have any groups once its done.
 
-	GroupRootNode.TransientDependencies = GroupRootNode.OriginalDependencies;
-	Solve(GroupRootNode, PriorityNodes, OutResult);
-
-#if WITH_UNREAL_DEVELOPER_TOOLS
-	if (bAnyCyclesDetected || bAnyUnresolvedDependencies)
-	{
-		FMessageLog EditorErrors("LogMass");
-		if(bAnyCyclesDetected)
-		{
-		EditorErrors.Error(LOCTEXT("ProcessorDependenciesCycle", "Processor dependencies cycle found!"));
-		    EditorErrors.Notify(LOCTEXT("ProcessorDependenciesCycle", "Processor dependencies cycle found!"));
-		}
-		if(bAnyUnresolvedDependencies)
-		{
-		    EditorErrors.Info(LOCTEXT("ProcessorUnresolvedDependencies", "Unresolved processor dependencies found!"));
-		    //EditorErrors.Notify(LOCTEXT("ProcessorUnresolvedDependencies", "Unresolved processor dependencies found!"));
-		}
-		EditorErrors.Info(FText::FromString(TEXT("See the log for details")));
-	}
-#endif // WITH_UNREAL_DEVELOPER_TOOLS
+	Solve(OutResult);
 
-	UE_LOG(LogMass, Log, TEXT("Dependency order:"));
+	UE_LOG(LogMass, Verbose, TEXT("Dependency order:"));
 	for (const FProcessorDependencySolver::FOrderInfo& Info : OutResult)
 	{
-		UE_LOG(LogMass, Log, TEXT("\t%s"), *Info.Name.ToString());
+		UE_LOG(LogMass, Verbose, TEXT("\t%s"), *Info.Name.ToString());
 	}
 }
 
-#undef LOCTEXT_NAMESPACE 
+#undef LOCTEXT_NAMESPACE 
\ No newline at end of file
diff --git a/Engine/Plugins/Runtime/MassEntity/Source/MassEntity/Public/MassEntitySettings.h b/Engine/Plugins/Runtime/MassEntity/Source/MassEntity/Public/MassEntitySettings.h
--- a/Engine/Plugins/Runtime/MassEntity/Source/MassEntity/Public/MassEntitySettings.h
+++ b/Engine/Plugins/Runtime/MassEntity/Source/MassEntity/Public/MassEntitySettings.h
@@ -27,9 +27,6 @@ struct FMassProcessingPhaseConfig
 	UPROPERTY(EditAnywhere, Category = Mass, config, NoClear)
 	TSubclassOf<UMassCompositeProcessor> PhaseGroupClass = UMassCompositeProcessor::StaticClass();
 
-	UPROPERTY(EditAnywhere, Category = Mass, config, NoClear)
-	TArray<FName> OffGameThreadGroupNames;
-
 	UPROPERTY(Transient)
 	TArray<UMassProcessor*> ProcessorCDOs;
 
diff --git a/Engine/Plugins/Runtime/MassEntity/Source/MassEntity/Public/MassProcessor.h b/Engine/Plugins/Runtime/MassEntity/Source/MassEntity/Public/MassProcessor.h
--- a/Engine/Plugins/Runtime/MassEntity/Source/MassEntity/Public/MassProcessor.h
+++ b/Engine/Plugins/Runtime/MassEntity/Source/MassEntity/Public/MassProcessor.h
@@ -68,8 +68,6 @@ public:
 	
 	const FMassProcessorExecutionOrder& GetExecutionOrder() const { return ExecutionOrder; }
 
-	TConstArrayView<const int32> GetPrerequisiteIndices() const { return DependencyIndices; }
-
 	/** By default fetches requirements declared entity queries registered via RegisterQuery. Processors can override 
 	 *	this function to supply additional requirements */
 	virtual void ExportRequirements(FMassExecutionRequirements& OutRequirements) const;
@@ -133,8 +131,6 @@ protected:
 #endif // WITH_EDITORONLY_DATA
 
 	friend class UMassCompositeProcessor;
-	TArray<int32> DependencyIndices;
-	TArray<int32> TransientDependencyIndices;
 
 private:
 	/** Stores processor's queries registered via RegisterQuery. 
@@ -186,9 +182,9 @@ protected:
 	virtual void Execute(UMassEntitySubsystem& EntitySubsystem, FMassExecutionContext& Context) override;
 
 	/**
-	 *  Called recursively to add processors and composite processors to ChildPipeline based on ProcessorsAndGroups
+	 *  Adds processors in OrderedProcessors to ChildPipeline and builds flat processing graph that's being used in multithreaded mode.
 	 */
-	int32 Populate(TArray<FProcessorDependencySolver::FOrderInfo>& ProcessorsAndGroups, const int32 StartIndex = 0);
+	void Populate(TConstArrayView<FProcessorDependencySolver::FOrderInfo> OrderedProcessors);
 
 	/** RequestedGroupName can indicate a multi-level group name, like so: A.B.C
 	 *  We need to extract the highest-level group name ('A' in the example), and see if it already exists. 
@@ -227,7 +223,4 @@ protected:
 		}
 	};
 	TArray<FProcessorCompletion> CompletionStatus;
-
-	bool bRunInSeparateThread;
-	bool bHasOffThreadSubGroups;
 };
diff --git a/Engine/Plugins/Runtime/MassEntity/Source/MassEntity/Public/MassProcessorDependencySolver.h b/Engine/Plugins/Runtime/MassEntity/Source/MassEntity/Public/MassProcessorDependencySolver.h
--- a/Engine/Plugins/Runtime/MassEntity/Source/MassEntity/Public/MassProcessorDependencySolver.h
+++ b/Engine/Plugins/Runtime/MassEntity/Source/MassEntity/Public/MassProcessorDependencySolver.h
@@ -63,28 +63,26 @@ struct MASSENTITY_API FMassExecutionRequirements
 	TMassExecutionAccess<FMassExternalSubystemBitSet> RequiredSubsystems;
 };
 
-struct MASSENTITY_API FProcessorDependencySolver
+struct FProcessorDependencySolver
 {
 private:
 	struct FNode
 	{
-		FNode(const FName InName, UMassProcessor* InProcessor) 
-			: Name(InName), Processor(InProcessor)
+		FNode(const FName InName, UMassProcessor* InProcessor, const int32 InNodeIndex = INDEX_NONE) 
+			: Name(InName), Processor(InProcessor), NodeIndex(InNodeIndex)
 		{}
 
 		FName Name = TEXT("");
-		TArray<FNode> SubNodes;
-		TMap<FName, int32> Indices;
 		UMassProcessor* Processor = nullptr;
 		TArray<int32> OriginalDependencies;
 		TArray<int32> TransientDependencies;
 		TArray<FName> ExecuteBefore;
 		TArray<FName> ExecuteAfter;
 		FMassExecutionRequirements Requirements;
+		int32 NodeIndex = INDEX_NONE;
+		TArray<int32> SubNodeIndices;
 
-		int32 FindOrAddGroupNodeIndex(const FString& GroupName);
-		int32 FindNodeIndex(FName InNodeName) const;
-		bool HasDependencies() const;
+		bool IsGroup() const { return Processor == nullptr; }
 	};
 
 	struct FResourceUsage
@@ -128,10 +126,10 @@ public:
 		TArray<FName> Dependencies;
 	};
 
-	FProcessorDependencySolver(TArrayView<UMassProcessor*> InProcessors, const FName Name, const FString& InDependencyGraphFileName = FString());
-	void ResolveDependencies(TArray<FOrderInfo>& OutResult, TConstArrayView<const FName> PriorityNodes = TConstArrayView<const FName>());
+	MASSENTITY_API FProcessorDependencySolver(TArrayView<UMassProcessor*> InProcessors, const FName Name, const FString& InDependencyGraphFileName = FString());
+	MASSENTITY_API void ResolveDependencies(TArray<FOrderInfo>& OutResult);
 
-	static void CreateSubGroupNames(FName InGroupName, TArray<FString>& SubGroupNames);
+	MASSENTITY_API static void CreateSubGroupNames(FName InGroupName, TArray<FString>& SubGroupNames);
 
 protected:
 	// note that internals are protected rather than private to support unit testing
@@ -142,20 +140,21 @@ protected:
 	 * quits.
 	 * @return 'true' if a dependency-less node has been found and added to OutNodeIndices; 'false' otherwise.
 	 */
-	static bool PerformSolverStep(FResourceUsage& ResourceUsage, FNode& RootNode, TArray<int32>& InOutIndicesRemaining, TArray<int32>& OutNodeIndices);
+	bool PerformSolverStep(FResourceUsage& ResourceUsage, TArray<int32>& InOutIndicesRemaining, TArray<int32>& OutNodeIndices);
 	
-	static FString NameViewToString(TConstArrayView<FName> View);
-
-	void AddNode(FName GroupName, UMassProcessor& Processor);
-	void BuildDependencies(FNode& RootNode);
-	void Solve(FNode& RootNode, TConstArrayView<const FName> PriorityNodes, TArray<FProcessorDependencySolver::FOrderInfo>& OutResult, int LoggingIndent = 0);
-	void LogNode(const FNode& RootNode, const FNode* ParentNode = nullptr, int Indent = 0);
-	void DumpGraph(FArchive& LogFile) const;
+	void CreateNodes(UMassProcessor& Processor);
+	void BuildDependencies();
+	void Solve(TArray<FProcessorDependencySolver::FOrderInfo>& OutResult);
+	void LogNode(const FNode& Node, int Indent = 0);
 	
+	// @todo due to fundamental change to how nodes are organized the graph generation needs reimplementation
+	// friend struct FDumpGraphDependencyUtils;
+	// void DumpGraph(FArchive& LogFile) const;
+
 	TArrayView<UMassProcessor*> Processors;
-	FNode GroupRootNode;
 	bool bAnyCyclesDetected = false;
 	FString DependencyGraphFileName;
-
-	friend struct FDumpGraphDependencyUtils;
+	FName CollectionName;
+	TArray<FNode> AllNodes;
+	TMap<FName, int32> NodeIndexMap;
 };
\ No newline at end of file
