Fixes to mass dependency solving required by failing unit tests.
* readded previously removed duplicates handling
* improved one of the dependency tests and improved the solving algorithm to handle "execute before group"-kind of dependencies.

diff --git a/Engine/Plugins/Runtime/MassEntity/Source/MassEntity/Private/MassProcessorDependencySolver.cpp b/Engine/Plugins/Runtime/MassEntity/Source/MassEntity/Private/MassProcessorDependencySolver.cpp
--- a/Engine/Plugins/Runtime/MassEntity/Source/MassEntity/Private/MassProcessorDependencySolver.cpp
+++ b/Engine/Plugins/Runtime/MassEntity/Source/MassEntity/Private/MassProcessorDependencySolver.cpp
@@ -207,27 +207,25 @@
 	check(Processor.GetClass());
 	const FName ProcName = Processor.GetClass()->GetFName();
 
-	check(NodeIndexMap.Find(ProcName) == nullptr);
+	if (NodeIndexMap.Find(ProcName) != nullptr)
+	{
+		UE_LOG(LogMass, Warning, TEXT("%s Processor %s already registered. Duplicates are not supported.")
+			, ANSI_TO_TCHAR(__FUNCTION__), *ProcName.ToString());
+		return;
+	}
 
-	// first check if a node for this Processor has been added, due to duplicates or some other processor given this 
-	// processor as a dependency by name
-	const int32 NodeIndex = AllNodes.Num();
-	NodeIndexMap.Add(ProcName, NodeIndex);
-	AllNodes.Add({ ProcName, &Processor, NodeIndex });
-
 	const FMassProcessorExecutionOrder& ExecutionOrder = Processor.GetExecutionOrder();
 
-	AllNodes[NodeIndex].ExecuteAfter.Append(ExecutionOrder.ExecuteAfter);
-	AllNodes[NodeIndex].ExecuteBefore.Append(ExecutionOrder.ExecuteBefore);
-	Processor.ExportRequirements(AllNodes[NodeIndex].Requirements);
-
+	// first figure out the groups so that the group nodes come before the processor nodes, this is required for child
+	// nodes to inherit group's dependencies like in scenarios where some processor required to ExecuteBefore a given group
+	int32 ParentGroupNodeIndex = INDEX_NONE;
 	if (ExecutionOrder.ExecuteInGroup.IsNone() == false)
 	{
 		TArray<FString> AllGroupNames;
 		CreateSubGroupNames(ExecutionOrder.ExecuteInGroup, AllGroupNames);
 	
 		check(AllGroupNames.Num() > 0);
-		int32 ParentGroupNodeIndex = INDEX_NONE;
+
 		for (const FString& GroupName : AllGroupNames)
 		{
 			const FName GroupFName(GroupName);
@@ -253,8 +251,18 @@
 			}
 
 		}
+	}
 
-		check(ParentGroupNodeIndex > 0);
+	const int32 NodeIndex = AllNodes.Num();
+	NodeIndexMap.Add(ProcName, NodeIndex);
+	FNode& ProcessorNode = AllNodes.Add_GetRef({ ProcName, &Processor, NodeIndex });
+
+	ProcessorNode.ExecuteAfter.Append(ExecutionOrder.ExecuteAfter);
+	ProcessorNode.ExecuteBefore.Append(ExecutionOrder.ExecuteBefore);
+	Processor.ExportRequirements(ProcessorNode.Requirements);
+
+	if (ParentGroupNodeIndex > 0)
+	{
 		AllNodes[ParentGroupNodeIndex].SubNodeIndices.Add(NodeIndex);
 	}
 }
@@ -291,8 +299,9 @@
 	//	else, if Name represents a group:
 	//		- append all group's child node names to ExecuteAfter
 	// 
-	for (FNode& Node : AllNodes)
+	for (int32 NodeIndex = 0; NodeIndex < AllNodes.Num(); ++NodeIndex)
 	{
+		FNode& Node = AllNodes[NodeIndex];
 		for (int i = 0; i < Node.ExecuteAfter.Num(); ++i)
 		{
 			const FName& AfterDependencyName = Node.ExecuteAfter[i];
@@ -319,6 +328,33 @@
 					, *Node.Name.ToString());
 			}
 		}
+
+		// if this node is a group push all the dependencies down on all the children
+		// by design all child nodes come after group nodes so the child nodes' dependencies have not been processed yet
+		if (Node.IsGroup() && Node.SubNodeIndices.Num())
+		{
+			for (int32 PrerequisiteNodeIndex : Node.OriginalDependencies)
+			{
+				checkSlow(PrerequisiteNodeIndex != NodeIndex);
+				// in case of processor nodes we can store it directly
+				if (AllNodes[PrerequisiteNodeIndex].IsGroup() == false)
+				{
+					for (int32 ChildNodeIndex : Node.SubNodeIndices)
+					{
+						AllNodes[ChildNodeIndex].OriginalDependencies.AddUnique(PrerequisiteNodeIndex);
+					}
+				}
+				// special case - if dependency is a group and we haven't processed that group yet, we need to add it by name
+				else if (PrerequisiteNodeIndex > NodeIndex)
+				{
+					const FName& PrerequisiteName = AllNodes[PrerequisiteNodeIndex].Name;
+					for (int32 ChildNodeIndex : Node.SubNodeIndices)
+					{
+						AllNodes[ChildNodeIndex].ExecuteAfter.AddUnique(PrerequisiteName);
+					}
+				}
+			}
+		}
 	}
 }
 
diff --git a/Engine/Plugins/Runtime/MassEntity/Source/MassEntityTestSuite/Private/MassDependencyTest.cpp b/Engine/Plugins/Runtime/MassEntity/Source/MassEntityTestSuite/Private/MassDependencyTest.cpp
--- a/Engine/Plugins/Runtime/MassEntity/Source/MassEntityTestSuite/Private/MassDependencyTest.cpp
+++ b/Engine/Plugins/Runtime/MassEntity/Source/MassEntityTestSuite/Private/MassDependencyTest.cpp
@@ -134,6 +134,7 @@
 		{
 			UMassTestProcessorBase* Proc = Processors.Add_GetRef(NewObject<UMassTestProcessor_D>());
 			Proc->GetMutableExecutionOrder().ExecuteBefore.Add(UMassTestProcessor_A::StaticClass()->GetFName());
+			Proc->GetMutableExecutionOrder().ExecuteBefore.Add(TEXT("X.Y"));
 		}
 
 		{
@@ -152,13 +153,9 @@
 	{
 		Solve();
 
-		// dump all the group information from the Result collection for easier ordering testing
 		for (int32 i = 0; i < Result.Num(); ++i)
 		{
-			if (Result[i].NodeType != EDependencyNodeType::Processor)
-			{
-				Result.RemoveAt(i--, 1, /*bAllowShrinking=*/false);
-			}
+			AITEST_EQUAL("We expect only processor nodes in the results", Result[i].NodeType, EDependencyNodeType::Processor);
 		}
 
 		AITEST_TRUE("D is the only fully dependency-less processor so should be first", Result[0].Name == GetProcessorName<UMassTestProcessor_D>());		
