Mass: Added more visually informative flat processing graph logging, essentially describing processors' runtime dependencies.

diff --git a/Engine/Plugins/Runtime/MassEntity/Source/MassEntity/Private/MassProcessor.cpp b/Engine/Plugins/Runtime/MassEntity/Source/MassEntity/Private/MassProcessor.cpp
--- a/Engine/Plugins/Runtime/MassEntity/Source/MassEntity/Private/MassProcessor.cpp
+++ b/Engine/Plugins/Runtime/MassEntity/Source/MassEntity/Private/MassProcessor.cpp
@@ -337,25 +337,28 @@ void UMassCompositeProcessor::CopyAndSort(const FMassProcessingPhaseConfig& Phas
 	}
 
 #if WITH_MASSENTITY_DEBUG
+	FScopedCategoryAndVerbosityOverride LogOverride(TEXT("LogMass"), ELogVerbosity::Log);
 	UE_LOG(LogMass, Log, TEXT("%s flat processing graph:"), *GroupName.ToString());
 
+	TArray<int32> Offset;
+	Offset.Reserve(ProcessingFlatGraph.Num());
+
 	for (int i = 0; i < ProcessingFlatGraph.Num(); ++i)
 	{
 		FDependencyNode& ProcessingNode = ProcessingFlatGraph[i];
 		FString DependenciesDesc;
+		int32 MyOffset = 0;
 		for (const int32 DependencyIndex : ProcessingNode.Dependencies)
 		{
-			DependenciesDesc += FString::Printf(TEXT("%s, "), *ProcessingFlatGraph[DependencyIndex].Name.ToString());
+			DependenciesDesc += FString::Printf(TEXT("%d, "), DependencyIndex);
+			MyOffset = FMath::Max(MyOffset, Offset[DependencyIndex]);
 		}
+		MyOffset += 2;
+		Offset.Add(MyOffset);
 
 		if (ProcessingNode.Processor)
 		{
-			UE_LOG(LogMass, Log, TEXT("Task %s%s%s"), *ProcessingNode.Processor->GetProcessorName()
-				, DependenciesDesc.Len() > 0 ? TEXT(" depends on ") : TEXT(""), *DependenciesDesc);
-		}
-		else
-		{
-			UE_LOG(LogMass, Log, TEXT("Group %s%s%s"), *ProcessingNode.Name.ToString()
+			UE_LOG(LogMass, Log, TEXT("[%2d]%*s%s%s%s"), i, MyOffset, TEXT(""), *ProcessingNode.Processor->GetProcessorName()
 				, DependenciesDesc.Len() > 0 ? TEXT(" depends on ") : TEXT(""), *DependenciesDesc);
 		}
 	}
