Fix to null dereference when calculating memory stats.

diff --git a/Engine/Plugins/Runtime/MassEntity/Source/MassEntity/Private/MassEntitySubsystem.cpp b/Engine/Plugins/Runtime/MassEntity/Source/MassEntity/Private/MassEntitySubsystem.cpp
--- a/Engine/Plugins/Runtime/MassEntity/Source/MassEntity/Private/MassEntitySubsystem.cpp
+++ b/Engine/Plugins/Runtime/MassEntity/Source/MassEntity/Private/MassEntitySubsystem.cpp
@@ -49,7 +49,7 @@ void UMassEntitySubsystem::GetResourceSizeEx(FResourceSizeEx& CumulativeResource
 
 	const SIZE_T MyExtraSize = Entities.GetAllocatedSize() + 
 		EntityFreeIndexList.GetAllocatedSize() +
-		DeferredCommandBuffer->GetAllocatedSize() +
+		(DeferredCommandBuffer != nullptr ? DeferredCommandBuffer->GetAllocatedSize() : 0) +
 		FragmentHashToArchetypeMap.GetAllocatedSize() +
 		FragmentTypeToArchetypeMap.GetAllocatedSize();
 	CumulativeResourceSize.AddDedicatedSystemMemoryBytes(MyExtraSize);
