Fixed circular dependencies case for mass processor dependencies solver.
Also added missing unit tests for dependencies solving:
* handling missing dependencies
* single non-null processor
* circular dependencies

diff --git a/Engine/Plugins/Runtime/MassEntity/Source/MassEntity/Private/MassProcessorDependencySolver.cpp b/Engine/Plugins/Runtime/MassEntity/Source/MassEntity/Private/MassProcessorDependencySolver.cpp
--- a/Engine/Plugins/Runtime/MassEntity/Source/MassEntity/Private/MassProcessorDependencySolver.cpp
+++ b/Engine/Plugins/Runtime/MassEntity/Source/MassEntity/Private/MassProcessorDependencySolver.cpp
@@ -28,6 +28,22 @@
 }
 
 //----------------------------------------------------------------------//
+//  FMassExecutionRequirements
+//----------------------------------------------------------------------//
+void FMassExecutionRequirements::CountResourcesUsed()
+{
+	ResourcesUsedCount = 0;
+
+	for (int i = 0; i < EMassAccessOperation::MAX; ++i)
+	{
+		ResourcesUsedCount += Fragments[i].CountStoredTypes();
+		ResourcesUsedCount += ChunkFragments[i].CountStoredTypes();
+		ResourcesUsedCount += SharedFragments[i].CountStoredTypes();
+		ResourcesUsedCount += RequiredSubsystems[i].CountStoredTypes();
+	}
+}
+
+//----------------------------------------------------------------------//
 //  FProcessorDependencySolver::FResourceUsage
 //----------------------------------------------------------------------//
 FProcessorDependencySolver::FResourceUsage::FResourceUsage()
@@ -260,6 +276,7 @@
 	ProcessorNode.ExecuteAfter.Append(ExecutionOrder.ExecuteAfter);
 	ProcessorNode.ExecuteBefore.Append(ExecutionOrder.ExecuteBefore);
 	Processor.ExportRequirements(ProcessorNode.Requirements);
+	ProcessorNode.Requirements.CountResourcesUsed();
 
 	if (ParentGroupNodeIndex > 0)
 	{
@@ -276,14 +293,14 @@
 		for (const FName& BeforeDependencyName : Node.ExecuteBefore)
 		{
 			int32* DependentNodeIndex = NodeIndexMap.Find(BeforeDependencyName);
-			if (ensure(DependentNodeIndex))
+			if (DependentNodeIndex != nullptr)
 			{
 				check(AllNodes.IsValidIndex(*DependentNodeIndex));
 				AllNodes[*DependentNodeIndex].ExecuteAfter.Add(Node.Name);
 			}
 			else
 			{
-				UE_LOG(LogMass, Log, TEXT("Unable to find dependency \"%s\" declared by %s"), *BeforeDependencyName.ToString()
+				UE_LOG(LogMass, Error, TEXT("Unable to find dependency \"%s\" declared by %s"), *BeforeDependencyName.ToString()
 					, *Node.Name.ToString());
 			}
 		}
@@ -306,7 +323,7 @@
 		{
 			const FName& AfterDependencyName = Node.ExecuteAfter[i];
 			int32* PrerequisiteNodeIndex = NodeIndexMap.Find(AfterDependencyName);
-			if (ensure(PrerequisiteNodeIndex))
+			if (PrerequisiteNodeIndex != nullptr)
 			{
 				const FNode& PrerequisiteNode = AllNodes[*PrerequisiteNodeIndex];
 
@@ -324,7 +341,7 @@
 			}
 			else
 			{
-				UE_LOG(LogMass, Log, TEXT("Unable to find dependency \"%s\" declared by %s"), *AfterDependencyName.ToString()
+				UE_LOG(LogMass, Error, TEXT("Unable to find dependency \"%s\" declared by %s"), *AfterDependencyName.ToString()
 					, *Node.Name.ToString());
 			}
 		}
@@ -780,9 +797,9 @@
 		return;
 	}
 
-	for (FNode& SubNode : AllNodes)
+	for (FNode& Node : AllNodes)
 	{
-		SubNode.TransientDependencies = SubNode.OriginalDependencies;
+		Node.TransientDependencies = Node.OriginalDependencies;
 	}
 
 	TArray<int32> IndicesRemaining;
@@ -795,6 +812,12 @@
 			IndicesRemaining.Add(i);
 		}
 	}
+	// order the indices to process according to calculated resource hunger
+	IndicesRemaining.Sort([this](const int32 IndexA, const int32 IndexB){
+		return AllNodes[IndexA].Requirements.ResourcesUsedCount < AllNodes[IndexB].Requirements.ResourcesUsedCount;
+	});
+
+	// this is where we'll be tracking what's being accessed by whom
 	FResourceUsage ResourceUsage;
 
 	TArray<int32> SortedNodeIndices;
@@ -814,7 +837,7 @@
 				UMassProcessor* Processor = AllNodes[Index].Processor;
 				if (Processor)
 				{
-					UE_LOG(LogMass, Error, TEXT("\t%s, group: %s, before: %s, after %s")
+					UE_LOG(LogMass, Warning, TEXT("\t%s, group: %s, before: %s, after %s")
 						, *Processor->GetName()
 						, *Processor->GetExecutionOrder().ExecuteInGroup.ToString()
 						, *NameViewToString(Processor->GetExecutionOrder().ExecuteBefore)
@@ -823,14 +846,17 @@
 				else
 				{
 					// group
-					UE_LOG(LogMass, Error, TEXT("\tGroup %s"), *AllNodes[Index].Name.ToString());
+					UE_LOG(LogMass, Warning, TEXT("\tGroup %s"), *AllNodes[Index].Name.ToString());
 				}
 			}
-			UE_LOG(LogMass, Error, TEXT("Cutting the chain at an arbitrary location."));
+			UE_LOG(LogMass, Warning, TEXT("Cutting the chain at an arbitrary location."));
 
 			// remove first dependency
 			// note that if we're in a cycle handling scenario every node does have some dependencies left
-			AllNodes[IndicesRemaining[0]].TransientDependencies.Pop(/*bAllowShrinking=*/false);
+			const int32 DependencyNodeIndex = AllNodes[IndicesRemaining[0]].TransientDependencies.Pop(/*bAllowShrinking=*/false);
+			// we need to remove this dependency from original dependencies as well, otherwise we'll still have the cycle
+			// in the data being produces as a result of the whole algorithm
+			AllNodes[IndicesRemaining[0]].OriginalDependencies.Remove(DependencyNodeIndex);
 		}
 	}
 
diff --git a/Engine/Plugins/Runtime/MassEntity/Source/MassEntity/Public/MassProcessorDependencySolver.h b/Engine/Plugins/Runtime/MassEntity/Source/MassEntity/Public/MassProcessorDependencySolver.h
--- a/Engine/Plugins/Runtime/MassEntity/Source/MassEntity/Public/MassProcessorDependencySolver.h
+++ b/Engine/Plugins/Runtime/MassEntity/Source/MassEntity/Public/MassProcessorDependencySolver.h
@@ -57,10 +57,13 @@ struct MASSENTITY_API FMassExecutionRequirements
 		return *this;
 	}
 
+	void CountResourcesUsed();
+
 	TMassExecutionAccess<FMassFragmentBitSet> Fragments;
 	TMassExecutionAccess<FMassChunkFragmentBitSet> ChunkFragments;
 	TMassExecutionAccess<FMassSharedFragmentBitSet> SharedFragments;
 	TMassExecutionAccess<FMassExternalSubystemBitSet> RequiredSubsystems;
+	int32 ResourcesUsedCount = INDEX_NONE;
 };
 
 struct FProcessorDependencySolver
diff --git a/Engine/Plugins/Runtime/MassEntity/Source/MassEntityTestSuite/Private/MassDependencyTest.cpp b/Engine/Plugins/Runtime/MassEntity/Source/MassEntityTestSuite/Private/MassDependencyTest.cpp
--- a/Engine/Plugins/Runtime/MassEntity/Source/MassEntityTestSuite/Private/MassDependencyTest.cpp
+++ b/Engine/Plugins/Runtime/MassEntity/Source/MassEntityTestSuite/Private/MassDependencyTest.cpp
@@ -40,6 +40,28 @@
 
 struct FTrivialDependency : FDependencySolverBase
 {
+	UMassTestProcessorBase* Proc = nullptr;
+
+	virtual bool SetUp() override
+	{
+		Proc = Processors.Add_GetRef(NewObject<UMassTestProcessor_A>());
+		return true;
+	}
+
+	virtual bool InstantTest() override
+	{
+		Solve();
+
+		AITEST_EQUAL("The results should contain only a single processor", Result.Num(), 1);
+		AITEST_EQUAL("The sole processor should be the one we've added", Result[0].Processor, Proc);
+
+		return true;
+	}
+};
+IMPLEMENT_AI_INSTANT_TEST(FTrivialDependency, "System.Mass.Dependencies.Trivial");
+
+struct FSimpleDependency : FDependencySolverBase
+{
 	virtual bool SetUp() override
 	{
 		{
@@ -68,7 +90,43 @@
 		return true;
 	}
 };
-IMPLEMENT_AI_INSTANT_TEST(FTrivialDependency, "System.Mass.Dependencies.Trivial");
+IMPLEMENT_AI_INSTANT_TEST(FSimpleDependency, "System.Mass.Dependencies.Simple");
+
+struct FMissingDependency : FDependencySolverBase
+{
+	virtual bool SetUp() override
+	{
+		{
+			UMassTestProcessorBase* Proc = Processors.Add_GetRef(NewObject<UMassTestProcessor_A>());
+			Proc->GetMutableExecutionOrder().ExecuteAfter.Add(TEXT("NonExistingDependency"));
+		}
+		{
+			UMassTestProcessorBase* Proc = Processors.Add_GetRef(NewObject<UMassTestProcessor_B>());
+			Proc->GetMutableExecutionOrder().ExecuteBefore.Add(TEXT("NonExistingDependency2"));
+			Proc->GetMutableExecutionOrder().ExecuteAfter.Add(GetProcessorName<UMassTestProcessor_C>());
+			Proc->GetMutableExecutionOrder().ExecuteBefore.Add(GetProcessorName<UMassTestProcessor_A>());
+		}
+		{
+			UMassTestProcessorBase* Proc = Processors.Add_GetRef(NewObject<UMassTestProcessor_C>());
+		}
+		return true;
+	}
+
+	virtual bool InstantTest() override
+	{
+		// we're expecting two occurrences, for NonExistingDependency and NonExistingDependency2
+		GetTestRunner().AddExpectedError(TEXT("Unable to find dependency"), EAutomationExpectedErrorFlags::Contains, 2);
+
+		Solve();
+
+		AITEST_TRUE("C is expected to be the first one", Result[0].Name == GetProcessorName<UMassTestProcessor_C>());
+		AITEST_TRUE("Then B", Result[1].Name == GetProcessorName<UMassTestProcessor_B>());
+		AITEST_TRUE("With A being last", Result[2].Name == GetProcessorName<UMassTestProcessor_A>());
+
+		return true;
+	}
+};
+IMPLEMENT_AI_INSTANT_TEST(FMissingDependency, "System.Mass.Dependencies.MissingDependencies");
 
 
 struct FDeepGroup : FDependencySolverBase
@@ -271,14 +329,56 @@
 		return true;
 	}
 };
-
 IMPLEMENT_AI_INSTANT_TEST(FGroupNamesGeneration, "System.Mass.Dependencies.SubgroupNames");
 
-// tests to add:
-// * handling missing dependencies
-// * single non-null processor
-// * circular dependencies  
+struct FCircularDependency : FDependencySolverBase
+{
+	virtual bool SetUp() override
+	{
+		{
+			UMassTestProcessorBase* Proc = Processors.Add_GetRef(NewObject<UMassTestProcessor_A>());
+			Proc->TestGetQuery().AddRequirement<FTestFragment_Int>(EMassFragmentAccess::ReadWrite);
+			Proc->GetMutableExecutionOrder().ExecuteAfter.Add(GetProcessorName<UMassTestProcessor_D>());
+		}
+		{
+			UMassTestProcessorBase* Proc = Processors.Add_GetRef(NewObject<UMassTestProcessor_B>());
+			Proc->TestGetQuery().AddRequirement<FTestFragment_Int>(EMassFragmentAccess::ReadWrite);
+			Proc->GetMutableExecutionOrder().ExecuteAfter.Add(GetProcessorName<UMassTestProcessor_A>());
+		}
+		{
+			UMassTestProcessorBase* Proc = Processors.Add_GetRef(NewObject<UMassTestProcessor_C>());
+			Proc->TestGetQuery().AddRequirement<FTestFragment_Int>(EMassFragmentAccess::ReadWrite);
+			Proc->GetMutableExecutionOrder().ExecuteAfter.Add(GetProcessorName<UMassTestProcessor_B>());
+		}
+		{
+			UMassTestProcessorBase* Proc = Processors.Add_GetRef(NewObject<UMassTestProcessor_D>());
+			Proc->TestGetQuery().AddRequirement<FTestFragment_Int>(EMassFragmentAccess::ReadWrite);
+			Proc->GetMutableExecutionOrder().ExecuteAfter.Add(GetProcessorName<UMassTestProcessor_C>());
+		}
 
+		return true;
+	}
+
+	virtual bool InstantTest() override
+	{
+		GetTestRunner().AddExpectedError(TEXT("Detected processing dependency cycle"), EAutomationExpectedErrorFlags::Contains, 1);
+
+		Solve();
+
+		// every subsequent processor is expected to depend only on the previous one since all the processors use exactly the same resources
+		AITEST_TRUE("The first processor has no dependencies", Result[0].Dependencies.IsEmpty());
+		for (int i = 1; i < Result.Num(); ++i)
+		{
+			const auto& ResultNode = Result[i];
+			AITEST_EQUAL("The subsequent processors has only one dependency", ResultNode.Dependencies.Num(), 1);
+			AITEST_EQUAL("The subsequent processors depend only on the previous one", ResultNode.Dependencies[0], Result[i - 1].Name);
+		}
+
+		return true;
+	}
+};
+IMPLEMENT_AI_INSTANT_TEST(FCircularDependency, "System.Mass.Dependencies.Circular");
+
 } // FMassDependencySolverTest
 
 PRAGMA_ENABLE_OPTIMIZATION
