Improvements to AITestSuite's latent tests.
This CL also contains Mass MT tests utilizing the improvements. The tests required a minor massaging of MassCompositeProcessor's API.

diff --git a/Engine/Plugins/Runtime/MassEntity/Source/MassEntity/Private/MassProcessor.cpp b/Engine/Plugins/Runtime/MassEntity/Source/MassEntity/Private/MassProcessor.cpp
--- a/Engine/Plugins/Runtime/MassEntity/Source/MassEntity/Private/MassProcessor.cpp
+++ b/Engine/Plugins/Runtime/MassEntity/Source/MassEntity/Private/MassProcessor.cpp
@@ -309,8 +309,13 @@ void UMassCompositeProcessor::CopyAndSort(const FMassProcessingPhaseConfig& Phas
 	FMassRuntimePipeline TmpPipeline;
 	TmpPipeline.CreateFromArray(PhaseConfig.ProcessorCDOs, *GetOuter());
 
+	SetProcessors(TmpPipeline.Processors, DependencyGraphFileName);
+}
+
+void UMassCompositeProcessor::SetProcessors(TArrayView<UMassProcessor*> InProcessorInstances, const FString& DependencyGraphFileName)
+{
 	// figure out dependencies
-	FProcessorDependencySolver Solver(TmpPipeline.Processors, GroupName, DependencyGraphFileName);
+	FProcessorDependencySolver Solver(InProcessorInstances, GroupName, DependencyGraphFileName);
 	TArray<FProcessorDependencySolver::FOrderInfo> SortedProcessorsAndGroups;
 	Solver.ResolveDependencies(SortedProcessorsAndGroups);
 
diff --git a/Engine/Plugins/Runtime/MassEntity/Source/MassEntity/Public/MassProcessor.h b/Engine/Plugins/Runtime/MassEntity/Source/MassEntity/Public/MassProcessor.h
--- a/Engine/Plugins/Runtime/MassEntity/Source/MassEntity/Public/MassProcessor.h
+++ b/Engine/Plugins/Runtime/MassEntity/Source/MassEntity/Public/MassProcessor.h
@@ -165,6 +165,7 @@ public:
 	FName GetGroupName() const { return GroupName; }
 
 	virtual void CopyAndSort(const FMassProcessingPhaseConfig& PhaseConfig, const FString& DependencyGraphFileName = FString());
+	void SetProcessors(TArrayView<UMassProcessor*> InProcessorInstances, const FString& DependencyGraphFileName = FString());
 
 	/** adds SubProcessor to an appropriately named group. If RequestedGroupName == None then SubProcessor
 	 *  will be added directly to ChildPipeline. If not then the indicated group will be searched for in ChildPipeline 
diff --git a/Engine/Plugins/Runtime/MassEntity/Source/MassEntityTestSuite/Private/MassMultiThreadTest.cpp b/Engine/Plugins/Runtime/MassEntity/Source/MassEntityTestSuite/Private/MassMultiThreadTest.cpp
new file mode 100644
--- /dev/null
+++ b/Engine/Plugins/Runtime/MassEntity/Source/MassEntityTestSuite/Private/MassMultiThreadTest.cpp
@@ -0,0 +1,187 @@
+// Copyright Epic Games, Inc. All Rights Reserved.
+
+#include "MassEntitySubsystem.h"
+#include "MassProcessingTypes.h"
+#include "MassEntityTestTypes.h"
+#include "MassEntitySettings.h"
+#include "MassExecutor.h"
+#include "MassEntityView.h"
+
+#define LOCTEXT_NAMESPACE "MassTest"
+
+PRAGMA_DISABLE_OPTIMIZATION
+
+//----------------------------------------------------------------------//
+// tests 
+//----------------------------------------------------------------------//
+namespace FMassMultiThreadingTest
+{
+
+template<typename T>
+static FName GetProcessorName()
+{
+	return T::StaticClass()->GetFName();
+}
+
+struct FMTTestBase : FEntityTestBase
+{
+	using Super = FExecutionTestBase;
+	UMassCompositeProcessor* CompositeProcessor = nullptr;
+	TArray<UMassTestProcessorBase*> Processors;
+	TArray<FProcessorDependencySolver::FOrderInfo> Result;
+	FGraphEventRef FinishEvent;
+	TArray<FMassEntityHandle> Entities;
+
+	virtual bool Update() override
+	{
+		// cannot do this in SetUp without adding a new virtual function for subtests to override
+		if (CompositeProcessor == nullptr)
+		{
+			CompositeProcessor = NewObject<UMassCompositeProcessor>();
+			CompositeProcessor->SetGroupName(TEXT("Test"));
+			CompositeProcessor->SetProcessors(MakeArrayView<UMassProcessor*>((UMassProcessor**)Processors.GetData(), Processors.Num()));
+
+			FMassProcessingContext Context(*EntitySubsystem, /*DeltaTime=*/0);
+			FinishEvent = UE::Mass::Executor::TriggerParallelTasks(*CompositeProcessor, Context, []() {});
+		}
+
+		if (FinishEvent->IsComplete())
+		{
+			VerifyResults();
+			// signal that we're done with this test
+			return true;
+		}
+		return false;
+	}
+
+	virtual void VerifyResults() = 0;
+};
+
+struct FMTTrivial : FMTTestBase
+{
+	using Super = FMTTestBase;
+	const int32 NumToCreate = 200;
+	int32 NumProcessed = 0;
+
+	virtual bool SetUp() override
+	{	
+		if (!Super::SetUp())
+		{
+			return false;
+		}
+
+		EntitySubsystem->BatchCreateEntities(IntsArchetype, NumToCreate, Entities);
+
+		Processors.Reset();
+		{
+			UMassTestProcessorBase* Proc = Processors.Add_GetRef(NewObject<UMassTestProcessor_A>());
+
+			Proc->TestGetQuery().AddRequirement<FTestFragment_Int>(EMassFragmentAccess::ReadOnly);
+			Proc->ExecutionFunction = [this, Proc](UMassEntitySubsystem& InEntitySubsystem, FMassExecutionContext& Context)
+			{
+				Proc->TestGetQuery().ForEachEntityChunk(InEntitySubsystem, Context, [this](FMassExecutionContext& Context)
+					{
+						NumProcessed += Context.GetNumEntities();
+					});
+			};
+		}
+
+		return true;
+	}
+	
+	virtual void VerifyResults()
+	{
+		AITEST_EQUAL_LATENT("Expected to process all the created entities.", NumToCreate, NumProcessed);
+	}
+};
+IMPLEMENT_AI_LATENT_TEST(FMTTrivial, "System.Mass.Multithreading.Trivial");
+
+
+struct FMTBasic : FMTTestBase
+{
+	using Super = FMTTestBase;
+	const int32 NumToCreate = 200;
+	int32 NumProcessed = 0;
+
+	virtual bool SetUp() override
+	{
+		if (!Super::SetUp())
+		{
+			return false;
+		}
+
+		EntitySubsystem->BatchCreateEntities(FloatsIntsArchetype, NumToCreate, Entities);
+
+		Processors.Reset();
+		{
+			UMassTestProcessorBase* Proc = Processors.Add_GetRef(NewObject<UMassTestProcessor_C>());
+			Proc->GetMutableExecutionOrder().ExecuteAfter.Add(GetProcessorName<UMassTestProcessor_B>());
+			Proc->TestGetQuery().AddRequirement<FTestFragment_Float>(EMassFragmentAccess::ReadOnly);
+			Proc->TestGetQuery().AddRequirement<FTestFragment_Int>(EMassFragmentAccess::ReadWrite);
+			Proc->ExecutionFunction = [this, Proc](UMassEntitySubsystem& InEntitySubsystem, FMassExecutionContext& Context)
+			{
+				Proc->TestGetQuery().ForEachEntityChunk(InEntitySubsystem, Context, [this](FMassExecutionContext& Context)
+					{
+						const TArrayView<FTestFragment_Int> IntsList = Context.GetMutableFragmentView<FTestFragment_Int>();
+						const TConstArrayView<FTestFragment_Float> FloatsList = Context.GetFragmentView<FTestFragment_Float>();
+						for (int32 i = 0; i < Context.GetNumEntities(); ++i)
+						{
+							IntsList[i].Value = int(FloatsList[i].Value) + IntsList[i].Value;
+						}
+					});
+			};
+		} 
+		{
+			UMassTestProcessorBase* Proc = Processors.Add_GetRef(NewObject<UMassTestProcessor_B>());
+			Proc->GetMutableExecutionOrder().ExecuteAfter.Add(GetProcessorName<UMassTestProcessor_A>());
+			Proc->TestGetQuery().AddRequirement<FTestFragment_Int>(EMassFragmentAccess::ReadOnly);
+			Proc->TestGetQuery().AddRequirement<FTestFragment_Float>(EMassFragmentAccess::ReadWrite);
+			Proc->ExecutionFunction = [this, Proc](UMassEntitySubsystem& InEntitySubsystem, FMassExecutionContext& Context)
+			{
+				Proc->TestGetQuery().ForEachEntityChunk(InEntitySubsystem, Context, [this](FMassExecutionContext& Context)
+					{
+						const TConstArrayView<FTestFragment_Int> IntsList = Context.GetFragmentView<FTestFragment_Int>();
+						const TArrayView<FTestFragment_Float> FloatsList = Context.GetMutableFragmentView<FTestFragment_Float>();
+						for (int32 i = 0; i < Context.GetNumEntities(); ++i)
+						{
+							FloatsList[i].Value = float(IntsList[i].Value * IntsList[i].Value);
+						}
+					});
+			};
+		} 
+		{
+			UMassTestProcessorBase* Proc = Processors.Add_GetRef(NewObject<UMassTestProcessor_A>());
+			Proc->TestGetQuery().AddRequirement<FTestFragment_Int>(EMassFragmentAccess::ReadWrite);
+			Proc->ExecutionFunction = [this, Proc](UMassEntitySubsystem& InEntitySubsystem, FMassExecutionContext& Context)
+			{
+				int Index = 0;
+				Proc->TestGetQuery().ForEachEntityChunk(InEntitySubsystem, Context, [this, &Index](FMassExecutionContext& Context)
+					{
+						const TArrayView<FTestFragment_Int> IntsList = Context.GetMutableFragmentView<FTestFragment_Int>();
+						for (int32 i = 0; i < Context.GetNumEntities(); ++i)
+						{
+							IntsList[i].Value = Index++;
+						}
+					});
+			};
+		}
+
+		return true;
+	}
+
+	virtual void VerifyResults()
+	{
+		for (int i = 0; i < Entities.Num(); ++i)
+		{
+			FMassEntityView View(FloatsIntsArchetype, Entities[i]);
+			AITEST_EQUAL_LATENT(TEXT("Should have predicted values"), View.GetFragmentData<FTestFragment_Int>().Value, i*i + i);
+		}
+	}
+};
+IMPLEMENT_AI_LATENT_TEST(FMTBasic, "System.Mass.Multithreading.Basic");
+
+} // FMassMultiThreadingTest
+
+PRAGMA_ENABLE_OPTIMIZATION
+
+#undef LOCTEXT_NAMESPACE
diff --git a/Engine/Source/Developer/AITestSuite/Public/AITestsCommon.h b/Engine/Source/Developer/AITestSuite/Public/AITestsCommon.h
--- a/Engine/Source/Developer/AITestSuite/Public/AITestsCommon.h
+++ b/Engine/Source/Developer/AITestSuite/Public/AITestsCommon.h
@@ -194,29 +194,16 @@ struct FAITest_SimpleComponentBasedTest : public FAITestBase
 // fails making sure the rest of the test relying on given condition being 
 // true doesn't crash
 //----------------------------------------------------------------------//
-#define AITEST_TRUE(What, Value)\
-	if (!GetTestRunner().TestTrue(What, Value))\
+#define __AITEST_IMPL(What, Value, Test, RetVal)\
+	if (!GetTestRunner().Test(What, Value))\
 	{\
-		return false;\
-	}
-
-#define AITEST_FALSE(What, Value)\
-	if (!GetTestRunner().TestFalse(What, Value))\
-	{\
-		return false;\
-	}
-
-#define AITEST_NULL(What, Pointer)\
-	if (!GetTestRunner().TestNull(What, Pointer))\
-	{\
-		return false;\
-	}
-
-#define AITEST_NOT_NULL(What, Pointer)\
-	if (!GetTestRunner().TestNotNull(What, Pointer))\
-	{\
-		return false;\
+		return RetVal;\
 	}
+	
+#define AITEST_TRUE(What, Value) __AITEST_IMPL(What, Value, TestTrue, false)
+#define AITEST_FALSE(What, Value) __AITEST_IMPL(What, Value, TestFalse, false)
+#define AITEST_NULL(What, Pointer) __AITEST_IMPL(What, Pointer, TestNull, false)
+#define AITEST_NOT_NULL(What, Pointer) __AITEST_IMPL(What, Pointer, TestNotNull, false)
 
 namespace FTestHelpers
 {
@@ -249,14 +236,28 @@ namespace FTestHelpers
 	}
 }
 
-#define AITEST_EQUAL(What, Actual, Expected)\
-	if (!FTestHelpers::TestEqual(What, Actual, Expected, GetTestRunner()))\
+#define __AITEST_HELPER_IMPL(What, Actual, Expected, Test, RetVal)\
+	if (!FTestHelpers::Test(What, Actual, Expected, GetTestRunner()))\
 	{\
-		return false;\
+		return RetVal;\
 	}
 
-#define AITEST_NOT_EQUAL(What, Actual, Expected)\
-	if (!FTestHelpers::TestNotEqual(What, Actual, Expected, GetTestRunner()))\
-	{\
-		return false;\
-	}
+#define AITEST_EQUAL(What, Actual, Expected) __AITEST_HELPER_IMPL(What, Actual, Expected, TestEqual, false)
+#define AITEST_NOT_EQUAL(What, Actual, Expected) __AITEST_HELPER_IMPL(What, Actual, Expected, TestNotEqual, false)
+
+
+//----------------------------------------------------------------------//
+// state testing macros, valid in FTestAIBase (and subclasses') methods. 
+// Using these macros makes sure the test function fails if the assertion
+// fails making sure the rest of the test relying on given condition being 
+// true doesn't crash. Note that these macros are intended to be put in 
+// a void-returning function where latent test's results can be verified. 
+// Update function is not a good place for testing results since its return
+// value controls whether the test will continue.
+//----------------------------------------------------------------------//
+#define AITEST_TRUE_LATENT(What, Value) __AITEST_IMPL(What, Value, TestTrue, )
+#define AITEST_FALSE_LATENT(What, Value) __AITEST_IMPL(What, Value, TestFalse, )
+#define AITEST_NULL_LATENT(What, Pointer) __AITEST_IMPL(What, Pointer, TestNull, )
+#define AITEST_NOT_NULL_LATENT(What, Pointer) __AITEST_IMPL(What, Pointer, TestNotNull, )
+#define AITEST_EQUAL_LATENT(What, Actual, Expected) __AITEST_HELPER_IMPL(What, Actual, Expected, TestEqual, )
+#define AITEST_NOT_EQUAL_LATENT(What, Actual, Expected) __AITEST_HELPER_IMPL(What, Actual, Expected, TestNotEqual, )
