Constified FMassArchetypeData::GetRequirements* functions.

diff --git a/Engine/Plugins/Runtime/MassEntity/Source/MassEntity/Private/MassArchetypeData.cpp b/Engine/Plugins/Runtime/MassEntity/Source/MassEntity/Private/MassArchetypeData.cpp
--- a/Engine/Plugins/Runtime/MassEntity/Source/MassEntity/Private/MassArchetypeData.cpp
+++ b/Engine/Plugins/Runtime/MassEntity/Source/MassEntity/Private/MassArchetypeData.cpp
@@ -4,8 +4,6 @@
 #include "MassEntityTypes.h"
 #include "Misc/StringBuilder.h"
 
-PRAGMA_DISABLE_OPTIMIZATION
-
 //////////////////////////////////////////////////////////////////////
 // FMassArchetypeData
 
@@ -523,7 +521,7 @@ void FMassArchetypeData::CompactEntities(const double TimeAllowed)
 	}
 }
 
-void FMassArchetypeData::GetRequirementsFragmentMapping(TConstArrayView<FMassFragmentRequirement> Requirements, FMassFragmentIndicesMapping& OutFragmentIndices)
+void FMassArchetypeData::GetRequirementsFragmentMapping(TConstArrayView<FMassFragmentRequirement> Requirements, FMassFragmentIndicesMapping& OutFragmentIndices) const
 {
 	OutFragmentIndices.Reset(Requirements.Num());
 	for (const FMassFragmentRequirement& Requirement : Requirements)
@@ -537,7 +535,7 @@ void FMassArchetypeData::GetRequirementsFragmentMapping(TConstArrayView<FMassFra
 	}
 }
 
-void FMassArchetypeData::GetRequirementsChunkFragmentMapping(TConstArrayView<FMassFragmentRequirement> ChunkRequirements, FMassFragmentIndicesMapping& OutFragmentIndices)
+void FMassArchetypeData::GetRequirementsChunkFragmentMapping(TConstArrayView<FMassFragmentRequirement> ChunkRequirements, FMassFragmentIndicesMapping& OutFragmentIndices) const
 {
 	int32 LastFoundFragmentIndex = -1;
 	OutFragmentIndices.Reset(ChunkRequirements.Num());
@@ -563,7 +561,7 @@ void FMassArchetypeData::GetRequirementsChunkFragmentMapping(TConstArrayView<FMa
 	}
 }
 
-void FMassArchetypeData::GetRequirementsConstSharedFragmentMapping(TConstArrayView<FMassFragmentRequirement> Requirements, FMassFragmentIndicesMapping& OutFragmentIndices)
+void FMassArchetypeData::GetRequirementsConstSharedFragmentMapping(TConstArrayView<FMassFragmentRequirement> Requirements, FMassFragmentIndicesMapping& OutFragmentIndices) const
 {
 	OutFragmentIndices.Reset(Requirements.Num());
 	for (const FMassFragmentRequirement& Requirement : Requirements)
@@ -577,7 +575,7 @@ void FMassArchetypeData::GetRequirementsConstSharedFragmentMapping(TConstArrayVi
 	}
 }
 
-void FMassArchetypeData::GetRequirementsSharedFragmentMapping(TConstArrayView<FMassFragmentRequirement> Requirements, FMassFragmentIndicesMapping& OutFragmentIndices)
+void FMassArchetypeData::GetRequirementsSharedFragmentMapping(TConstArrayView<FMassFragmentRequirement> Requirements, FMassFragmentIndicesMapping& OutFragmentIndices) const
 {
 	OutFragmentIndices.Reset(Requirements.Num());
 	for (const FMassFragmentRequirement& Requirement : Requirements)
@@ -1108,4 +1106,3 @@ void FMassArchetypeData::BatchSetFragmentValues(TConstArrayView<FMassArchetypeEn
 	}
 }
 
-PRAGMA_ENABLE_OPTIMIZATION
diff --git a/Engine/Plugins/Runtime/MassEntity/Source/MassEntity/Private/MassArchetypeData.h b/Engine/Plugins/Runtime/MassEntity/Source/MassEntity/Private/MassArchetypeData.h
--- a/Engine/Plugins/Runtime/MassEntity/Source/MassEntity/Private/MassArchetypeData.h
+++ b/Engine/Plugins/Runtime/MassEntity/Source/MassEntity/Private/MassArchetypeData.h
@@ -243,16 +243,16 @@ public:
 	void SetFragmentData(FMassArchetypeEntityCollection::FConstEntityRangeArrayView EntityRangeContainer, const FInstancedStruct& FragmentSource);
 
 	/** Returns conversion from given Requirements to archetype's fragment indices */
-	void GetRequirementsFragmentMapping(TConstArrayView<FMassFragmentRequirement> Requirements, FMassFragmentIndicesMapping& OutFragmentIndices);
+	void GetRequirementsFragmentMapping(TConstArrayView<FMassFragmentRequirement> Requirements, FMassFragmentIndicesMapping& OutFragmentIndices) const;
 
 	/** Returns conversion from given ChunkRequirements to archetype's chunk fragment indices */
-	void GetRequirementsChunkFragmentMapping(TConstArrayView<FMassFragmentRequirement> ChunkRequirements, FMassFragmentIndicesMapping& OutFragmentIndices);
+	void GetRequirementsChunkFragmentMapping(TConstArrayView<FMassFragmentRequirement> ChunkRequirements, FMassFragmentIndicesMapping& OutFragmentIndices) const;
 
 	/** Returns conversion from given const shared requirements to archetype's const shared fragment indices */
-	void GetRequirementsConstSharedFragmentMapping(TConstArrayView<FMassFragmentRequirement> Requirements, FMassFragmentIndicesMapping& OutFragmentIndices);
+	void GetRequirementsConstSharedFragmentMapping(TConstArrayView<FMassFragmentRequirement> Requirements, FMassFragmentIndicesMapping& OutFragmentIndices) const;
 
 	/** Returns conversion from given shared requirements to archetype's shared fragment indices */
-	void GetRequirementsSharedFragmentMapping(TConstArrayView<FMassFragmentRequirement> Requirements, FMassFragmentIndicesMapping& OutFragmentIndices);
+	void GetRequirementsSharedFragmentMapping(TConstArrayView<FMassFragmentRequirement> Requirements, FMassFragmentIndicesMapping& OutFragmentIndices) const;
 
 	SIZE_T GetAllocatedSize() const;
 
