Improved mass archetype hashing which in turns improves performance of finding archetypes.
And constified UMassEntitySubsystem::GetValidArchetypes while at it.

diff --git a/Engine/Plugins/Runtime/MassEntity/Source/MassEntity/Private/MassEntitySubsystem.cpp b/Engine/Plugins/Runtime/MassEntity/Source/MassEntity/Private/MassEntitySubsystem.cpp
--- a/Engine/Plugins/Runtime/MassEntity/Source/MassEntity/Private/MassEntitySubsystem.cpp
+++ b/Engine/Plugins/Runtime/MassEntity/Source/MassEntity/Private/MassEntitySubsystem.cpp
@@ -169,7 +169,7 @@ FMassArchetypeHandle UMassEntitySubsystem::CreateArchetype(const TSharedPtr<FMas
 
 FMassArchetypeHandle UMassEntitySubsystem::CreateArchetype(const FMassArchetypeCompositionDescriptor& Composition, const FMassArchetypeSharedFragmentValues& SharedFragmentValues)
 {
-	const uint32 TypeHash = Composition.CalculateHash();
+	const uint32 TypeHash = HashCombine(Composition.CalculateHash(), GetTypeHash(SharedFragmentValues));
 
 	TArray<TSharedPtr<FMassArchetypeData>>& HashRow = FragmentHashToArchetypeMap.FindOrAdd(TypeHash);
 
@@ -222,8 +222,7 @@ FMassArchetypeHandle UMassEntitySubsystem::InternalCreateSimilarArchetype(const
 
 FMassArchetypeHandle UMassEntitySubsystem::InternalCreateSimilarArchetype(const FMassArchetypeData& SourceArchetypeRef, FMassArchetypeCompositionDescriptor&& NewComposition)
 {
-	//const FMassArchetypeCompositionDescriptor NewComposition(SourceArchetypeRef.GetFragmentBitSet(), OverrideTags, SourceArchetypeRef.GetChunkFragmentBitSet(), SourceArchetypeRef.GetSharedFragmentBitSet());
-	const uint32 TypeHash = NewComposition.CalculateHash();
+	const uint32 TypeHash = HashCombine(NewComposition.CalculateHash(), GetTypeHash(SourceArchetypeRef.GetSharedFragmentValues()));
 
 	TArray<TSharedPtr<FMassArchetypeData>>& HashRow = FragmentHashToArchetypeMap.FindOrAdd(TypeHash);
 
@@ -1093,7 +1092,7 @@ void UMassEntitySubsystem::CheckIfEntityIsActive(FMassEntityHandle Entity) const
 	checkf(IsEntityBuilt(Entity), TEXT("Entity not yet created(ID: %d, SN:%d)"));
 }
 
-void UMassEntitySubsystem::GetValidArchetypes(const FMassEntityQuery& Query, TArray<FMassArchetypeHandle>& OutValidArchetypes)
+void UMassEntitySubsystem::GetValidArchetypes(const FMassEntityQuery& Query, TArray<FMassArchetypeHandle>& OutValidArchetypes) const
 {
 	//@TODO: Not optimized yet, but we call this rarely now, so not a big deal.
 
@@ -1104,7 +1103,7 @@ void UMassEntitySubsystem::GetValidArchetypes(const FMassEntityQuery& Query, TAr
 		check(Requirement.StructType);
 		if (Requirement.Presence != EMassFragmentPresence::None)
 		{
-			if (TArray<TSharedPtr<FMassArchetypeData>>* pData = FragmentTypeToArchetypeMap.Find(Requirement.StructType))
+			if (const TArray<TSharedPtr<FMassArchetypeData>>* pData = FragmentTypeToArchetypeMap.Find(Requirement.StructType))
 			{
 				AnyArchetypes.Append(*pData);
 			}
diff --git a/Engine/Plugins/Runtime/MassEntity/Source/MassEntity/Public/MassEntitySubsystem.h b/Engine/Plugins/Runtime/MassEntity/Source/MassEntity/Public/MassEntitySubsystem.h
--- a/Engine/Plugins/Runtime/MassEntity/Source/MassEntity/Public/MassEntitySubsystem.h
+++ b/Engine/Plugins/Runtime/MassEntity/Source/MassEntity/Public/MassEntitySubsystem.h
@@ -362,7 +362,7 @@ public:
 #endif // WITH_MASSENTITY_DEBUG
 
 protected:
-	void GetValidArchetypes(const FMassEntityQuery& Query, TArray<FMassArchetypeHandle>& OutValidArchetypes);
+	void GetValidArchetypes(const FMassEntityQuery& Query, TArray<FMassArchetypeHandle>& OutValidArchetypes) const;
 	
 	/** 
 	 * A "similar" archetype is an archetype exactly the same as SourceArchetype except for one composition aspect 
