Stub of the MassDebugger UI, plugged into the Editor and supplied with the `mass.debug` console command to spawn it in `-game`

diff --git a/Engine/Plugins/Runtime/MassEntity/MassEntity.uplugin b/Engine/Plugins/Runtime/MassEntity/MassEntity.uplugin
--- a/Engine/Plugins/Runtime/MassEntity/MassEntity.uplugin
+++ b/Engine/Plugins/Runtime/MassEntity/MassEntity.uplugin
@@ -29,6 +29,10 @@
 		{
 			"Name": "MassEntityEditor",
 			"Type": "Editor"
+		},
+		{
+			"Name": "MassEntityDebugger",
+			"Type": "DeveloperTool"
 		}
 	],
 	"Plugins": [
diff --git a/Engine/Plugins/Runtime/MassEntity/Source/MassEntityDebugger/MassEntityDebugger.Build.cs b/Engine/Plugins/Runtime/MassEntity/Source/MassEntityDebugger/MassEntityDebugger.Build.cs
new file mode 100644
--- /dev/null
+++ b/Engine/Plugins/Runtime/MassEntity/Source/MassEntityDebugger/MassEntityDebugger.Build.cs
@@ -0,0 +1,34 @@
+// Copyright Epic Games, Inc. All Rights Reserved.
+
+using UnrealBuildTool;
+
+namespace UnrealBuildTool.Rules
+{
+	public class MassEntityDebugger : ModuleRules
+	{
+		public MassEntityDebugger(ReadOnlyTargetRules Target) : base(Target)
+		{
+			PCHUsage = ModuleRules.PCHUsageMode.UseExplicitOrSharedPCHs;
+
+			PublicIncludePaths.AddRange(
+				new string[] {
+					ModuleDirectory + "/Public"
+				}
+			);
+
+			PublicDependencyModuleNames.AddRange(
+				new string[] {
+					"Core",
+					"CoreUObject",
+					"Engine",
+					"StructUtils",
+					"MassEntity",
+					"SlateCore",
+					"Slate",
+					"WorkspaceMenuStructure",
+					"Projects"
+				}
+			);
+		}
+	}
+}
\ No newline at end of file
diff --git a/Engine/Plugins/Runtime/MassEntity/Source/MassEntityDebugger/Private/MassDebuggerModule.cpp b/Engine/Plugins/Runtime/MassEntity/Source/MassEntityDebugger/Private/MassDebuggerModule.cpp
new file mode 100644
--- /dev/null
+++ b/Engine/Plugins/Runtime/MassEntity/Source/MassEntityDebugger/Private/MassDebuggerModule.cpp
@@ -0,0 +1,76 @@
+// Copyright Epic Games, Inc. All Rights Reserved.
+
+#include "MassProcessingTypes.h"
+
+#if WITH_MASSENTITY_DEBUG
+
+#include "IMassDebuggerModule.h"
+
+#include "WorkspaceMenuStructure.h"
+#include "WorkspaceMenuStructureModule.h"
+
+#include "SMassDebugger.h"
+#include "MassDebuggerStyle.h"
+
+#define LOCTEXT_NAMESPACE "Mass"
+
+static const FName MassDebugTabName("MassDebugger");
+
+class FMassDebuggerModule : public IMassDebuggerModule, public IModularFeature
+{
+	/** IModuleInterface implementation */
+	virtual void StartupModule() override;
+
+private:
+	TSharedRef<SDockTab> SpawnTab(const FSpawnTabArgs& SpawnTabArgs);
+};
+
+IMPLEMENT_MODULE(FMassDebuggerModule, MassEntityDebugger)
+
+
+void FMassDebuggerModule::StartupModule()
+{
+	// This code will execute after your module is loaded into memory (but after global variables are initialized, of course.)
+	FMassDebuggerStyle::Initialize();
+
+	// register the UI tab
+	IModularFeatures::Get().RegisterModularFeature(MassDebugTabName, this);
+
+	FGlobalTabmanager::Get()->RegisterNomadTabSpawner(
+		MassDebugTabName,
+		FOnSpawnTab::CreateRaw(this, &FMassDebuggerModule::SpawnTab))
+		.SetGroup(WorkspaceMenu::GetMenuStructure().GetDeveloperToolsDebugCategory())
+		.SetDisplayName(NSLOCTEXT("MassDebuggerApp", "TabTitle", "Mass Debugger"))
+		.SetTooltipText(NSLOCTEXT("MassDebuggerApp", "TooltipText", "Opens Mass Debugger tool."))
+		.SetIcon(FSlateIcon(FMassDebuggerStyle::GetStyleSetName(), "MassDebuggerApp.TabIcon"));
+}
+
+TSharedRef<SDockTab> FMassDebuggerModule::SpawnTab(const FSpawnTabArgs& SpawnTabArgs)
+{
+	const TSharedRef<SDockTab> MajorTab = SNew(SMassDebuggerTab)
+		.TabRole(ETabRole::NomadTab);
+
+	TSharedPtr<SWidget> TabContent;
+
+	TabContent = SNew(SMassDebugger, MajorTab, SpawnTabArgs.GetOwnerWindow());
+
+	MajorTab->SetContent(TabContent.ToSharedRef());
+
+	return MajorTab;
+}
+
+namespace UE::Mass::Debug::Private
+{
+	FAutoConsoleCommand SummonDebbugerCommand(
+		TEXT("mass.debug"),
+		TEXT("Summon Mass Debugger UI"),
+		FConsoleCommandDelegate::CreateLambda([]() 
+		{ 
+			FGlobalTabmanager::Get()->TryInvokeTab(MassDebugTabName);
+		})
+	);
+
+} // UE::Mass::Debug::Private
+#endif // WITH_MASSENTITY_DEBUG
+
+#undef LOCTEXT_NAMESPACE
diff --git a/Engine/Plugins/Runtime/MassEntity/Source/MassEntityDebugger/Private/MassDebuggerStyle.cpp b/Engine/Plugins/Runtime/MassEntity/Source/MassEntityDebugger/Private/MassDebuggerStyle.cpp
new file mode 100644
--- /dev/null
+++ b/Engine/Plugins/Runtime/MassEntity/Source/MassEntityDebugger/Private/MassDebuggerStyle.cpp
@@ -0,0 +1,52 @@
+// Copyright Epic Games, Inc. All Rights Reserved.
+
+#include "MassDebuggerStyle.h"
+#include "Styling/SlateStyleRegistry.h"
+#include "Interfaces/IPluginManager.h"
+#include "Styling/SlateStyle.h"
+#include "Styling/SlateStyleMacros.h"
+
+
+TSharedPtr<FSlateStyleSet> FMassDebuggerStyle::StyleSet = nullptr;
+
+FString FMassDebuggerStyle::InContent(const FString& RelativePath, const ANSICHAR* Extension)
+{
+	static FString ContentDir = IPluginManager::Get().FindPlugin(TEXT("MassEntity"))->GetContentDir() / TEXT("Slate");
+	return (ContentDir / RelativePath) + Extension;
+}
+
+void FMassDebuggerStyle::Initialize()
+{
+	if (StyleSet.IsValid())
+	{
+		return;
+	}
+
+	const FVector2D Icon16x16(16.0f, 16.0f);
+
+	StyleSet = MakeShared<FSlateStyleSet>(GetStyleSetName());
+
+	StyleSet->SetContentRoot(FPaths::EngineContentDir() / TEXT("Editor/Slate"));
+	StyleSet->SetCoreContentRoot(FPaths::EngineContentDir() / TEXT("Slate"));
+
+	FSlateStyleSet& Style = *StyleSet.Get();
+	StyleSet->Set("MassDebuggerApp.TabIcon", new FSlateVectorImageBrush(InContent("MassDebugger", ".svg"), Icon16x16));
+	
+	FSlateStyleRegistry::RegisterSlateStyle(*StyleSet.Get());
+}
+
+void FMassDebuggerStyle::Shutdown()
+{
+	if (StyleSet.IsValid())
+	{
+		FSlateStyleRegistry::UnRegisterSlateStyle(*StyleSet.Get());
+		ensure(StyleSet.IsUnique());
+		StyleSet.Reset();
+	}
+}
+
+FName FMassDebuggerStyle::GetStyleSetName()
+{
+	static FName StyleName("MassDebuggerStyle");
+	return StyleName;
+}
diff --git a/Engine/Plugins/Runtime/MassEntity/Source/MassEntityDebugger/Private/MassDebuggerStyle.h b/Engine/Plugins/Runtime/MassEntity/Source/MassEntityDebugger/Private/MassDebuggerStyle.h
new file mode 100644
--- /dev/null
+++ b/Engine/Plugins/Runtime/MassEntity/Source/MassEntityDebugger/Private/MassDebuggerStyle.h
@@ -0,0 +1,25 @@
+// Copyright Epic Games, Inc. All Rights Reserved.
+
+#pragma once
+
+#include "Styling/SlateStyle.h"
+
+
+class MASSENTITYDEBUGGER_API FMassDebuggerStyle
+{
+public:
+	static void Initialize();
+	static void Shutdown();
+	static TSharedPtr<ISlateStyle> Get() { return StyleSet; }
+	static FName GetStyleSetName();
+
+	static const FSlateBrush* GetBrush(FName PropertyName, const ANSICHAR* Specifier = NULL)
+	{
+		return StyleSet->GetBrush(PropertyName, Specifier);
+	}
+
+	static FString InContent(const FString& RelativePath, const ANSICHAR* Extension);
+
+private:
+	static TSharedPtr<FSlateStyleSet> StyleSet;
+};
diff --git a/Engine/Plugins/Runtime/MassEntity/Source/MassEntityDebugger/Private/SMassDebugger.cpp b/Engine/Plugins/Runtime/MassEntity/Source/MassEntityDebugger/Private/SMassDebugger.cpp
new file mode 100644
--- /dev/null
+++ b/Engine/Plugins/Runtime/MassEntity/Source/MassEntityDebugger/Private/SMassDebugger.cpp
@@ -0,0 +1,12 @@
+// Copyright Epic Games, Inc. All Rights Reserved.
+
+#include "SMassDebugger.h"
+
+#define LOCTEXT_NAMESPACE "SMassDebugger"
+
+void SMassDebugger::Construct(const FArguments& InArgs, const TSharedRef<SDockTab>& ConstructUnderMajorTab, const TSharedPtr<SWindow>& ConstructUnderWindow)
+{
+}
+	
+#undef LOCTEXT_NAMESPACE
+
diff --git a/Engine/Plugins/Runtime/MassEntity/Source/MassEntityDebugger/Private/SMassDebugger.h b/Engine/Plugins/Runtime/MassEntity/Source/MassEntityDebugger/Private/SMassDebugger.h
new file mode 100644
--- /dev/null
+++ b/Engine/Plugins/Runtime/MassEntity/Source/MassEntityDebugger/Private/SMassDebugger.h
@@ -0,0 +1,34 @@
+// Copyright Epic Games, Inc. All Rights Reserved.
+#pragma once
+
+#include "Widgets/Docking/SDockTab.h"
+
+
+class SMassDebuggerTab : public SDockTab
+{
+public:
+	virtual bool SupportsKeyboardFocus() const override { return true; }
+};
+
+class SMassDebugger : public SCompoundWidget
+{
+public:
+
+	SLATE_BEGIN_ARGS(SMassDebugger) 
+	{
+	}
+	SLATE_END_ARGS()
+
+public:
+
+	/**
+	* Constructs the application.
+	*
+	* @param InArgs The Slate argument list.
+	* @param ConstructUnderMajorTab The major tab which will contain the session front-end.
+	* @param ConstructUnderWindow The window in which this widget is being constructed.
+	*/
+	void Construct(const FArguments& InArgs, const TSharedRef<SDockTab>& ConstructUnderMajorTab, const TSharedPtr<SWindow>& ConstructUnderWindow);
+
+	virtual bool SupportsKeyboardFocus() const override { return true; }
+};
diff --git a/Engine/Plugins/Runtime/MassEntity/Source/MassEntityDebugger/Public/IMassDebuggerModule.h b/Engine/Plugins/Runtime/MassEntity/Source/MassEntityDebugger/Public/IMassDebuggerModule.h
new file mode 100644
--- /dev/null
+++ b/Engine/Plugins/Runtime/MassEntity/Source/MassEntityDebugger/Public/IMassDebuggerModule.h
@@ -0,0 +1,36 @@
+// Copyright Epic Games, Inc. All Rights Reserved.
+
+#pragma once
+
+#include "Modules/ModuleManager.h"
+
+/**
+* The public interface to this module
+*/
+class IMassDebuggerModule : public IModuleInterface
+{
+
+public:
+
+	/**
+	* Singleton-like access to this module's interface.  This is just for convenience!
+	* Beware of calling this during the shutdown phase, though.  Your module might have been unloaded already.
+	*
+	* @return Returns singleton instance, loading the module on demand if needed
+	*/
+	static inline IMassDebuggerModule& Get()
+	{
+		return FModuleManager::LoadModuleChecked<IMassDebuggerModule>("MassEntityDebugger");
+	}
+
+	/**
+	* Checks to see if this module is loaded and ready.  It is only valid to call Get() if IsAvailable() returns true.
+	*
+	* @return True if the module is loaded and ready to use
+	*/
+	static inline bool IsAvailable()
+	{
+		return FModuleManager::Get().IsModuleLoaded("MassEntityDebugger");
+	}
+};
+
diff --git a/Engine/Plugins/Runtime/MassEntity/Source/MassEntityDebugger/Public/MassDebuggerSettings.h b/Engine/Plugins/Runtime/MassEntity/Source/MassEntityDebugger/Public/MassDebuggerSettings.h
new file mode 100644
--- /dev/null
+++ b/Engine/Plugins/Runtime/MassEntity/Source/MassEntityDebugger/Public/MassDebuggerSettings.h
@@ -0,0 +1,15 @@
+// Copyright Epic Games, Inc. All Rights Reserved.
+
+#pragma once
+
+#include "MassSettings.h"
+#include "MassDebuggerSettings.generated.h"
+
+
+UCLASS(config = Mass, defaultconfig, DisplayName = "Mass Debugger")
+class MASSENTITYDEBUGGER_API UMassDebuggerSettings : public UMassModuleSettings
+{
+	GENERATED_BODY()
+
+public:
+};
