Mass-side changes for MassDebugger.
Intorduced a FMassDebugger struct to host all debug-time API and avoid polluting other Mass classes with those. Existing functions will be moved over with time.

diff --git a/Engine/Plugins/Runtime/MassEntity/Source/MassEntity/Private/MassArchetypeData.h b/Engine/Plugins/Runtime/MassEntity/Source/MassEntity/Private/MassArchetypeData.h
--- a/Engine/Plugins/Runtime/MassEntity/Source/MassEntity/Private/MassArchetypeData.h
+++ b/Engine/Plugins/Runtime/MassEntity/Source/MassEntity/Private/MassArchetypeData.h
@@ -169,6 +169,7 @@ private:
 
 	friend FMassEntityQuery;
 	friend FMassArchetypeEntityCollection;
+	friend FMassDebugger;
 
 public:
 	TConstArrayView<FMassArchetypeFragmentConfig> GetFragmentConfigs() const { return FragmentConfigs; }
diff --git a/Engine/Plugins/Runtime/MassEntity/Source/MassEntity/Private/MassDebugger.cpp b/Engine/Plugins/Runtime/MassEntity/Source/MassEntity/Private/MassDebugger.cpp
new file mode 100644
--- /dev/null
+++ b/Engine/Plugins/Runtime/MassEntity/Source/MassEntity/Private/MassDebugger.cpp
@@ -0,0 +1,77 @@
+// Copyright Epic Games, Inc. All Rights Reserved.
+
+#include "MassDebugger.h"
+#if WITH_MASSENTITY_DEBUG
+#include "MassProcessor.h"
+#include "MassEntitySubsystem.h"
+#include "MassArchetypeTypes.h"
+#include "MassArchetypeData.h"
+
+
+DEFINE_ENUM_TO_STRING(EMassFragmentAccess, "/Script/MassEntity");
+DEFINE_ENUM_TO_STRING(EMassFragmentPresence, "/Script/MassEntity");
+
+//----------------------------------------------------------------------//
+// FMassDebugger
+//----------------------------------------------------------------------//
+TConstArrayView<FMassEntityQuery*> FMassDebugger::GetProcessorQueries(const UMassProcessor& Processor)
+{
+	return Processor.OwnedQueries;
+}
+
+UE::Mass::Debug::FQueryRequirementsView FMassDebugger::GetQueryRequirements(const FMassEntityQuery& Query)
+{
+	UE::Mass::Debug::FQueryRequirementsView View = { Query.Requirements, Query.ChunkRequirements, Query.ConstSharedRequirements, Query.SharedRequirements
+		, Query.RequiredAllTags, Query.RequiredAnyTags, Query.RequiredNoneTags
+		, Query.RequiredConstSubsystems, Query.RequiredMutableSubsystems };
+
+	return View;
+}
+
+void FMassDebugger::GetQueryExecutionRequirements(const FMassEntityQuery& Query, FMassExecutionRequirements& OutExecutionRequirements)
+{
+	Query.ExportRequirements(OutExecutionRequirements);
+}
+
+TArray<FMassArchetypeHandle> FMassDebugger::GetAllArchetypes(const UMassEntitySubsystem& EntitySubsystem)
+{
+	TArray<FMassArchetypeHandle> Archetypes;
+
+	for (auto& KVP : EntitySubsystem.FragmentHashToArchetypeMap)
+	{
+		for (const TSharedPtr<FMassArchetypeData>& Archetype : KVP.Value)
+		{
+			Archetypes.Add(FMassArchetypeHandle(Archetype));
+		}
+	}
+
+	return Archetypes;
+}
+
+const FMassArchetypeCompositionDescriptor& FMassDebugger::GetArchetypeComposition(const FMassArchetypeHandle& ArchetypeHandle)
+{
+	check(ArchetypeHandle.IsValid());
+	return ArchetypeHandle.DataPtr->CompositionDescriptor;
+}
+
+const FMassArchetypeSharedFragmentValues& FMassDebugger::GetArchetypeSharedFragmentValues(const FMassArchetypeHandle& ArchetypeHandle)
+{
+	check(ArchetypeHandle.IsValid());
+	return ArchetypeHandle.DataPtr->SharedFragmentValues;
+}
+
+void FMassDebugger::GetArchetypeEntityStats(const FMassArchetypeHandle& ArchetypeHandle
+	, int32& OutEntitiesCount, int32& OutNumEntitiesPerChunk, int32& OutChunksCount)
+{
+	check(ArchetypeHandle.IsValid());
+	OutEntitiesCount = ArchetypeHandle.DataPtr->GetNumEntities();
+	OutNumEntitiesPerChunk = ArchetypeHandle.DataPtr->GetNumEntitiesPerChunk();
+	OutChunksCount = ArchetypeHandle.DataPtr->GetChunkCount();
+}
+
+TConstArrayView<struct UMassCompositeProcessor::FDependencyNode> FMassDebugger::GetProcessingGraph(const UMassCompositeProcessor& GraphOwner)
+{
+	return GraphOwner.ProcessingFlatGraph;
+}
+
+#endif // WITH_MASSENTITY_DEBUG
\ No newline at end of file
diff --git a/Engine/Plugins/Runtime/MassEntity/Source/MassEntity/Private/MassEntityQuery.cpp b/Engine/Plugins/Runtime/MassEntity/Source/MassEntity/Private/MassEntityQuery.cpp
--- a/Engine/Plugins/Runtime/MassEntity/Source/MassEntity/Private/MassEntityQuery.cpp
+++ b/Engine/Plugins/Runtime/MassEntity/Source/MassEntity/Private/MassEntityQuery.cpp
@@ -358,6 +358,10 @@ void FMassEntityQuery::ExportRequirements(FMassExecutionRequirements& OutRequire
 	ExportRequirements<FMassChunkFragmentBitSet>(ChunkRequirements, OutRequirements.ChunkFragments);
 	ExportRequirements<FMassSharedFragmentBitSet>(ConstSharedRequirements, OutRequirements.SharedFragments);
 	ExportRequirements<FMassSharedFragmentBitSet>(SharedRequirements, OutRequirements.SharedFragments);
+
+	OutRequirements.RequiredAllTags = RequiredAllTags;
+	OutRequirements.RequiredAnyTags = RequiredAnyTags;
+	OutRequirements.RequiredNoneTags = RequiredNoneTags;
 }
 
 FString FMassEntityQuery::DebugGetDescription() const
diff --git a/Engine/Plugins/Runtime/MassEntity/Source/MassEntity/Private/MassProcessorDependencySolver.cpp b/Engine/Plugins/Runtime/MassEntity/Source/MassEntity/Private/MassProcessorDependencySolver.cpp
--- a/Engine/Plugins/Runtime/MassEntity/Source/MassEntity/Private/MassProcessorDependencySolver.cpp
+++ b/Engine/Plugins/Runtime/MassEntity/Source/MassEntity/Private/MassProcessorDependencySolver.cpp
@@ -30,6 +30,23 @@ namespace UE::Mass::Private
 //----------------------------------------------------------------------//
 //  FMassExecutionRequirements
 //----------------------------------------------------------------------//
+void FMassExecutionRequirements::Append(const FMassExecutionRequirements& Other)
+{
+	for (int i = 0; i < EMassAccessOperation::MAX; ++i)
+	{
+		Fragments[i] += Other.Fragments[i];
+		ChunkFragments[i] += Other.ChunkFragments[i];
+		SharedFragments[i] += Other.SharedFragments[i];
+		RequiredSubsystems[i] += Other.RequiredSubsystems[i];
+	}
+	RequiredAllTags += Other.RequiredAllTags;
+	RequiredAnyTags += Other.RequiredAnyTags;
+	RequiredNoneTags += Other.RequiredNoneTags;
+
+	// signal that it requires recalculation;
+	ResourcesUsedCount = INDEX_NONE;
+}
+
 void FMassExecutionRequirements::CountResourcesUsed()
 {
 	ResourcesUsedCount = 0;
@@ -134,7 +151,7 @@ void FProcessorDependencySolver::FResourceUsage::SubmitNode(const int32 NodeInde
 	HandleElementType<FMassSharedFragmentBitSet>(SharedFragmentsAccess, InOutNode.Requirements.SharedFragments, InOutNode, NodeIndex);
 	HandleElementType<FMassExternalSubystemBitSet>(RequiredSubsystemsAccess, InOutNode.Requirements.RequiredSubsystems, InOutNode, NodeIndex);
 
-	Requirements += InOutNode.Requirements;
+	Requirements.Append(InOutNode.Requirements);
 }
 
 //----------------------------------------------------------------------//
diff --git a/Engine/Plugins/Runtime/MassEntity/Source/MassEntity/Public/MassArchetypeTypes.h b/Engine/Plugins/Runtime/MassEntity/Source/MassEntity/Public/MassArchetypeTypes.h
--- a/Engine/Plugins/Runtime/MassEntity/Source/MassEntity/Public/MassArchetypeTypes.h
+++ b/Engine/Plugins/Runtime/MassEntity/Source/MassEntity/Public/MassArchetypeTypes.h
@@ -16,6 +16,7 @@ struct FMassArchetypeChunkIterator;
 struct FMassEntityQuery;
 struct FMassArchetypeEntityCollection;
 struct FMassEntityView;
+struct FMassDebugger;
 
 typedef TFunction< void(FMassExecutionContext& /*ExecutionContext*/) > FMassExecuteFunction;
 typedef TFunction< bool(const FMassExecutionContext& /*ExecutionContext*/) > FMassChunkConditionFunction;
@@ -45,6 +46,7 @@ private:
 	friend FMassArchetypeEntityCollection;
 	friend FMassEntityQuery;
 	friend FMassEntityView;
+	friend FMassDebugger;
 };
 
 
diff --git a/Engine/Plugins/Runtime/MassEntity/Source/MassEntity/Public/MassDebugger.h b/Engine/Plugins/Runtime/MassEntity/Source/MassEntity/Public/MassDebugger.h
new file mode 100644
--- /dev/null
+++ b/Engine/Plugins/Runtime/MassEntity/Source/MassEntity/Public/MassDebugger.h
@@ -0,0 +1,54 @@
+// Copyright Epic Games, Inc. All Rights Reserved.
+
+#pragma once
+
+#include "MassProcessingTypes.h"
+
+#if WITH_MASSENTITY_DEBUG
+#include "Containers/ContainersFwd.h"
+#include "VisualLogger/VisualLoggerTypes.h"
+#include "MassEntityQuery.h"
+#include "MassProcessor.h"
+
+class UMassProcessor;
+struct FMassEntityQuery;
+class UMassEntitySubsystem;
+struct FMassArchetypeHandle;
+enum class EMassFragmentAccess : uint8;
+enum class EMassFragmentPresence : uint8;
+
+MASSENTITY_API DECLARE_ENUM_TO_STRING(EMassFragmentAccess);
+MASSENTITY_API DECLARE_ENUM_TO_STRING(EMassFragmentPresence);
+
+namespace UE::Mass::Debug
+{
+struct MASSENTITY_API FQueryRequirementsView
+{
+	TConstArrayView<FMassFragmentRequirement> FragmentRequirements;
+	TConstArrayView<FMassFragmentRequirement> ChunkRequirements;
+	TConstArrayView<FMassFragmentRequirement> ConstSharedRequirements;
+	TConstArrayView<FMassFragmentRequirement> SharedRequirements;
+	const FMassTagBitSet& RequiredAllTags;
+	const FMassTagBitSet& RequiredAnyTags;
+	const FMassTagBitSet& RequiredNoneTags;
+	const FMassExternalSubystemBitSet& RequiredConstSubsystems;
+	const FMassExternalSubystemBitSet& RequiredMutableSubsystems;
+};
+} // namespace UE::Mass::Debug
+
+struct MASSENTITY_API FMassDebugger
+{
+	static TConstArrayView<FMassEntityQuery*> GetProcessorQueries(const UMassProcessor& Processor);
+	static UE::Mass::Debug::FQueryRequirementsView GetQueryRequirements(const FMassEntityQuery& Query);
+	static void GetQueryExecutionRequirements(const FMassEntityQuery& Query, FMassExecutionRequirements& OutExecutionRequirements);
+	
+	static TArray<FMassArchetypeHandle> GetAllArchetypes(const UMassEntitySubsystem& EntitySubsystem);
+	static const FMassArchetypeCompositionDescriptor& GetArchetypeComposition(const FMassArchetypeHandle& ArchetypeHandle);
+	static const FMassArchetypeSharedFragmentValues& GetArchetypeSharedFragmentValues(const FMassArchetypeHandle& ArchetypeHandle);
+
+	static void GetArchetypeEntityStats(const FMassArchetypeHandle& ArchetypeHandle, int32& OutEntitiesCount, int32& OutNumEntitiesPerChunk, int32& OutChunksCount);
+
+	static TConstArrayView<UMassCompositeProcessor::FDependencyNode> GetProcessingGraph(const UMassCompositeProcessor& GraphOwner);
+};
+
+#endif // WITH_MASSENTITY_DEBUG
\ No newline at end of file
diff --git a/Engine/Plugins/Runtime/MassEntity/Source/MassEntity/Public/MassEntityQuery.h b/Engine/Plugins/Runtime/MassEntity/Source/MassEntity/Public/MassEntityQuery.h
--- a/Engine/Plugins/Runtime/MassEntity/Source/MassEntity/Public/MassEntityQuery.h
+++ b/Engine/Plugins/Runtime/MassEntity/Source/MassEntity/Public/MassEntityQuery.h
@@ -20,6 +20,7 @@ struct FMassArchetypeHandle;
 struct FMassExecutionRequirements;
 struct FMassRequirementAccessDetector;
 
+UENUM()
 enum class EMassFragmentAccess : uint8
 {
 	// no binding required
@@ -34,6 +35,7 @@ enum class EMassFragmentAccess : uint8
 	MAX
 };
 
+UENUM()
 enum class EMassFragmentPresence : uint8
 {
 	// All of the required fragments must be present
@@ -101,6 +103,8 @@ struct MASSENTITY_API FMassEntityQuery
 {
 	GENERATED_BODY()
 
+	friend struct FMassDebugger;
+
 public:
 	FMassEntityQuery();
 	FMassEntityQuery(std::initializer_list<UScriptStruct*> InitList);
diff --git a/Engine/Plugins/Runtime/MassEntity/Source/MassEntity/Public/MassEntitySubsystem.h b/Engine/Plugins/Runtime/MassEntity/Source/MassEntity/Public/MassEntitySubsystem.h
--- a/Engine/Plugins/Runtime/MassEntity/Source/MassEntity/Public/MassEntitySubsystem.h
+++ b/Engine/Plugins/Runtime/MassEntity/Source/MassEntity/Public/MassEntitySubsystem.h
@@ -23,6 +23,7 @@ struct FMassArchetypeData;
 struct FMassCommandBuffer;
 struct FMassArchetypeEntityCollection;
 class FOutputDevice;
+struct FMassDebugger;
 enum class EMassFragmentAccess : uint8;
 
 //@TODO: Comment this
@@ -31,7 +32,9 @@ class MASSENTITY_API UMassEntitySubsystem : public UWorldSubsystem
 {
 	GENERATED_BODY()
 
-	friend struct FMassEntityQuery;
+	friend FMassEntityQuery;
+	friend FMassDebugger;
+
 private:
 	// Index 0 is reserved so we can treat that index as an invalid entity handle
 	constexpr static int32 NumReservedEntities = 1;
diff --git a/Engine/Plugins/Runtime/MassEntity/Source/MassEntity/Public/MassEntityTypes.h b/Engine/Plugins/Runtime/MassEntity/Source/MassEntity/Public/MassEntityTypes.h
--- a/Engine/Plugins/Runtime/MassEntity/Source/MassEntity/Public/MassEntityTypes.h
+++ b/Engine/Plugins/Runtime/MassEntity/Source/MassEntity/Public/MassEntityTypes.h
@@ -190,6 +190,11 @@ struct FMassArchetypeCompositionDescriptor
 		return CalculateHash(Fragments, Tags, ChunkFragments, SharedFragments);
 	}
 
+	int32 CountStoredTypes() const
+	{
+		return Fragments.CountStoredTypes() + Tags.CountStoredTypes() + ChunkFragments.CountStoredTypes() + SharedFragments.CountStoredTypes();
+	}
+
 	void DebugOutputDescription(FOutputDevice& Ar) const
 	{
 #if WITH_MASSENTITY_DEBUG
diff --git a/Engine/Plugins/Runtime/MassEntity/Source/MassEntity/Public/MassProcessor.h b/Engine/Plugins/Runtime/MassEntity/Source/MassEntity/Public/MassProcessor.h
--- a/Engine/Plugins/Runtime/MassEntity/Source/MassEntity/Public/MassProcessor.h
+++ b/Engine/Plugins/Runtime/MassEntity/Source/MassEntity/Public/MassProcessor.h
@@ -12,6 +12,8 @@
 
 
 struct FMassProcessingPhaseConfig;
+class UMassCompositeProcessor;
+struct FMassDebugger;
 
 enum class EProcessorCompletionStatus : uint8
 {
@@ -130,7 +132,8 @@ protected:
 	bool bCanShowUpInSettings = true;
 #endif // WITH_EDITORONLY_DATA
 
-	friend class UMassCompositeProcessor;
+	friend UMassCompositeProcessor;
+	friend FMassDebugger;
 
 private:
 	/** Stores processor's queries registered via RegisterQuery. 
@@ -145,6 +148,8 @@ class MASSENTITY_API UMassCompositeProcessor : public UMassProcessor
 {
 	GENERATED_BODY()
 
+	friend FMassDebugger;
+public:
 	struct FDependencyNode
 	{
 		FName Name;
diff --git a/Engine/Plugins/Runtime/MassEntity/Source/MassEntity/Public/MassProcessorDependencySolver.h b/Engine/Plugins/Runtime/MassEntity/Source/MassEntity/Public/MassProcessorDependencySolver.h
--- a/Engine/Plugins/Runtime/MassEntity/Source/MassEntity/Public/MassProcessorDependencySolver.h
+++ b/Engine/Plugins/Runtime/MassEntity/Source/MassEntity/Public/MassProcessorDependencySolver.h
@@ -25,7 +25,7 @@ namespace EMassAccessOperation
 };
 
 template<typename T>
-struct MASSENTITY_API TMassExecutionAccess
+struct TMassExecutionAccess
 {
 	T Read;
 	T Write;
@@ -41,28 +41,24 @@ struct MASSENTITY_API TMassExecutionAccess
 		check(OpIndex <= EMassAccessOperation::MAX);
 		return OpIndex == EMassAccessOperation::Read ? Read : Write;
 	}
+
+	TConstArrayView<T> AsArrayView() const { return MakeArrayView(&Read, 2); }
+
+	bool IsEmpty() const { return Read.IsEmpty() && Write.IsEmpty(); }
 };
 
 struct MASSENTITY_API FMassExecutionRequirements
 {
-	FMassExecutionRequirements& operator+=(const FMassExecutionRequirements& Other)
-	{
-		for (int i = 0; i < EMassAccessOperation::MAX; ++i)
-		{
-			Fragments[i] += Other.Fragments[i];
-			ChunkFragments[i] += Other.ChunkFragments[i];
-			SharedFragments[i] += Other.SharedFragments[i];
-			RequiredSubsystems[i] += Other.RequiredSubsystems[i];
-		}
-		return *this;
-	}
-
+	void Append(const FMassExecutionRequirements& Other);
 	void CountResourcesUsed();
 
 	TMassExecutionAccess<FMassFragmentBitSet> Fragments;
 	TMassExecutionAccess<FMassChunkFragmentBitSet> ChunkFragments;
 	TMassExecutionAccess<FMassSharedFragmentBitSet> SharedFragments;
 	TMassExecutionAccess<FMassExternalSubystemBitSet> RequiredSubsystems;
+	FMassTagBitSet RequiredAllTags;
+	FMassTagBitSet RequiredAnyTags;
+	FMassTagBitSet RequiredNoneTags;
 	int32 ResourcesUsedCount = INDEX_NONE;
 };
 
