More FMassDebugger functionality
Added a function that will fetch Processor's queries, and also make sure the queries have their archetypes cache up to date.
Also, constified bits of MassEntityQuery

diff --git a/Engine/Plugins/Runtime/MassEntity/Source/MassEntity/Private/MassDebugger.cpp b/Engine/Plugins/Runtime/MassEntity/Source/MassEntity/Private/MassDebugger.cpp
--- a/Engine/Plugins/Runtime/MassEntity/Source/MassEntity/Private/MassDebugger.cpp
+++ b/Engine/Plugins/Runtime/MassEntity/Source/MassEntity/Private/MassDebugger.cpp
@@ -6,6 +6,7 @@
 #include "MassEntitySubsystem.h"
 #include "MassArchetypeTypes.h"
 #include "MassArchetypeData.h"
+#include "MassEntityQuery.h"
 
 
 DEFINE_ENUM_TO_STRING(EMassFragmentAccess, "/Script/MassEntity");
@@ -19,6 +20,19 @@ TConstArrayView<FMassEntityQuery*> FMassDebugger::GetProcessorQueries(const UMas
 	return Processor.OwnedQueries;
 }
 
+TConstArrayView<FMassEntityQuery*> FMassDebugger::GetUpToDateProcessorQueries(const UMassEntitySubsystem& EntitySubsystem, UMassProcessor& Processor)
+{
+	for (FMassEntityQuery* Query : Processor.OwnedQueries)
+	{
+		if (Query)
+		{
+			Query->CacheArchetypes(EntitySubsystem);
+		}
+	}
+
+	return Processor.OwnedQueries;
+}
+
 UE::Mass::Debug::FQueryRequirementsView FMassDebugger::GetQueryRequirements(const FMassEntityQuery& Query)
 {
 	UE::Mass::Debug::FQueryRequirementsView View = { Query.Requirements, Query.ChunkRequirements, Query.ConstSharedRequirements, Query.SharedRequirements
diff --git a/Engine/Plugins/Runtime/MassEntity/Source/MassEntity/Private/MassEntityQuery.cpp b/Engine/Plugins/Runtime/MassEntity/Source/MassEntity/Private/MassEntityQuery.cpp
--- a/Engine/Plugins/Runtime/MassEntity/Source/MassEntity/Private/MassEntityQuery.cpp
+++ b/Engine/Plugins/Runtime/MassEntity/Source/MassEntity/Private/MassEntityQuery.cpp
@@ -104,7 +104,7 @@ void FMassEntityQuery::SortRequirements()
 	SharedRequirements.Sort(FScriptStructSortOperator());
 }
 
-void FMassEntityQuery::CacheArchetypes(UMassEntitySubsystem& InEntitySubsystem)
+void FMassEntityQuery::CacheArchetypes(const UMassEntitySubsystem& InEntitySubsystem)
 {
 	const uint32 InEntitySubsystemHash = PointerHash(&InEntitySubsystem);
 	if (EntitySubsystemHash != InEntitySubsystemHash || InEntitySubsystem.GetArchetypeDataVersion() != ArchetypeDataVersion)
diff --git a/Engine/Plugins/Runtime/MassEntity/Source/MassEntity/Public/MassDebugger.h b/Engine/Plugins/Runtime/MassEntity/Source/MassEntity/Public/MassDebugger.h
--- a/Engine/Plugins/Runtime/MassEntity/Source/MassEntity/Public/MassDebugger.h
+++ b/Engine/Plugins/Runtime/MassEntity/Source/MassEntity/Public/MassDebugger.h
@@ -39,6 +39,10 @@ struct MASSENTITY_API FQueryRequirementsView
 struct MASSENTITY_API FMassDebugger
 {
 	static TConstArrayView<FMassEntityQuery*> GetProcessorQueries(const UMassProcessor& Processor);
+	/** fetches all queries registered for given Processor. Note that in order to get up to date information
+	 *  FMassEntityQuery::CacheArchetypes will be called on each query */
+	static TConstArrayView<FMassEntityQuery*> GetUpToDateProcessorQueries(const UMassEntitySubsystem& EntitySubsystem, UMassProcessor& Processor);
+
 	static UE::Mass::Debug::FQueryRequirementsView GetQueryRequirements(const FMassEntityQuery& Query);
 	static void GetQueryExecutionRequirements(const FMassEntityQuery& Query, FMassExecutionRequirements& OutExecutionRequirements);
 	
diff --git a/Engine/Plugins/Runtime/MassEntity/Source/MassEntity/Public/MassEntityQuery.h b/Engine/Plugins/Runtime/MassEntity/Source/MassEntity/Public/MassEntityQuery.h
--- a/Engine/Plugins/Runtime/MassEntity/Source/MassEntity/Public/MassEntityQuery.h
+++ b/Engine/Plugins/Runtime/MassEntity/Source/MassEntity/Public/MassEntityQuery.h
@@ -127,7 +127,7 @@ public:
 	/** Will gather all archetypes from InEntitySubsystem matching this->Requirements.
 	 *  Note that no work will be done if the cached data is up to date (as tracked by EntitySubsystemHash and 
 	 *	ArchetypeDataVersion properties). */
-	void CacheArchetypes(UMassEntitySubsystem& InEntitySubsystem);
+	void CacheArchetypes(const UMassEntitySubsystem& InEntitySubsystem);
 
 	FMassEntityQuery& AddRequirement(const UScriptStruct* FragmentType, const EMassFragmentAccess AccessMode, const EMassFragmentPresence Presence = EMassFragmentPresence::All)
 	{
