Significant functionality added to MassDebugger UI tool

* SMassProcessorsView shows all processors available in currently selected world
* SMassProcessingView shows all processing phases and displays processors in a dependency-organized fashion
* SMassArchetypes lists all available archetypes
The views are synchronized meaning clicking archetype or a processing in one view will reflect in others marking related elements

diff --git a/Engine/Plugins/Runtime/MassEntity/Source/MassEntity/Private/MassProcessorDependencySolver.cpp b/Engine/Plugins/Runtime/MassEntity/Source/MassEntity/Private/MassProcessorDependencySolver.cpp
--- a/Engine/Plugins/Runtime/MassEntity/Source/MassEntity/Private/MassProcessorDependencySolver.cpp
+++ b/Engine/Plugins/Runtime/MassEntity/Source/MassEntity/Private/MassProcessorDependencySolver.cpp
@@ -60,6 +60,14 @@ void FMassExecutionRequirements::CountResourcesUsed()
 	}
 }
 
+int32 FMassExecutionRequirements::GetTotalBitsUsedCount()
+{
+	CountResourcesUsed();
+
+	return ResourcesUsedCount + RequiredAllTags.CountStoredTypes()
+		+ RequiredAnyTags.CountStoredTypes() + RequiredNoneTags.CountStoredTypes();
+}
+
 //----------------------------------------------------------------------//
 //  FProcessorDependencySolver::FResourceUsage
 //----------------------------------------------------------------------//
diff --git a/Engine/Plugins/Runtime/MassEntity/Source/MassEntity/Public/MassProcessorDependencySolver.h b/Engine/Plugins/Runtime/MassEntity/Source/MassEntity/Public/MassProcessorDependencySolver.h
--- a/Engine/Plugins/Runtime/MassEntity/Source/MassEntity/Public/MassProcessorDependencySolver.h
+++ b/Engine/Plugins/Runtime/MassEntity/Source/MassEntity/Public/MassProcessorDependencySolver.h
@@ -51,6 +51,7 @@ struct MASSENTITY_API FMassExecutionRequirements
 {
 	void Append(const FMassExecutionRequirements& Other);
 	void CountResourcesUsed();
+	int32 GetTotalBitsUsedCount();
 
 	TMassExecutionAccess<FMassFragmentBitSet> Fragments;
 	TMassExecutionAccess<FMassChunkFragmentBitSet> ChunkFragments;
diff --git a/Engine/Plugins/Runtime/MassEntity/Source/MassEntityDebugger/MassEntityDebugger.Build.cs b/Engine/Plugins/Runtime/MassEntity/Source/MassEntityDebugger/MassEntityDebugger.Build.cs
--- a/Engine/Plugins/Runtime/MassEntity/Source/MassEntityDebugger/MassEntityDebugger.Build.cs
+++ b/Engine/Plugins/Runtime/MassEntity/Source/MassEntityDebugger/MassEntityDebugger.Build.cs
@@ -26,9 +26,20 @@ namespace UnrealBuildTool.Rules
 					"SlateCore",
 					"Slate",
 					"WorkspaceMenuStructure",
-					"Projects"
+					"Projects",
+					"UMG"
 				}
 			);
+
+			if (Target.bBuildEditor == true)
+			{
+				PublicDependencyModuleNames.AddRange(
+					new string[] {
+						"InputCore",
+						"UnrealEd"
+					}
+				);
+			}
 		}
 	}
 }
\ No newline at end of file
diff --git a/Engine/Plugins/Runtime/MassEntity/Source/MassEntityDebugger/Private/MassDebuggerModel.cpp b/Engine/Plugins/Runtime/MassEntity/Source/MassEntityDebugger/Private/MassDebuggerModel.cpp
new file mode 100644
--- /dev/null
+++ b/Engine/Plugins/Runtime/MassEntity/Source/MassEntityDebugger/Private/MassDebuggerModel.cpp
@@ -0,0 +1,474 @@
+// Copyright Epic Games, Inc. All Rights Reserved.
+
+#include "MassDebuggerModel.h"
+#include "MassProcessor.h"
+#include "MassEntitySubsystem.h"
+#include "MassEntityQuery.h"
+#include "Engine/Engine.h"
+#include "MassDebugger.h"
+#include "MassDebuggerSettings.h"
+#include "UObject/UObjectIterator.h"
+#include "Containers/UnrealString.h"
+
+#define LOCTEXT_NAMESPACE "SMassDebugger"
+
+namespace UE::Mass::Debugger::Private
+{
+	template<typename TBitSet>
+	int32 BitSetDistance(const TBitSet& A, const TBitSet& B)
+	{
+		return (A - B).CountStoredTypes() + (B - A).CountStoredTypes();
+	}
+
+	float CalcArchetypeBitDistance(FMassDebuggerArchetypeData& A, FMassDebuggerArchetypeData& B)
+	{
+		int32 TotalLength = A.Composition.CountStoredTypes() +  B.Composition.CountStoredTypes();
+
+		check(TotalLength > 0);
+
+		return float(BitSetDistance(A.Composition.Fragments, B.Composition.Fragments)
+			+ BitSetDistance(A.Composition.Tags, B.Composition.Tags)
+			+ BitSetDistance(A.Composition.ChunkFragments, B.Composition.ChunkFragments)
+			+ BitSetDistance(A.Composition.SharedFragments, B.Composition.SharedFragments)) 
+			/ TotalLength;
+	}
+
+	void MakeDisplayName(const FString& InName, FString& OutDisplayName)
+	{
+		OutDisplayName = InName;
+		if (GET_MASSDEBUGGER_CONFIG_VALUE(bStripMassPrefix) == false)
+		{
+			return;
+		}
+		
+		OutDisplayName.RemoveFromStart(TEXT("Mass"), ESearchCase::CaseSensitive);
+	}
+
+	uint32 CalcProcessorHash(const UMassProcessor& Processor)
+	{
+		return PointerHash(&Processor);
+	}
+} // namespace UE::Mass::Debugger::Private
+
+//----------------------------------------------------------------------//
+// FMassDebuggerEnvironment
+//----------------------------------------------------------------------//
+FString FMassDebuggerEnvironment::GetDisplayName() const
+{
+	return World.IsValid() ? World->GetPathName() : TEXT("Invalid");
+}
+
+const UMassEntitySubsystem* FMassDebuggerEnvironment::GetEntitySubsystem() const
+{
+	if (World.IsValid())
+	{
+		return World->GetSubsystem<UMassEntitySubsystem>();
+	}
+	return nullptr;
+}
+
+//----------------------------------------------------------------------//
+// FMassDebuggerQueryData
+//----------------------------------------------------------------------//
+FMassDebuggerQueryData::FMassDebuggerQueryData(const FMassEntityQuery& Query)
+{
+	FMassDebugger::GetQueryExecutionRequirements(Query, ExecutionRequirements);
+}
+
+int32 FMassDebuggerQueryData::GetTotalBitsUsedCount() 
+{
+	return ExecutionRequirements.GetTotalBitsUsedCount();
+}
+
+//----------------------------------------------------------------------//
+// FMassDebuggerProcessorData
+//----------------------------------------------------------------------//
+FMassDebuggerProcessorData::FMassDebuggerProcessorData(const UMassProcessor& InProcessor)
+{
+	SetProcessor(InProcessor);
+
+	TConstArrayView<FMassEntityQuery*> ProcessorQueries = FMassDebugger::GetProcessorQueries(InProcessor);
+	Queries.Reserve(ProcessorQueries.Num());
+	for (const FMassEntityQuery* Query : ProcessorQueries)
+	{
+		check(Query);
+		Queries.Add(MakeShareable(new FMassDebuggerQueryData(*Query)));
+	}
+}
+
+FMassDebuggerProcessorData::FMassDebuggerProcessorData(const UMassEntitySubsystem& EntitySubsystem, UMassProcessor& InProcessor, const TMap<FMassArchetypeHandle, TSharedPtr<FMassDebuggerArchetypeData>>& InTransientArchetypesMap)
+{
+	SetProcessor(InProcessor);
+
+	TConstArrayView<FMassEntityQuery*> ProcessorQueries = FMassDebugger::GetUpToDateProcessorQueries(EntitySubsystem, InProcessor);
+	Queries.Reserve(ProcessorQueries.Num());
+
+	for (const FMassEntityQuery* Query : ProcessorQueries)
+	{
+		check(Query);
+		Queries.Add(MakeShareable(new FMassDebuggerQueryData(*Query)));
+
+		for (const FMassArchetypeHandle& ArchetypeHandle : Query->GetArchetypes())
+		{
+			ValidArchetypes.Add(InTransientArchetypesMap.FindChecked(ArchetypeHandle));
+		}
+	}
+}
+
+void FMassDebuggerProcessorData::SetProcessor(const UMassProcessor& InProcessor)
+{
+	Name = InProcessor.GetProcessorName();
+	UE::Mass::Debugger::Private::MakeDisplayName(Name, Label);
+
+	ProcessorHash = UE::Mass::Debugger::Private::CalcProcessorHash(InProcessor);
+}
+
+//----------------------------------------------------------------------//
+// FMassDebuggerArchetypeData
+//----------------------------------------------------------------------//
+FMassDebuggerArchetypeData::FMassDebuggerArchetypeData(const FMassArchetypeHandle& ArchetypeHandle)
+{
+	Composition = FMassDebugger::GetArchetypeComposition(ArchetypeHandle);
+	SharedFragments = FMassDebugger::GetArchetypeSharedFragmentValues(ArchetypeHandle);
+
+	// @todo should ensure we're using same hashing as the EntitySubsystem here
+	FamilyHash = Composition.CalculateHash();
+	FullHash = HashCombine(FamilyHash, GetTypeHash(SharedFragments));
+
+	BytesToHexLower(reinterpret_cast<const uint8*>(&FullHash), sizeof(FullHash), Label);
+
+	FMassDebugger::GetArchetypeEntityStats(ArchetypeHandle, EntitiesCount, EntitiesCountPerChunk, ChunksCount);
+}
+
+int32 FMassDebuggerArchetypeData::GetTotalBitsUsedCount() const
+{
+	return Composition.CountStoredTypes();
+}
+
+//----------------------------------------------------------------------//
+// FMassDebuggerProcessingGraphNode
+//----------------------------------------------------------------------//
+FMassDebuggerProcessingGraphNode::FMassDebuggerProcessingGraphNode(const TSharedPtr<FMassDebuggerProcessorData>& InProcessorData, const UMassCompositeProcessor::FDependencyNode& InProcessorNode)
+	: ProcessorData(InProcessorData)
+{
+	if (InProcessorNode.Processor == nullptr)
+	{
+		return;
+	}
+
+	WaitForNodes = InProcessorNode.Dependencies;
+}
+
+FText FMassDebuggerProcessingGraphNode::GetLabel() const
+{
+	if (ProcessorData.IsValid())
+	{
+		return FText::FromString(ProcessorData->Label);
+	}
+
+	return LOCTEXT("InvalidProcessor", "Invalid");
+}
+
+//----------------------------------------------------------------------//
+// FMassDebuggerProcessingGraph
+//----------------------------------------------------------------------//
+FMassDebuggerProcessingGraph::FMassDebuggerProcessingGraph(const FMassDebuggerModel& DebuggerModel, UMassCompositeProcessor& InGraphOwner)
+{
+	Label = InGraphOwner.GetProcessorName();
+
+	TConstArrayView<UMassCompositeProcessor::FDependencyNode> ProcessingGraph = FMassDebugger::GetProcessingGraph(InGraphOwner);
+
+	GraphNodes.Reserve(ProcessingGraph.Num());
+	for (const UMassCompositeProcessor::FDependencyNode& Node : ProcessingGraph)
+	{
+		check(Node.Processor);
+		const TSharedPtr<FMassDebuggerProcessorData>& ProcessorData = DebuggerModel.GetProcessorDataChecked(*Node.Processor);
+		check(ProcessorData.IsValid());
+		GraphNodes.Add(FMassDebuggerProcessingGraphNode(ProcessorData, Node));
+	}
+}
+
+//----------------------------------------------------------------------//
+// FMassDebuggerModel
+//----------------------------------------------------------------------//
+void FMassDebuggerModel::SetEnvironment(const TSharedPtr<FMassDebuggerEnvironment>& Item)
+{
+	if (Item)
+	{
+		Environment = Item;
+		EnvironmentDisplayName = Item->GetDisplayName();
+	}
+	else
+	{
+		Environment = nullptr;
+		EnvironmentDisplayName.Reset();
+	}
+
+	RefreshAll();
+}
+
+void FMassDebuggerModel::RefreshAll()
+{
+	TMap<FMassArchetypeHandle, TSharedPtr<FMassDebuggerArchetypeData>> TransientArchetypesMap;
+
+	CacheArchetypesData(TransientArchetypesMap);
+	CacheProcessorsData(TransientArchetypesMap);
+	CacheProcessingGraphs();
+
+	OnRefreshDelegate.Broadcast();
+}
+
+void FMassDebuggerModel::SelectProcessor(TSharedPtr<FMassDebuggerProcessorData>& Processor)
+{
+	SelectProcessors(MakeArrayView(&Processor, 1), ESelectInfo::Direct);
+}
+
+void FMassDebuggerModel::SelectProcessors(TArrayView<TSharedPtr<FMassDebuggerProcessorData>> Processors, ESelectInfo::Type SelectInfo)
+{
+	SelectionMode = EMassDebuggerSelectionMode::Processor;
+
+	ResetSelectedProcessors();
+	ResetSelectedArchetypes();
+
+	SelectedProcessors = Processors;
+
+	for (TSharedPtr<FMassDebuggerProcessorData>& ProcessorData : SelectedProcessors)
+	{
+		check(ProcessorData.IsValid());
+		ProcessorData->Selection = EMassDebuggerProcessorSelection::Selected;
+
+		for (TSharedPtr<FMassDebuggerArchetypeData>& ArchetypeData : ProcessorData->ValidArchetypes)
+		{
+			SelectedArchetypes.AddUnique(ArchetypeData);
+			ArchetypeData->bIsSelected = true;
+		}
+	}
+
+	OnProcessorsSelectedDelegate.Broadcast(SelectedProcessors, SelectInfo);
+}
+
+void FMassDebuggerModel::ClearProcessorSelection()
+{
+	SelectionMode = EMassDebuggerSelectionMode::None;
+
+	ResetSelectedProcessors();
+
+	OnProcessorsSelectedDelegate.Broadcast(SelectedProcessors, ESelectInfo::Direct);
+}
+
+void FMassDebuggerModel::SelectArchetypes(TArray<TSharedPtr<FMassDebuggerArchetypeData>> InSelectedArchetypes, ESelectInfo::Type SelectInfo)
+{
+	ResetSelectedProcessors();
+	ResetSelectedArchetypes();
+
+	SelectionMode = EMassDebuggerSelectionMode::Archetype;
+
+	SelectedArchetypes = InSelectedArchetypes;
+
+	for (TSharedPtr<FMassDebuggerProcessorData>& ProcessorData : CachedProcessors)
+	{
+		check(ProcessorData.IsValid());
+		for (const TSharedPtr<FMassDebuggerArchetypeData>& ArchetypeData : InSelectedArchetypes)
+		{
+			if (ProcessorData->ValidArchetypes.Find(ArchetypeData) != INDEX_NONE)
+			{
+				ProcessorData->Selection = EMassDebuggerProcessorSelection::Selected;
+				SelectedProcessors.Add(ProcessorData);
+				break;
+			}
+		}
+	}
+
+	OnArchetypesSelectedDelegate.Broadcast(SelectedArchetypes, SelectInfo);
+}
+
+void FMassDebuggerModel::ClearArchetypeSelection()
+{
+	SelectionMode = EMassDebuggerSelectionMode::None;
+
+	ResetSelectedArchetypes();
+	OnArchetypesSelectedDelegate.Broadcast(SelectedArchetypes, ESelectInfo::Direct);
+}
+
+void FMassDebuggerModel::CacheProcessorsData(const TMap<FMassArchetypeHandle, TSharedPtr<FMassDebuggerArchetypeData>>& InTransientArchetypesMap)
+{
+	CachedProcessors.Reset();
+
+	UWorld* World = Environment ? Environment->World.Get() : nullptr;
+	const UMassEntitySubsystem* EntitySubsystem = Environment ? Environment->GetEntitySubsystem() : nullptr;
+
+	if (EntitySubsystem)
+	{
+		check(World);
+		for (FThreadSafeObjectIterator It(UMassProcessor::StaticClass()); It; ++It)
+		{
+			UMassProcessor* Processor = Cast<UMassProcessor>(*It);
+			if (Processor && (Processor->HasAnyFlags(RF_ClassDefaultObject) == false)
+				&& Cast<UMassCompositeProcessor>(Processor) == nullptr
+				&& Processor->GetWorld() == World)
+			{
+				CachedProcessors.Add(MakeShareable(new FMassDebuggerProcessorData(*EntitySubsystem, *Processor, InTransientArchetypesMap)));
+			}
+		}
+	}
+	else
+	{
+		for (FThreadSafeObjectIterator It(UMassProcessor::StaticClass()); It; ++It)
+		{
+			UMassProcessor* Processor = Cast<UMassProcessor>(*It);
+			if (Processor && (Processor->HasAnyFlags(RF_ClassDefaultObject) == false)
+				&& Cast<UMassCompositeProcessor>(Processor) == nullptr 
+				&& (World == nullptr || Processor->GetWorld() == World))
+			{
+				CachedProcessors.Add(MakeShareable(new FMassDebuggerProcessorData(*Processor)));
+			}
+		}
+	}
+
+	CachedProcessors.Sort([](const TSharedPtr<FMassDebuggerProcessorData>& A, const TSharedPtr<FMassDebuggerProcessorData>& B)
+	{
+		return A->Label < B->Label;
+	});
+}
+
+void FMassDebuggerModel::CacheProcessingGraphs()
+{
+	CachedProcessingGraphs.Reset();
+
+	UWorld* World = Environment ? Environment->World.Get() : nullptr;
+	for (FThreadSafeObjectIterator It(UMassProcessor::StaticClass()); It; ++It)
+	{
+		UMassCompositeProcessor* Processor = Cast<UMassCompositeProcessor>(*It);
+		if (Processor && (Processor->HasAnyFlags(RF_ClassDefaultObject) == false)
+			&& (World == nullptr || Processor->GetWorld() == World))
+		{
+			CachedProcessingGraphs.Add(MakeShareable(new FMassDebuggerProcessingGraph(*this, *Processor)));
+		}
+	}	
+}
+
+void FMassDebuggerModel::CacheArchetypesData(TMap<FMassArchetypeHandle, TSharedPtr<FMassDebuggerArchetypeData>>& OutTransientArchetypesMap)
+{
+	CachedArchetypes.Reset();
+
+	if (Environment)
+	{
+		if (const UMassEntitySubsystem* EntitySubsystem = Environment->GetEntitySubsystem())
+		{
+			StoreArchetypes(*EntitySubsystem, OutTransientArchetypesMap);
+		}
+	}
+}
+
+float FMassDebuggerModel::MinDistanceToSelectedArchetypes(const TSharedPtr<FMassDebuggerArchetypeData>& InArchetypeData) const
+{
+	float MinDistance = MAX_flt;
+	for (const TSharedPtr<FMassDebuggerArchetypeData>& SelectedArchetype : SelectedArchetypes)
+	{
+		MinDistance = FMath::Min(MinDistance, ArchetypeDistances[SelectedArchetype->Index][InArchetypeData->Index]);
+	}
+	return MinDistance;
+}
+
+void FMassDebuggerModel::StoreArchetypes(const UMassEntitySubsystem& EntitySubsystem, TMap<FMassArchetypeHandle, TSharedPtr<FMassDebuggerArchetypeData>>& OutTransientArchetypesMap)
+{
+	TArray<FMassArchetypeHandle> ArchetypeHandles = FMassDebugger::GetAllArchetypes(EntitySubsystem);
+
+	CachedArchetypes.Reset(ArchetypeHandles.Num());
+
+	int32 MaxBitsUsed = 0;
+
+	// @todo build an archetype handle map
+	for (FMassArchetypeHandle& ArchetypeHandle : ArchetypeHandles)
+	{
+		FMassDebuggerArchetypeData* ArchetypeDataPtr = new FMassDebuggerArchetypeData(ArchetypeHandle);
+		ArchetypeDataPtr->Index = CachedArchetypes.Add(MakeShareable(ArchetypeDataPtr));
+		OutTransientArchetypesMap.Add(ArchetypeHandle, CachedArchetypes.Last());
+
+		MaxBitsUsed = FMath::Max(MaxBitsUsed, ArchetypeDataPtr->GetTotalBitsUsedCount());
+	}
+	
+	// calculate distances
+	ArchetypeDistances.Reset();
+	ArchetypeDistances.AddDefaulted(CachedArchetypes.Num());
+	for (int i = 0; i < CachedArchetypes.Num(); ++i)
+	{
+		ArchetypeDistances[i].AddDefaulted(CachedArchetypes.Num());
+	}
+
+	for (int i = 0; i < CachedArchetypes.Num(); ++i)
+	{
+		for (int k = i + 1; k < CachedArchetypes.Num(); ++k)
+		{
+			const float Distance = UE::Mass::Debugger::Private::CalcArchetypeBitDistance(*CachedArchetypes[i].Get(), *CachedArchetypes[k].Get());
+			ArchetypeDistances[i][k] = Distance;
+			ArchetypeDistances[k][i] = Distance;
+		}
+	}
+}
+
+FText FMassDebuggerModel::GetDisplayName() const
+{
+	if (!Environment)
+	{
+		return LOCTEXT("PickEnvironment", "Pick Environment");
+	}
+	else if (IsStale())
+	{
+		return FText::FromString(FString::Printf(TEXT("(%s) %s")
+				, *(LOCTEXT("StaleEnvironmentPrefix", "Stale").ToString())
+				, *EnvironmentDisplayName));
+	}
+
+	return FText::FromString(Environment->GetDisplayName());
+}
+
+void FMassDebuggerModel::MarkAsStale()
+{
+	if (Environment)
+	{
+		Environment->World = nullptr;
+	}
+}
+
+bool FMassDebuggerModel::IsStale() const 
+{ 
+	return Environment.IsValid() && (Environment->IsWorldValid() == false); 
+}
+
+const TSharedPtr<FMassDebuggerProcessorData>& FMassDebuggerModel::GetProcessorDataChecked(const UMassProcessor& Processor) const
+{
+	check(CachedProcessors.Num());
+
+	const uint32 ProcessorHash = UE::Mass::Debugger::Private::CalcProcessorHash(Processor);
+
+	const TSharedPtr<FMassDebuggerProcessorData>* DataFound = CachedProcessors.FindByPredicate([ProcessorHash](const TSharedPtr<FMassDebuggerProcessorData>& Element)
+		{
+			return Element->ProcessorHash == ProcessorHash;
+		});
+
+	check(DataFound);
+	return *DataFound;
+}
+
+void FMassDebuggerModel::ResetSelectedArchetypes()
+{
+	for (TSharedPtr<FMassDebuggerArchetypeData>& ArchetypeData : SelectedArchetypes)
+	{
+		ArchetypeData->bIsSelected = false;
+	}
+	SelectedArchetypes.Reset();
+}
+
+void FMassDebuggerModel::ResetSelectedProcessors()
+{
+	// using CachedProcessors instead of SelectedProcessors to be on the safe side
+	for (TSharedPtr<FMassDebuggerProcessorData>& ProcessorData : CachedProcessors)
+	{
+		check(ProcessorData.IsValid());
+		ProcessorData->Selection = EMassDebuggerProcessorSelection::None;
+	}
+	SelectedProcessors.Reset();
+}
+
+#undef LOCTEXT_NAMESPACE
diff --git a/Engine/Plugins/Runtime/MassEntity/Source/MassEntityDebugger/Private/MassDebuggerModel.h b/Engine/Plugins/Runtime/MassEntity/Source/MassEntityDebugger/Private/MassDebuggerModel.h
new file mode 100644
--- /dev/null
+++ b/Engine/Plugins/Runtime/MassEntity/Source/MassEntityDebugger/Private/MassDebuggerModel.h
@@ -0,0 +1,189 @@
+// Copyright Epic Games, Inc. All Rights Reserved.
+#pragma once
+
+#include "MassArchetypeTypes.h"
+#include "MassEntityQuery.h"
+#include "MassProcessorDependencySolver.h"
+#include "MassProcessingTypes.h"
+#include "MassDebugger.h"
+
+
+class UMassProcessor;
+struct FMassArchetypeHandle;
+struct FMassDebuggerModel;
+class UMassEntitySubsystem;
+
+enum class EMassDebuggerSelectionMode : uint8
+{
+	None,
+	Processor,
+	Archetype,
+	// @todo future:
+	// Fragment
+	MAX
+};
+
+enum class EMassDebuggerProcessorSelection : uint8
+{
+	None,
+	Selected,
+	MAX
+};
+
+enum class EMassDebuggerProcessingGraphNodeSelection : uint8
+{
+	None,
+	WaitFor,
+	Block,
+	MAX
+};
+
+
+struct FMassDebuggerQueryData
+{
+	FMassDebuggerQueryData(const FMassEntityQuery& Query);
+
+	FMassExecutionRequirements ExecutionRequirements;
+
+	int32 GetTotalBitsUsedCount();
+}; 
+
+struct FMassDebuggerArchetypeData
+{
+	FMassDebuggerArchetypeData(const FMassArchetypeHandle& ArchetypeHandle);
+
+	FMassArchetypeCompositionDescriptor Composition;
+	FMassArchetypeSharedFragmentValues SharedFragments;
+	uint32 FamilyHash = 0;
+	uint32 FullHash = 0;
+
+	int32 EntitiesCount = 0;
+	int32 EntitiesCountPerChunk = 0;
+	int32 ChunksCount = 0;
+
+	TArray<TSharedPtr<FMassDebuggerArchetypeData>> Children;
+
+	int32 Index = INDEX_NONE;
+	FString Label;
+
+	bool bIsSelected = false;
+
+	int32 GetTotalBitsUsedCount() const;
+};
+
+struct FMassDebuggerProcessorData
+{
+	FMassDebuggerProcessorData(const UMassProcessor& InProcessor);
+	FMassDebuggerProcessorData(const UMassEntitySubsystem& EntitySubsystem, UMassProcessor& InProcessor, const TMap<FMassArchetypeHandle, TSharedPtr<FMassDebuggerArchetypeData>>& InTransientArchetypesMap);
+
+private:
+	void SetProcessor(const UMassProcessor& InProcessor);
+
+public:
+	FString Name;
+	FString Label;
+	uint32 ProcessorHash = 0; 
+	TArray<TSharedPtr<FMassDebuggerQueryData>> Queries;
+	TArray<TSharedPtr<FMassDebuggerArchetypeData>> ValidArchetypes;
+	
+	EMassDebuggerProcessorSelection Selection = EMassDebuggerProcessorSelection::None;
+};
+
+struct FMassDebuggerProcessingGraphNode
+{
+	FMassDebuggerProcessingGraphNode(const TSharedPtr<FMassDebuggerProcessorData>& InProcessorData, const UMassCompositeProcessor::FDependencyNode& InProcessorNode);
+	
+	FText GetLabel() const;
+
+	TSharedPtr<FMassDebuggerProcessorData> ProcessorData;
+	TArray<int32> WaitForNodes;
+	TArray<int32> BlockNodes;
+	EMassDebuggerProcessingGraphNodeSelection GraphNodeSelection = EMassDebuggerProcessingGraphNodeSelection::None;
+};
+
+struct FMassDebuggerProcessingGraph
+{
+	FMassDebuggerProcessingGraph(const FMassDebuggerModel& DebuggerModel, UMassCompositeProcessor& InGraphOwner);
+
+	FString Label;
+	TArray<FMassDebuggerProcessingGraphNode> GraphNodes;
+};
+
+
+struct FMassDebuggerEnvironment
+{
+	explicit FMassDebuggerEnvironment(UWorld* InWorld)
+		: World(InWorld)
+	{}
+
+	bool operator==(const FMassDebuggerEnvironment& Other) const { return World == Other.World; }
+
+	FString GetDisplayName() const;
+	const UMassEntitySubsystem* GetEntitySubsystem() const;
+	bool IsWorldValid() const { return World.IsValid(); }
+	
+	TWeakObjectPtr<UWorld> World;
+};
+
+
+struct FMassDebuggerModel
+{
+	DECLARE_MULTICAST_DELEGATE(FOnRefresh);
+	DECLARE_MULTICAST_DELEGATE_TwoParams(FOnProcessorsSelected, TConstArrayView<TSharedPtr<FMassDebuggerProcessorData>>, ESelectInfo::Type);
+	DECLARE_MULTICAST_DELEGATE_TwoParams(FOnArchetypesSelected, TConstArrayView<TSharedPtr<FMassDebuggerArchetypeData>>, ESelectInfo::Type);
+
+	void SetEnvironment(const TSharedPtr<FMassDebuggerEnvironment>& Item);
+
+	void RefreshAll();
+
+	void SelectProcessor(TSharedPtr<FMassDebuggerProcessorData>& Processor);
+	void SelectProcessors(TArrayView<TSharedPtr<FMassDebuggerProcessorData>> Processors, ESelectInfo::Type SelectInfo);
+	void ClearProcessorSelection();
+
+	void SelectArchetypes(TArray<TSharedPtr<FMassDebuggerArchetypeData>> SelectedArchetypes, ESelectInfo::Type SelectInfo);
+	void ClearArchetypeSelection();
+
+	bool IsCurrentEnvironment(const FMassDebuggerEnvironment& InEnvironment) const { return Environment && *Environment.Get() == InEnvironment; }
+
+	void CacheArchetypesData(TMap<FMassArchetypeHandle, TSharedPtr<FMassDebuggerArchetypeData>>& OutTransientArchetypesMap); 
+	void CacheProcessorsData(const TMap<FMassArchetypeHandle, TSharedPtr<FMassDebuggerArchetypeData>>& InTransientArchetypesMap);
+	void CacheProcessingGraphs();
+
+	float MinDistanceToSelectedArchetypes(const TSharedPtr<FMassDebuggerArchetypeData>& InArchetypeData) const;
+
+	FText GetDisplayName() const;
+
+	void MarkAsStale();
+	bool IsStale() const;
+
+	const TSharedPtr<FMassDebuggerProcessorData>& GetProcessorDataChecked(const UMassProcessor& Processor) const;
+
+protected:
+	void StoreArchetypes(const UMassEntitySubsystem& EntitySubsystem, TMap<FMassArchetypeHandle, TSharedPtr<FMassDebuggerArchetypeData>>& OutTransientArchetypesMap);
+
+	void ResetSelectedArchetypes();
+	void ResetSelectedProcessors();
+
+public:
+	FOnRefresh OnRefreshDelegate;
+	FOnProcessorsSelected OnProcessorsSelectedDelegate;
+	FOnArchetypesSelected OnArchetypesSelectedDelegate;
+
+	EMassDebuggerSelectionMode SelectionMode = EMassDebuggerSelectionMode::None;
+
+	TSharedPtr<FMassDebuggerEnvironment> Environment;
+	TArray<TSharedPtr<FMassDebuggerProcessorData>> CachedProcessors;
+	TArray<TSharedPtr<FMassDebuggerProcessorData>> SelectedProcessors;
+
+	TArray<TSharedPtr<FMassDebuggerArchetypeData>> CachedArchetypes;
+	TArray<TSharedPtr<FMassDebuggerArchetypeData>> SelectedArchetypes;
+
+	TArray<TSharedPtr<FMassDebuggerProcessingGraph>> CachedProcessingGraphs;
+
+	TMap<FMassArchetypeHandle, TSharedPtr<FMassDebuggerArchetypeData>> HandleToArchetypeMap;
+
+	TArray<TArray<float>> ArchetypeDistances;
+
+	FString EnvironmentDisplayName;
+};
+
diff --git a/Engine/Plugins/Runtime/MassEntity/Source/MassEntityDebugger/Private/MassDebuggerModule.cpp b/Engine/Plugins/Runtime/MassEntity/Source/MassEntityDebugger/Private/MassDebuggerModule.cpp
--- a/Engine/Plugins/Runtime/MassEntity/Source/MassEntityDebugger/Private/MassDebuggerModule.cpp
+++ b/Engine/Plugins/Runtime/MassEntity/Source/MassEntityDebugger/Private/MassDebuggerModule.cpp
@@ -33,6 +33,7 @@ void FMassDebuggerModule::StartupModule()
 {
 	// This code will execute after your module is loaded into memory (but after global variables are initialized, of course.)
 	FMassDebuggerStyle::Initialize();
+	FMassDebuggerCommands::Register();
 
 	// register the UI tab
 	IModularFeatures::Get().RegisterModularFeature(MassDebugTabName, this);
diff --git a/Engine/Plugins/Runtime/MassEntity/Source/MassEntityDebugger/Private/MassDebuggerStyle.cpp b/Engine/Plugins/Runtime/MassEntity/Source/MassEntityDebugger/Private/MassDebuggerStyle.cpp
--- a/Engine/Plugins/Runtime/MassEntity/Source/MassEntityDebugger/Private/MassDebuggerStyle.cpp
+++ b/Engine/Plugins/Runtime/MassEntity/Source/MassEntityDebugger/Private/MassDebuggerStyle.cpp
@@ -3,17 +3,56 @@
 #include "MassDebuggerStyle.h"
 #include "Styling/SlateStyleRegistry.h"
 #include "Interfaces/IPluginManager.h"
+#include "Styling/CoreStyle.h"
+#include "Styling/SlateTypes.h"
 #include "Styling/SlateStyle.h"
 #include "Styling/SlateStyleMacros.h"
+#include "Styling/AppStyle.h"
+#include "UObject/ConstructorHelpers.h"
+#include "Engine/Font.h"
+#include "Components/Widget.h"
 
 
-TSharedPtr<FSlateStyleSet> FMassDebuggerStyle::StyleSet = nullptr;
+#define TTF_CORE_FONT(RelativePath, ...) FSlateFontInfo(RootToCoreContentDir(RelativePath, TEXT(".ttf")), __VA_ARGS__)
+#define FONT(...) FSlateFontInfo(FCoreStyle::GetDefaultFont(), __VA_ARGS__)
 
-FString FMassDebuggerStyle::InContent(const FString& RelativePath, const ANSICHAR* Extension)
+class FMassDebuggerStyleSet : public FSlateStyleSet
 {
-	static FString ContentDir = IPluginManager::Get().FindPlugin(TEXT("MassEntity"))->GetContentDir() / TEXT("Slate");
-	return (ContentDir / RelativePath) + Extension;
-}
+public:
+	static FString InContent(const FString& RelativePath, const ANSICHAR* Extension)
+	{
+		static FString ContentDir = IPluginManager::Get().FindPlugin(TEXT("MassEntity"))->GetContentDir() / TEXT("Slate");
+		return (ContentDir / RelativePath) + Extension;
+	}
+
+	FMassDebuggerStyleSet(const FName& InStyleSetName)
+		: FSlateStyleSet(InStyleSetName)
+	{
+		const FVector2D Icon16x16(16.0f, 16.0f);
+
+		SetContentRoot(FPaths::EngineContentDir() / TEXT("Editor/Slate"));
+		SetCoreContentRoot(FPaths::EngineContentDir() / TEXT("Slate"));
+
+		Set("MassDebuggerApp.TabIcon", new FSlateVectorImageBrush(InContent("MassDebugger", ".svg"), Icon16x16));
+
+		const FTextBlockStyle& NormalText = FAppStyle::Get().GetWidgetStyle<FTextBlockStyle>("NormalText");
+
+		const FSlateFontInfo LargeFont(FONT(50, "Bold"));
+
+		FTextBlockStyle StateTitle = FTextBlockStyle(NormalText)
+			.SetFont(LargeFont)
+			.SetColorAndOpacity(FLinearColor(230.0f / 255.0f, 230.0f / 255.0f, 230.0f / 255.0f, 0.9f));
+		Set("MassDebug.Font.Large", StateTitle);
+
+		Set("MassDebug.Fragment", new BOX_BRUSH("Common/LightGroupBorder", FMargin(4.0f / 16.0f), FLinearColor(.2f, .2f, .9f)));
+		Set("MassDebug.Fragment.ReadOnly", new BOX_BRUSH("Common/LightGroupBorder", FMargin(4.0f / 16.0f), FLinearColor(.5f, .5f, .5f)));
+		Set("MassDebug.Fragment.ReadWrite", new BOX_BRUSH("Common/LightGroupBorder", FMargin(4.0f / 16.0f), FLinearColor(0.f, .7f, 0.f)));
+	}
+};
+
+
+TSharedPtr<FSlateStyleSet> FMassDebuggerStyle::StyleSet = nullptr;
+
 
 void FMassDebuggerStyle::Initialize()
 {
@@ -22,15 +61,7 @@ void FMassDebuggerStyle::Initialize()
 		return;
 	}
 
-	const FVector2D Icon16x16(16.0f, 16.0f);
-
-	StyleSet = MakeShared<FSlateStyleSet>(GetStyleSetName());
-
-	StyleSet->SetContentRoot(FPaths::EngineContentDir() / TEXT("Editor/Slate"));
-	StyleSet->SetCoreContentRoot(FPaths::EngineContentDir() / TEXT("Slate"));
-
-	FSlateStyleSet& Style = *StyleSet.Get();
-	StyleSet->Set("MassDebuggerApp.TabIcon", new FSlateVectorImageBrush(InContent("MassDebugger", ".svg"), Icon16x16));
+	StyleSet = MakeShared<FMassDebuggerStyleSet>(GetStyleSetName());
 	
 	FSlateStyleRegistry::RegisterSlateStyle(*StyleSet.Get());
 }
diff --git a/Engine/Plugins/Runtime/MassEntity/Source/MassEntityDebugger/Private/MassDebuggerStyle.h b/Engine/Plugins/Runtime/MassEntity/Source/MassEntityDebugger/Private/MassDebuggerStyle.h
--- a/Engine/Plugins/Runtime/MassEntity/Source/MassEntityDebugger/Private/MassDebuggerStyle.h
+++ b/Engine/Plugins/Runtime/MassEntity/Source/MassEntityDebugger/Private/MassDebuggerStyle.h
@@ -10,7 +10,7 @@ class MASSENTITYDEBUGGER_API FMassDebuggerStyle
 public:
 	static void Initialize();
 	static void Shutdown();
-	static TSharedPtr<ISlateStyle> Get() { return StyleSet; }
+	static ISlateStyle& Get() { return *StyleSet.Get(); }
 	static FName GetStyleSetName();
 
 	static const FSlateBrush* GetBrush(FName PropertyName, const ANSICHAR* Specifier = NULL)
@@ -18,8 +18,6 @@ public:
 		return StyleSet->GetBrush(PropertyName, Specifier);
 	}
 
-	static FString InContent(const FString& RelativePath, const ANSICHAR* Extension);
-
 private:
 	static TSharedPtr<FSlateStyleSet> StyleSet;
 };
diff --git a/Engine/Plugins/Runtime/MassEntity/Source/MassEntityDebugger/Private/SMassArchetype.cpp b/Engine/Plugins/Runtime/MassEntity/Source/MassEntityDebugger/Private/SMassArchetype.cpp
new file mode 100644
--- /dev/null
+++ b/Engine/Plugins/Runtime/MassEntity/Source/MassEntityDebugger/Private/SMassArchetype.cpp
@@ -0,0 +1,76 @@
+// Copyright Epic Games, Inc. All Rights Reserved.
+
+#include "SMassArchetype.h"
+#include "MassDebuggerModel.h"
+#include "MassEntityTypes.h"
+#include "SMassBitSet.h"
+#include "Widgets/Text/SRichTextBlock.h"
+#include "Styling/AppStyle.h"
+
+#define LOCTEXT_NAMESPACE "SMassDebugger"
+
+//----------------------------------------------------------------------//
+// SMassArchetype
+//----------------------------------------------------------------------//
+void SMassArchetype::Construct(const FArguments& InArgs, TSharedPtr<FMassDebuggerArchetypeData> InArchetypeData)
+{
+	using UE::Mass::Debugger::UI::AddBitSet;
+
+	if (!InArchetypeData)
+	{
+		return;
+	}
+	
+	ArchetypeData = InArchetypeData;
+
+	FMassDebuggerArchetypeData& ArchetypeDebugData = *InArchetypeData.Get();
+
+	TSharedRef<SVerticalBox> Box = SNew(SVerticalBox);
+
+	const TArray<FText> LabelBits = {
+		LOCTEXT("MassArchetypeLabel", "Archetype")
+		, FText::FromString(InArchetypeData->Label)
+	};
+
+	Box->AddSlot()
+	.AutoHeight()
+	[
+		SNew(SRichTextBlock)
+		.Text(FText::Join(FText::FromString(TEXT(": ")), LabelBits))
+		.DecoratorStyleSet(&FAppStyle::Get())
+		.TextStyle(FAppStyle::Get(), "LargeText")
+	];
+
+	FString ArchetypeDescription = FString::Printf(TEXT("EntitiesCount: %d\nEntitiesCountPerChunk: %d\nChunksCount: %d")
+		, ArchetypeDebugData.EntitiesCount
+		, ArchetypeDebugData.EntitiesCountPerChunk
+		, ArchetypeDebugData.ChunksCount);
+	
+	Box->AddSlot()
+	.AutoHeight()
+	[
+		SNew(STextBlock)
+		.Text(FText::FromString(ArchetypeDescription))
+	];
+
+	const FMassArchetypeCompositionDescriptor& Composition = ArchetypeData->Composition;
+	const FSlateBrush* Brush = FMassDebuggerStyle::GetBrush("MassDebug.Fragment");
+
+	AddBitSet(Box, Composition.Fragments, TEXT("Fragments"), Brush);
+	AddBitSet(Box, Composition.Tags, TEXT("Tags"), Brush);
+	AddBitSet(Box, Composition.ChunkFragments, TEXT("Chunk Fragments"), Brush);
+	AddBitSet(Box, Composition.SharedFragments, TEXT("Shared Fragments"), Brush);
+
+	ChildSlot
+	[
+		SNew(SBorder)
+		.Padding(5.0f)
+		[
+			Box
+		]
+	];
+	
+}
+
+#undef LOCTEXT_NAMESPACE
+
diff --git a/Engine/Plugins/Runtime/MassEntity/Source/MassEntityDebugger/Private/SMassArchetype.h b/Engine/Plugins/Runtime/MassEntity/Source/MassEntityDebugger/Private/SMassArchetype.h
new file mode 100644
--- /dev/null
+++ b/Engine/Plugins/Runtime/MassEntity/Source/MassEntityDebugger/Private/SMassArchetype.h
@@ -0,0 +1,20 @@
+// Copyright Epic Games, Inc. All Rights Reserved.
+#pragma once
+
+#include "Widgets/SCompoundWidget.h"
+
+
+struct FMassDebuggerArchetypeData;
+
+class SMassArchetype : public SCompoundWidget
+{
+public:
+
+	SLATE_BEGIN_ARGS(SMassArchetype){}
+	SLATE_END_ARGS()
+
+	void Construct(const FArguments& InArgs, TSharedPtr<FMassDebuggerArchetypeData> InArchetypeData);
+
+protected:
+	TSharedPtr<FMassDebuggerArchetypeData> ArchetypeData;
+};
diff --git a/Engine/Plugins/Runtime/MassEntity/Source/MassEntityDebugger/Private/SMassArchetypesView.cpp b/Engine/Plugins/Runtime/MassEntity/Source/MassEntityDebugger/Private/SMassArchetypesView.cpp
new file mode 100644
--- /dev/null
+++ b/Engine/Plugins/Runtime/MassEntity/Source/MassEntityDebugger/Private/SMassArchetypesView.cpp
@@ -0,0 +1,233 @@
+// Copyright Epic Games, Inc. All Rights Reserved.
+
+#include "SMassArchetypesView.h"
+#include "MassDebuggerModel.h"
+#include "SMassArchetype.h"
+
+#define LOCTEXT_NAMESPACE "SMassDebugger"
+
+
+namespace UE::Mass::Debugger::UI::Private
+{
+	FLinearColor GetArchetypeDistanceColor(float Distance)
+	{
+		const FLinearColor MinDistance = FLinearColor::Green;
+		const FLinearColor MaxDistance = FLinearColor::Red;
+		
+		return (1.f - Distance) * MinDistance + Distance * MaxDistance;
+	}
+}
+
+using FMassDebuggerArchetypeDataPtr = TSharedPtr<FMassDebuggerArchetypeData, ESPMode::ThreadSafe>;
+
+//----------------------------------------------------------------------//
+// SMassArchetypeTableRowBase
+//----------------------------------------------------------------------//
+class SMassArchetypeTableRowBase : public STableRow<FMassDebuggerArchetypeDataPtr>
+{
+public:
+	SLATE_BEGIN_ARGS(SMassArchetypeTableRowBase) { }
+	SLATE_END_ARGS()
+
+protected:
+	void Construct(const FArguments& InArgs, const TSharedPtr<STableViewBase>& InOwnerTableView, const FMassDebuggerArchetypeDataPtr InEntryItem)
+	{
+		Item = InEntryItem;
+		STableRow<FMassDebuggerArchetypeDataPtr>::Construct(STableRow<FMassDebuggerArchetypeDataPtr>::FArguments(), InOwnerTableView.ToSharedRef());
+	}
+
+	FMassDebuggerArchetypeDataPtr Item;
+};
+
+//----------------------------------------------------------------------//
+// SMassArchetypeTableRow
+//----------------------------------------------------------------------//
+class SMassArchetypeTableRow : public SMassArchetypeTableRowBase
+{
+public:
+	void Construct(const FArguments& InArgs, const TSharedPtr<STableViewBase>& InOwnerTableView, const FMassDebuggerArchetypeDataPtr InEntryItem, TSharedPtr<FMassDebuggerModel> DebuggerModel)
+	{
+		SMassArchetypeTableRowBase::Construct(InArgs, InOwnerTableView, InEntryItem);
+
+		TSharedPtr<SHorizontalBox> Box = SNew(SHorizontalBox);
+		Box->AddSlot()
+			.VAlign(VAlign_Fill)
+			.HAlign(HAlign_Left)
+			.AutoWidth()
+			[
+				SNew(SExpanderArrow, SharedThis(this))
+				.ShouldDrawWires(true)
+				.IndentAmount(32)
+				.BaseIndentLevel(0)
+			];
+
+		if (DebuggerModel->SelectedArchetypes.Num())
+		{
+			if (InEntryItem->bIsSelected)
+			{
+				Box->AddSlot()
+				[
+					SNew(STextBlock)
+					.Text(FText::FromString(InEntryItem->Label))
+					.ColorAndOpacity(FLinearColor::Green)
+				];
+			}
+			else
+			{
+				const float DistanceToSelected = DebuggerModel->MinDistanceToSelectedArchetypes(InEntryItem);
+				const FLinearColor LerpedColor = UE::Mass::Debugger::UI::Private::GetArchetypeDistanceColor(DistanceToSelected);
+
+				Box->AddSlot()
+				[
+					SNew(STextBlock)
+					.Text(FText::FromString(InEntryItem->Label))
+					.ColorAndOpacity(LerpedColor)
+				];
+			}
+		}
+		else
+		{
+			Box->AddSlot()
+			[
+				SNew(STextBlock)
+				.Text(FText::FromString(InEntryItem->Label))
+			];
+		}
+
+		ChildSlot
+		[
+			Box.ToSharedRef()
+		];
+	}
+};
+
+//----------------------------------------------------------------------//
+// SMassArchetypeDetailTableRow
+//----------------------------------------------------------------------//
+class SMassArchetypeDetailTableRow : public SMassArchetypeTableRowBase
+{
+public:
+	void Construct(const FArguments& InArgs, const TSharedPtr<STableViewBase>& InOwnerTableView, const FMassDebuggerArchetypeDataPtr InEntryItem)
+	{
+		SMassArchetypeTableRowBase::Construct(InArgs, InOwnerTableView, InEntryItem);
+
+		ChildSlot
+			[
+				SNew(SMassArchetype, InEntryItem)
+			];
+	}
+};
+
+//----------------------------------------------------------------------//
+// SMassArchetypesView
+//----------------------------------------------------------------------//
+void SMassArchetypesView::Construct(const FArguments& InArgs, TSharedRef<FMassDebuggerModel> InDebuggerModel)
+{
+	Initialize(InDebuggerModel);
+
+	ChildSlot
+		[
+			SNew(SHorizontalBox)
+			+ SHorizontalBox::Slot()
+			.FillWidth(0.2f)
+			[
+				SAssignNew(ArchetypesTreeView, STreeView<TSharedPtr<FMassDebuggerArchetypeData>>)
+				.TreeItemsSource(&TreeViewSource)
+				.OnGenerateRow_Lambda([this](TSharedPtr<FMassDebuggerArchetypeData> Item, const TSharedPtr<STableViewBase>& OwnerTable)
+					{
+						return SNew(SMassArchetypeTableRow, OwnerTable, Item, DebuggerModel);
+					})
+				.OnGetChildren_Lambda([](TSharedPtr<FMassDebuggerArchetypeData> InItem, TArray<TSharedPtr<FMassDebuggerArchetypeData>>& OutChildren)
+					{
+						if (InItem->Children.Num())
+						{
+							OutChildren.Append(InItem->Children);
+						}
+					})
+				.OnSelectionChanged(this, &SMassArchetypesView::HandleSelectionChanged)
+			]
+
+			+ SHorizontalBox::Slot()
+			[
+				SAssignNew(SelectedArchetypesListWidget, SListView<TSharedPtr<FMassDebuggerArchetypeData>>)
+				.ListItemsSource(&DebuggerModel->SelectedArchetypes)
+				.SelectionMode(ESelectionMode::None)
+				.OnGenerateRow_Lambda([](TSharedPtr<FMassDebuggerArchetypeData> Item, const TSharedPtr<STableViewBase>& OwnerTable)
+					{
+						return SNew(SMassArchetypeDetailTableRow, OwnerTable, Item);
+					})
+			]
+		];
+}
+
+void SMassArchetypesView::HandleSelectionChanged(TSharedPtr<FMassDebuggerArchetypeData> InNode, ESelectInfo::Type InSelectInfo)
+{
+	if (!DebuggerModel || InSelectInfo == ESelectInfo::Direct)
+	{
+		return;
+	}
+
+	TArray<TSharedPtr<FMassDebuggerArchetypeData>> SelectedArchetypes;
+	ArchetypesTreeView->GetSelectedItems(SelectedArchetypes);
+	DebuggerModel->SelectArchetypes(SelectedArchetypes, InSelectInfo);
+}
+
+void SMassArchetypesView::OnRefresh()
+{
+	if (DebuggerModel)
+	{
+		TreeViewSource.Reset();
+		
+		TMap<uint32, TSharedPtr<FMassDebuggerArchetypeData>> ArchetypeRepresentativeMap;
+		for (const TSharedPtr<FMassDebuggerArchetypeData>& ArchetypeData : DebuggerModel->CachedArchetypes)
+		{
+			if (TSharedPtr<FMassDebuggerArchetypeData>* Representative = ArchetypeRepresentativeMap.Find(ArchetypeData->FamilyHash))
+			{
+				(*Representative)->Children.Add(ArchetypeData);
+			}
+			else
+			{
+				ArchetypeRepresentativeMap.Add(ArchetypeData->FamilyHash, ArchetypeData);
+			}
+		}
+
+		for (auto& KVP : ArchetypeRepresentativeMap)
+		{
+			TreeViewSource.Add(KVP.Value);
+			ArchetypesTreeView->SetItemExpansion(KVP.Value, true);
+		}
+	}
+
+	ArchetypesTreeView->RequestListRefresh();
+	SelectedArchetypesListWidget->RequestListRefresh();
+}
+
+void SMassArchetypesView::OnProcessorsSelected(TConstArrayView<TSharedPtr<FMassDebuggerProcessorData>> SelectedProcessors, ESelectInfo::Type SelectInfo)
+{
+	if (DebuggerModel)
+	{
+		OnArchetypesSelected(DebuggerModel->SelectedArchetypes, ESelectInfo::Direct);
+	}
+}
+
+void SMassArchetypesView::OnArchetypesSelected(TConstArrayView<TSharedPtr<FMassDebuggerArchetypeData>> SelectedArchetypes, ESelectInfo::Type SelectInfo)
+{
+	if (SelectInfo == ESelectInfo::Direct)
+	{
+		ArchetypesTreeView->ClearSelection();
+		if (SelectedArchetypes.Num())
+		{
+			for (const TSharedPtr<FMassDebuggerArchetypeData>& ArchetypeData : SelectedArchetypes)
+			{
+				ArchetypesTreeView->SetItemSelection(ArchetypeData, true);
+			}
+
+			// scroll to the first item to make sure there's anything in view
+			ArchetypesTreeView->RequestScrollIntoView(SelectedArchetypes[0]);
+		}
+	}
+	ArchetypesTreeView->RebuildList();
+	SelectedArchetypesListWidget->RequestListRefresh();
+}
+
+#undef LOCTEXT_NAMESPACE
diff --git a/Engine/Plugins/Runtime/MassEntity/Source/MassEntityDebugger/Private/SMassArchetypesView.h b/Engine/Plugins/Runtime/MassEntity/Source/MassEntityDebugger/Private/SMassArchetypesView.h
new file mode 100644
--- /dev/null
+++ b/Engine/Plugins/Runtime/MassEntity/Source/MassEntityDebugger/Private/SMassArchetypesView.h
@@ -0,0 +1,34 @@
+// Copyright Epic Games, Inc. All Rights Reserved.
+#pragma once
+
+#include "SMassDebuggerViewBase.h"
+#include "Widgets/Views/STreeView.h"
+#include "Widgets/Views/SListView.h"
+
+struct FMassDebuggerModel;
+struct FMassDebuggerArchetypeData;
+struct FMassDebuggerProcessorData;
+
+
+class SMassArchetypesView : public SMassDebuggerViewBase
+{
+public:
+	SLATE_BEGIN_ARGS(SMassArchetypesView)
+	{
+	}
+	SLATE_END_ARGS()
+
+	void Construct(const FArguments& InArgs, TSharedRef<FMassDebuggerModel> DebuggerModel);
+
+	void HandleSelectionChanged(TSharedPtr<FMassDebuggerArchetypeData> InNode, ESelectInfo::Type InSelectInfo);
+
+protected:
+	virtual void OnRefresh() override;
+	virtual void OnProcessorsSelected(TConstArrayView<TSharedPtr<FMassDebuggerProcessorData>> SelectedProcessors, ESelectInfo::Type SelectInfo) override;
+	virtual void OnArchetypesSelected(TConstArrayView<TSharedPtr<FMassDebuggerArchetypeData>> SelectedArchetypes, ESelectInfo::Type SelectInfo) override;
+
+	TSharedPtr<STreeView<TSharedPtr<FMassDebuggerArchetypeData>>> ArchetypesTreeView;
+	TSharedPtr<SListView<TSharedPtr<FMassDebuggerArchetypeData>>> SelectedArchetypesListWidget;
+
+	TArray<TSharedPtr<FMassDebuggerArchetypeData>> TreeViewSource;
+};
diff --git a/Engine/Plugins/Runtime/MassEntity/Source/MassEntityDebugger/Private/SMassBitSet.h b/Engine/Plugins/Runtime/MassEntity/Source/MassEntityDebugger/Private/SMassBitSet.h
new file mode 100644
--- /dev/null
+++ b/Engine/Plugins/Runtime/MassEntity/Source/MassEntityDebugger/Private/SMassBitSet.h
@@ -0,0 +1,108 @@
+// Copyright Epic Games, Inc. All Rights Reserved.
+#pragma once
+
+#include "Widgets/SCompoundWidget.h"
+#include "Widgets/Text/SRichTextBlock.h"
+#include "Widgets/Text/STextBlock.h"
+#include "MassDebuggerStyle.h"
+#include "Widgets/Layout/SWrapBox.h"
+#include "Widgets/Layout/SBorder.h"
+#include "Widgets/SBoxPanel.h"
+
+
+struct FMassDebuggerQueryData;
+
+enum class EMassBitSetDisplayMode 
+{
+	ReadOnly,
+	ReadWrite,
+	MAX
+};
+
+template<typename TBitSet>
+class SMassBitSet : public SCompoundWidget
+{
+public:
+	SLATE_BEGIN_ARGS(SMassBitSet)
+		: _SlotPadding(5.f)
+		, _TextColor(FLinearColor(1.f, 1.f, 1.f))
+		, _BackgroundBrush(FMassDebuggerStyle::GetBrush("MassDebug.Fragment"))
+		{}
+		SLATE_ATTRIBUTE(FMargin, SlotPadding)
+		SLATE_ATTRIBUTE(FLinearColor, TextColor)
+		SLATE_ATTRIBUTE(const FSlateBrush*, BackgroundBrush)
+	SLATE_END_ARGS()
+
+	void Construct(const SMassBitSet::FArguments& InArgs, const FString& Label, const TBitSet& BitSet)
+	{
+		Construct(InArgs, Label, MakeArrayView(&BitSet, 1));
+	}
+
+	void Construct(const SMassBitSet::FArguments& InArgs, const FString& Label, TConstArrayView<TBitSet> BitSets, TConstArrayView<const FSlateBrush*> InBrushes = TConstArrayView<const FSlateBrush*>())
+	{		
+		/*Content = TEXT("<LargeText>Large test</>, <RichTextBlock.Bold>Bold</>");
+		Content += BitSet.DebugGetStringDesc();*/
+		TSharedRef<SWrapBox> ButtonBox = SNew(SWrapBox).UseAllottedSize(true);
+		
+		for (int i = 0; i < BitSets.Num(); ++i)
+		{
+			const FSlateBrush* Brush = InBrushes.IsValidIndex(i) ? InBrushes[i] : InArgs._BackgroundBrush.Get();
+			AddBitSet(InArgs, ButtonBox, BitSets[i], Brush);
+		}
+
+		ChildSlot
+		[
+			SNew(SHorizontalBox)
+			+ SHorizontalBox::Slot()
+			.AutoWidth()
+			[
+				SNew(STextBlock)
+				.Text(FText::FromString(Label))
+			]
+			+ SHorizontalBox::Slot()
+			[
+				ButtonBox
+			]
+		];
+	}
+
+protected:
+	void AddBitSet(const SMassBitSet::FArguments& InArgs, TSharedRef<SWrapBox>& ButtonBox, const TBitSet& BitSet, const FSlateBrush* Brush)
+	{
+		TArray<FName> TypeNames;
+		BitSet.DebugGetIndividualNames(TypeNames);
+
+		for (const FName& Name : TypeNames)
+		{
+			ButtonBox->AddSlot()
+			.Padding(InArgs._SlotPadding)
+			[
+				SNew(SBorder)
+				.ColorAndOpacity(InArgs._TextColor)
+				.BorderImage(Brush)
+				[
+					SNew(STextBlock)
+					.Text(FText::FromName(Name))
+				]
+			];
+		}
+	}
+};
+
+namespace UE::Mass::Debugger::UI
+{
+	template<typename TBitSet>
+	void AddBitSet(TSharedRef<SVerticalBox>& Box, const TBitSet& BitSetAccess, const FString& Label, const FSlateBrush* Brush)
+	{
+		if (BitSetAccess.IsEmpty() == false)
+		{
+			Box->AddSlot()
+			.AutoHeight()
+			[
+				SNew(SMassBitSet<TBitSet>, Label, BitSetAccess)
+				.BackgroundBrush(Brush)
+				.SlotPadding(5)
+			];
+		}
+	}
+}
diff --git a/Engine/Plugins/Runtime/MassEntity/Source/MassEntityDebugger/Private/SMassDebugger.cpp b/Engine/Plugins/Runtime/MassEntity/Source/MassEntityDebugger/Private/SMassDebugger.cpp
--- a/Engine/Plugins/Runtime/MassEntity/Source/MassEntityDebugger/Private/SMassDebugger.cpp
+++ b/Engine/Plugins/Runtime/MassEntity/Source/MassEntityDebugger/Private/SMassDebugger.cpp
@@ -1,12 +1,320 @@
 // Copyright Epic Games, Inc. All Rights Reserved.
 
 #include "SMassDebugger.h"
+#include "SMassProcessorsView.h"
+#include "SMassProcessingView.h"
+#include "SMassArchetypesView.h"
+#include "MassDebuggerModel.h"
+#include "Engine/Engine.h"
+#include "CoreGlobals.h"
+#include "Framework/Docking/LayoutService.h"
+#include "Framework/MultiBox/MultiBoxBuilder.h"
+#include "Widgets/Input/SComboBox.h"
+#if WITH_EDITOR
+#include "Editor.h"
+#include "UnrealEdMisc.h"
+#endif // WITH_EDITOR
 
 #define LOCTEXT_NAMESPACE "SMassDebugger"
 
+//----------------------------------------------------------------------//
+// FMassDebuggerCommands
+//----------------------------------------------------------------------//
+FMassDebuggerCommands::FMassDebuggerCommands()
+	: TCommands<FMassDebuggerCommands>("MassDebugger", LOCTEXT("MassDebuggerName", "Mass Debugger"), NAME_None, "MassDebuggerStyle")
+{ }
+
+void FMassDebuggerCommands::RegisterCommands() 
+{
+	UI_COMMAND(RefreshData, "RecacheData", "Recache data", EUserInterfaceActionType::Button, FInputChord());
+}
+
+namespace UE::Mass::Debugger::Private
+{
+	const FName ToolbarTabId("Toolbar");
+	const FName ProcessorsTabId("Processors");
+	const FName ProcessingGraphTabId("Processing Graphs");
+	const FName ArchetypesTabId("Archetypes");
+
+	bool IsSupportedWorldType(const EWorldType::Type WorldType)
+	{
+		return WorldType == EWorldType::Game || WorldType == EWorldType::Editor || WorldType == EWorldType::PIE;
+	}
+}
+
+//----------------------------------------------------------------------//
+// SMassDebugger
+//----------------------------------------------------------------------//
+SMassDebugger::SMassDebugger()
+	: SCompoundWidget(), CommandList(MakeShareable(new FUICommandList))
+	, DebuggerModel(MakeShareable(new FMassDebuggerModel))
+{
+}
+
+SMassDebugger::~SMassDebugger()
+{
+#if WITH_EDITOR
+	FEditorDelegates::EndPIE.Remove(PIEEndHandle);
+	FWorldDelegates::OnPostDuplicate.Remove(PIEWorldInitialize);
+#endif // WITH_EDITOR
+	GEngine->OnWorldAdded().Remove(OnWorldAddedHandle);
+	GEngine->OnWorldDestroyed().Remove(OnWorldDestroyedHandle);
+}
+
 void SMassDebugger::Construct(const FArguments& InArgs, const TSharedRef<SDockTab>& ConstructUnderMajorTab, const TSharedPtr<SWindow>& ConstructUnderWindow)
 {
+	using namespace UE::Mass::Debugger::Private;
+
+	BindDelegates();
+
+	const FMassDebuggerCommands& Commands = FMassDebuggerCommands::Get();
+	FUICommandList& ActionList = *CommandList;
+
+	ActionList.MapAction(Commands.RefreshData
+		, FExecuteAction::CreateLambda([this]() 
+		{
+				DebuggerModel->RefreshAll();
+		})
+		// check if there's any mass environment picked
+		, FCanExecuteAction::CreateLambda([this]()->bool { return true; })
+		, FIsActionChecked::CreateLambda([]()->bool { return false; })
+		, FIsActionButtonVisible::CreateLambda([this]()->bool { return true; }));
+
+	// Tab Spawners
+	TabManager = FGlobalTabmanager::Get()->NewTabManager(ConstructUnderMajorTab);
+	TSharedRef<FWorkspaceItem> AppMenuGroup = TabManager->AddLocalWorkspaceMenuCategory(LOCTEXT("MassDebuggerGroupName", "Mass Debugger"));
+
+	TabManager->RegisterTabSpawner(ToolbarTabId, FOnSpawnTab::CreateRaw(this, &SMassDebugger::SpawnToolbar))
+		.SetDisplayName(LOCTEXT("ToolbarTabTitle", "Toolbar"))
+		.SetGroup(AppMenuGroup);
+
+	TabManager->RegisterTabSpawner(ProcessorsTabId, FOnSpawnTab::CreateRaw(this, &SMassDebugger::SpawnProcessorsTab))
+		.SetDisplayName(LOCTEXT("ProcessorsTabTitle", "Processors"))
+		.SetGroup(AppMenuGroup);
+
+	TabManager->RegisterTabSpawner(ProcessingGraphTabId, FOnSpawnTab::CreateRaw(this, &SMassDebugger::SpawnProcessingTab))
+		.SetDisplayName(LOCTEXT("ProcessingTabTitle", "Processing Graphs"))
+		.SetGroup(AppMenuGroup);
+
+	TabManager->RegisterTabSpawner(ArchetypesTabId, FOnSpawnTab::CreateRaw(this, &SMassDebugger::SpawnArchetypesTab))
+		.SetDisplayName(LOCTEXT("ArchetypesTabTitle", "Archetypes"))
+		.SetGroup(AppMenuGroup);
+
+
+	// Default Layout
+	TSharedRef<FTabManager::FLayout> Layout = FTabManager::NewLayout("MassDebuggerLayout_v1.0")
+		->AddArea
+		(
+			FTabManager::NewPrimaryArea()
+			->SetOrientation(Orient_Vertical)
+			->Split
+			(
+				FTabManager::NewStack()
+				->AddTab(ToolbarTabId, ETabState::OpenedTab)
+				->SetHideTabWell(true)
+			)
+			->Split
+			(
+				FTabManager::NewSplitter()
+				->SetOrientation(Orient_Horizontal)
+				->Split
+				(
+					FTabManager::NewStack()
+						->AddTab(ProcessorsTabId, ETabState::OpenedTab)
+						->AddTab(ProcessingGraphTabId, ETabState::OpenedTab)
+						->AddTab(ArchetypesTabId, ETabState::OpenedTab)
+						->SetForegroundTab(ProcessorsTabId)
+				)
+			)
+		);
+
+	Layout = FLayoutSaveRestore::LoadFromConfig(GEditorLayoutIni, Layout);
+
+	ChildSlot
+		[
+			TabManager->RestoreFrom(Layout, ConstructUnderWindow).ToSharedRef()
+		];
+
+	TabManager->SetOnPersistLayout(
+		FTabManager::FOnPersistLayout::CreateStatic(
+			[](const TSharedRef<FTabManager::FLayout>& InLayout)
+			{
+#if WITH_EDITOR
+				if (FUnrealEdMisc::Get().IsSavingLayoutOnClosedAllowed() == false)
+				{
+					return;
+				}
+#endif // WITH_EDITOR
+				if (InLayout->GetPrimaryArea().Pin().IsValid())
+				{
+					FLayoutSaveRestore::SaveToConfig(GEditorLayoutIni, InLayout);
+				}
+			}
+		)
+	);
 }
+
+TSharedRef<SDockTab> SMassDebugger::SpawnToolbar(const FSpawnTabArgs& Args)
+{
+	const TSharedRef<SDockTab> MajorTab = SNew(SDockTab)
+		.TabRole(ETabRole::PanelTab)
+		.ShouldAutosize(true);
+
+	FSlimHorizontalToolBarBuilder ToolBarBuilder(CommandList, FMultiBoxCustomization::None);
+	ToolBarBuilder.BeginSection("Debugger");
+	{
+		ToolBarBuilder.AddToolBarButton(FMassDebuggerCommands::Get().RefreshData, NAME_None, LOCTEXT("RefreshData", "Refresh"), LOCTEXT("RefreshDebuggerTooltip", "Refreshes data cached by the debugger instance"));//, FSlateIcon(FLogVisualizerStyle::Get().GetStyleSetName(), TEXT("Toolbar.Update")));
+	}
+
+	RebuildEnvironmentsList();
 	
-#undef LOCTEXT_NAMESPACE
+	MajorTab->SetContent(
+		SNew(SHorizontalBox)
+		+SHorizontalBox::Slot()
+		[
+			ToolBarBuilder.MakeWidget()
+		]
+		+SHorizontalBox::Slot()
+		.HAlign(HAlign_Right)
+		.VAlign(VAlign_Center)
+		.Padding(2.f)
+		.AutoWidth()
+		[
+			SAssignNew(EnvironmentComboBox, SComboBox<TSharedPtr<FMassDebuggerEnvironment>>)
+			.OptionsSource(&EnvironmentsList)
+			.OnGenerateWidget_Lambda([](TSharedPtr<FMassDebuggerEnvironment> Item)
+				{
+					check(Item);
+					return SNew(STextBlock).Text(FText::FromString(Item->GetDisplayName()));
+				})
+			.OnSelectionChanged(this, &SMassDebugger::HandleEnvironmentChanged)
+			.ToolTipText(LOCTEXT("Environment_Tooltip", "Pick where to get the data from"))
+			[
+				SAssignNew(EnvironmentComboLabel, STextBlock)
+				.Text(LOCTEXT("PickEnvironment", "Pick Environment"))
+			]
+		]
+	);
+
+	return MajorTab;
+}
 
+void SMassDebugger::OnWorldAdded(UWorld* NewWorld)
+{
+	if (NewWorld &&  UE::Mass::Debugger::Private::IsSupportedWorldType(NewWorld->WorldType))
+	{
+		EnvironmentsList.Add(MakeShareable(new FMassDebuggerEnvironment(NewWorld)));
+	}
+}
+
+void SMassDebugger::OnWorldDestroyed(UWorld* InWorld)
+{
+	if (InWorld == nullptr || UE::Mass::Debugger::Private::IsSupportedWorldType(InWorld->WorldType) == false)
+	{
+		return;
+	}
+
+	FMassDebuggerEnvironment InEnvironment(InWorld);
+
+	if (EnvironmentsList.RemoveAll([&InEnvironment](const TSharedPtr<FMassDebuggerEnvironment>& Element)
+		{
+			return *Element.Get() == InEnvironment;
+		}) > 0)
+	{
+		if (DebuggerModel->IsCurrentEnvironment(InEnvironment))
+		{
+			EnvironmentComboLabel->SetText(FText::FromString(TEXT("Stale")));
+		}
+	}
+}
+
+void SMassDebugger::HandleEnvironmentChanged(TSharedPtr<FMassDebuggerEnvironment> Item, ESelectInfo::Type SelectInfo)
+{
+	DebuggerModel->SetEnvironment(Item);
+	EnvironmentComboLabel->SetText(DebuggerModel->GetDisplayName());
+}
+
+void SMassDebugger::RebuildEnvironmentsList()
+{
+	const TIndirectArray<FWorldContext>& WorldContexts = GEngine->GetWorldContexts();
+	EnvironmentsList.Reset();
+	for (const FWorldContext& Context : WorldContexts)
+	{
+		if (UE::Mass::Debugger::Private::IsSupportedWorldType(Context.WorldType))
+		{
+			EnvironmentsList.Add(MakeShareable(new FMassDebuggerEnvironment(Context.World())));
+		}
+	}
+}
+
+TSharedRef<SDockTab> SMassDebugger::SpawnProcessorsTab(const FSpawnTabArgs& Args)
+{
+	const TSharedRef<SDockTab> MajorTab = SNew(SDockTab)
+		.TabRole(ETabRole::PanelTab);
+
+	TSharedPtr<SWidget> TabContent = SNew(SMassProcessorsView, DebuggerModel);
+	MajorTab->SetContent(TabContent.ToSharedRef());
+
+	return MajorTab;
+}
+
+TSharedRef<SDockTab> SMassDebugger::SpawnProcessingTab(const FSpawnTabArgs& Args)
+{
+	const TSharedRef<SDockTab> MajorTab = SNew(SDockTab)
+		.TabRole(ETabRole::PanelTab);
+
+	TSharedPtr<SWidget> TabContent = SNew(SMassProcessingView, DebuggerModel);
+	MajorTab->SetContent(TabContent.ToSharedRef());
+
+	return MajorTab;
+}
+
+TSharedRef<SDockTab> SMassDebugger::SpawnArchetypesTab(const FSpawnTabArgs& Args)
+{
+	const TSharedRef<SDockTab> MajorTab = SNew(SDockTab)
+		.TabRole(ETabRole::PanelTab);
+
+	TSharedPtr<SWidget> TabContent = SNew(SMassArchetypesView, DebuggerModel);
+	MajorTab->SetContent(TabContent.ToSharedRef());
+
+	return MajorTab;
+}
+
+void SMassDebugger::BindDelegates()
+{
+#if WITH_EDITOR
+	PIEWorldInitialize = FWorldDelegates::OnPostDuplicate.AddLambda([this](UWorld* World, bool bDuplicateForPIE, TMap<UObject*, UObject*>& ReplacementMap, TArray<UObject*>& ObjectsToFixReferences)
+		{
+			if (World && bDuplicateForPIE)
+			{
+				EnvironmentsList.Add(MakeShareable(new FMassDebuggerEnvironment(World)));
+			}
+		});
+
+	PIEEndHandle = FEditorDelegates::EndPIE.AddLambda([this](const bool bIsSimulating)
+		{
+			const TIndirectArray<FWorldContext>& WorldContexts = GEngine->GetWorldContexts();
+			EnvironmentsList.Reset();
+			for (const FWorldContext& Context : WorldContexts)
+			{
+				if (Context.WorldType == EWorldType::Game || Context.WorldType == EWorldType::Editor)
+				{
+					EnvironmentsList.Add(MakeShareable(new FMassDebuggerEnvironment(Context.World())));
+				}
+				else if (Context.WorldType == EWorldType::PIE)
+				{
+					FMassDebuggerEnvironment TempEnvironment(Context.World());
+					if (DebuggerModel->IsCurrentEnvironment(TempEnvironment))
+					{
+						DebuggerModel->MarkAsStale();
+					}
+				}
+			}
+			EnvironmentComboLabel->SetText(DebuggerModel->GetDisplayName());
+		});
+#endif
+	OnWorldAddedHandle = GEngine->OnWorldAdded().AddRaw(this, &SMassDebugger::OnWorldAdded);
+	OnWorldDestroyedHandle = GEngine->OnWorldDestroyed().AddRaw(this, &SMassDebugger::OnWorldDestroyed);
+}
+
+#undef LOCTEXT_NAMESPACE
diff --git a/Engine/Plugins/Runtime/MassEntity/Source/MassEntityDebugger/Private/SMassDebugger.h b/Engine/Plugins/Runtime/MassEntity/Source/MassEntityDebugger/Private/SMassDebugger.h
--- a/Engine/Plugins/Runtime/MassEntity/Source/MassEntityDebugger/Private/SMassDebugger.h
+++ b/Engine/Plugins/Runtime/MassEntity/Source/MassEntityDebugger/Private/SMassDebugger.h
@@ -2,7 +2,26 @@
 #pragma once
 
 #include "Widgets/Docking/SDockTab.h"
+#include "Framework/Commands/Commands.h"
 
+struct FMassDebuggerEnvironment;
+struct FMassDebuggerModel;
+template<typename T> class SComboBox;
+
+
+class FMassDebuggerCommands : public TCommands<FMassDebuggerCommands>
+{
+public:
+	/** Default constructor. */
+	FMassDebuggerCommands();
+
+public:
+	// TCommands interface
+	virtual void RegisterCommands() override;
+public:
+
+	TSharedPtr<FUICommandInfo> RefreshData;
+};
 
 class SMassDebuggerTab : public SDockTab
 {
@@ -20,6 +39,8 @@ public:
 	SLATE_END_ARGS()
 
 public:
+	SMassDebugger();
+	virtual ~SMassDebugger();
 
 	/**
 	* Constructs the application.
@@ -31,4 +52,36 @@ public:
 	void Construct(const FArguments& InArgs, const TSharedRef<SDockTab>& ConstructUnderMajorTab, const TSharedPtr<SWindow>& ConstructUnderWindow);
 
 	virtual bool SupportsKeyboardFocus() const override { return true; }
+
+protected:
+	TSharedRef<SDockTab> SpawnToolbar(const FSpawnTabArgs& Args);
+	TSharedRef<SDockTab> SpawnProcessorsTab(const FSpawnTabArgs& Args);
+	TSharedRef<SDockTab> SpawnProcessingTab(const FSpawnTabArgs& Args);
+	TSharedRef<SDockTab> SpawnArchetypesTab(const FSpawnTabArgs& Args);
+
+	void BindDelegates();
+	void OnWorldAdded(UWorld* NewWorld);
+	void OnWorldDestroyed(UWorld* InWorld);
+	void HandleEnvironmentChanged(TSharedPtr<FMassDebuggerEnvironment> Item, ESelectInfo::Type SelectInfo);
+	void RebuildEnvironmentsList();
+
+	// Holds the list of UI commands.
+	TSharedRef<FUICommandList> CommandList;
+
+	// Holds the tab manager that manages the front-end's tabs.
+	TSharedPtr<FTabManager> TabManager;
+
+	TSharedPtr<SComboBox<TSharedPtr<FMassDebuggerEnvironment>>> EnvironmentComboBox;
+	TSharedPtr<STextBlock> EnvironmentComboLabel;
+
+	TArray<TSharedPtr<FMassDebuggerEnvironment>> EnvironmentsList;
+
+	TSharedRef<FMassDebuggerModel> DebuggerModel;
+
+#if WITH_EDITOR
+	FDelegateHandle PIEEndHandle;
+	FDelegateHandle PIEWorldInitialize;
+#endif // WITH_EDITOR
+	FDelegateHandle OnWorldAddedHandle;
+	FDelegateHandle OnWorldDestroyedHandle;
 };
diff --git a/Engine/Plugins/Runtime/MassEntity/Source/MassEntityDebugger/Private/SMassDebuggerViewBase.cpp b/Engine/Plugins/Runtime/MassEntity/Source/MassEntityDebugger/Private/SMassDebuggerViewBase.cpp
new file mode 100644
--- /dev/null
+++ b/Engine/Plugins/Runtime/MassEntity/Source/MassEntityDebugger/Private/SMassDebuggerViewBase.cpp
@@ -0,0 +1,26 @@
+// Copyright Epic Games, Inc. All Rights Reserved.
+
+#include "SMassDebuggerViewBase.h"
+#include "MassDebuggerModel.h"
+
+//----------------------------------------------------------------------//
+// SMassDebuggerViewBase
+//----------------------------------------------------------------------//
+SMassDebuggerViewBase::~SMassDebuggerViewBase()
+{
+	if (DebuggerModel)
+	{
+		DebuggerModel->OnRefreshDelegate.Remove(OnRefreshHandle);
+		DebuggerModel->OnProcessorsSelectedDelegate.Remove(OnProcessorsSelectedHandle);
+		DebuggerModel->OnArchetypesSelectedDelegate.Remove(OnArchetypesSelectedHandle);
+	}
+}
+
+void SMassDebuggerViewBase::Initialize(TSharedRef<FMassDebuggerModel> InDebuggerModel)
+{
+	DebuggerModel = InDebuggerModel;
+
+	OnRefreshHandle = DebuggerModel->OnRefreshDelegate.AddRaw(this, &SMassDebuggerViewBase::OnRefresh);
+	OnProcessorsSelectedHandle = DebuggerModel->OnProcessorsSelectedDelegate.AddRaw(this, &SMassDebuggerViewBase::OnProcessorsSelected);
+	OnArchetypesSelectedHandle = DebuggerModel->OnArchetypesSelectedDelegate.AddRaw(this, &SMassDebuggerViewBase::OnArchetypesSelected);
+}
diff --git a/Engine/Plugins/Runtime/MassEntity/Source/MassEntityDebugger/Private/SMassDebuggerViewBase.h b/Engine/Plugins/Runtime/MassEntity/Source/MassEntityDebugger/Private/SMassDebuggerViewBase.h
new file mode 100644
--- /dev/null
+++ b/Engine/Plugins/Runtime/MassEntity/Source/MassEntityDebugger/Private/SMassDebuggerViewBase.h
@@ -0,0 +1,31 @@
+// Copyright Epic Games, Inc. All Rights Reserved.
+#pragma once
+
+#include "Widgets/SCompoundWidget.h"
+
+
+struct FMassDebuggerProcessorData;
+struct FMassDebuggerArchetypeData;
+struct FMassDebuggerModel;
+
+class SMassDebuggerViewBase : public SCompoundWidget
+{
+public:
+	SLATE_BEGIN_ARGS(SMassDebuggerViewBase){}
+	SLATE_END_ARGS()
+
+	virtual ~SMassDebuggerViewBase();
+
+protected:
+	void Initialize(TSharedRef<FMassDebuggerModel> InDebuggerModel);
+
+	virtual void OnRefresh() = 0;
+	virtual void OnProcessorsSelected(TConstArrayView<TSharedPtr<FMassDebuggerProcessorData>> SelectedProcessors, ESelectInfo::Type SelectInfo) = 0;
+	virtual void OnArchetypesSelected(TConstArrayView<TSharedPtr<FMassDebuggerArchetypeData>> SelectedArchetypes, ESelectInfo::Type SelectInfo) = 0;
+
+	FDelegateHandle OnRefreshHandle;
+	FDelegateHandle OnProcessorsSelectedHandle;
+	FDelegateHandle OnArchetypesSelectedHandle;
+
+	TSharedPtr<FMassDebuggerModel> DebuggerModel;
+};
diff --git a/Engine/Plugins/Runtime/MassEntity/Source/MassEntityDebugger/Private/SMassProcessingGraphView.cpp b/Engine/Plugins/Runtime/MassEntity/Source/MassEntityDebugger/Private/SMassProcessingGraphView.cpp
new file mode 100644
--- /dev/null
+++ b/Engine/Plugins/Runtime/MassEntity/Source/MassEntityDebugger/Private/SMassProcessingGraphView.cpp
@@ -0,0 +1,269 @@
+// Copyright Epic Games, Inc. All Rights Reserved.
+
+#include "SMassProcessingGraphView.h"
+#include "MassDebuggerModel.h"
+#include "Widgets/Layout/SScrollBox.h"
+#include "Widgets/Layout/SBorder.h"
+#include "Styling/SlateBrush.h"
+#include "MassDebuggerStyle.h"
+#include "Widgets/Views/STreeView.h"
+
+#define LOCTEXT_NAMESPACE "SMassDebugger"
+
+
+//----------------------------------------------------------------------//
+// FMassDebuggerProcessingGraphNodeTreeItem
+//----------------------------------------------------------------------//
+FMassDebuggerProcessingGraphNodeTreeItem::FMassDebuggerProcessingGraphNodeTreeItem(const FMassDebuggerProcessingGraphNode& InNode)
+	: Node(InNode)
+{
+
+}
+
+//----------------------------------------------------------------------//
+// SMassProcessingGraphTableRow
+//----------------------------------------------------------------------//
+using FMassDebuggerProcessingGraphNodeTreeItemPtr = TSharedPtr<FMassDebuggerProcessingGraphNodeTreeItem, ESPMode::ThreadSafe>;
+
+class SMassProcessingGraphTableRow : public STableRow<FMassDebuggerProcessingGraphNodeTreeItemPtr>
+{
+public:
+	SLATE_BEGIN_ARGS(SMassProcessingGraphTableRow) { }
+	SLATE_END_ARGS()
+
+	void Construct(const FArguments& InArgs, const TSharedPtr<STableViewBase>& InOwnerTableView, const FMassDebuggerProcessingGraphNodeTreeItemPtr InEntryItem)
+	{
+		Item = InEntryItem;
+
+		STableRow<FMassDebuggerProcessingGraphNodeTreeItemPtr>::ConstructInternal(STableRow::FArguments()
+			.Padding(5.0f)
+			, InOwnerTableView.ToSharedRef());
+	
+		ChildSlot
+			[
+				SNew(SHorizontalBox)
+				+ SHorizontalBox::Slot()
+				.VAlign(VAlign_Fill)
+				.HAlign(HAlign_Left)
+				.AutoWidth()
+				[
+					SNew(SExpanderArrow, SharedThis(this))
+					.ShouldDrawWires(true)
+					.IndentAmount(32)
+					.BaseIndentLevel(0)
+				]
+
+				+ SHorizontalBox::Slot()
+				[
+					SNew(SBox)
+					.Padding(4, 2)
+					[
+						
+						SNew(STextBlock)
+						.Text(Item->Node.GetLabel())
+						.ColorAndOpacity_Lambda([this]()
+							{
+								switch (Item->Node.GraphNodeSelection)
+								{
+									case EMassDebuggerProcessingGraphNodeSelection::None: // fall through on purpose
+										return FLinearColor::White;
+										break;
+									case EMassDebuggerProcessingGraphNodeSelection::WaitFor:
+										return FLinearColor::Green;
+										break;
+									case EMassDebuggerProcessingGraphNodeSelection::Block:
+										return FLinearColor::Red;
+										break;
+									default:
+										return FLinearColor::White;
+										break;
+								}
+								return FLinearColor::Gray;
+							})
+					]
+				]
+				+ SHorizontalBox::Slot()
+				.VAlign(VAlign_Fill)
+				.HAlign(HAlign_Left)
+				.AutoWidth()
+				[
+					SNew(SBox)
+				]
+			];
+	}
+
+	FMassDebuggerProcessingGraphNodeTreeItemPtr Item;
+};
+
+//----------------------------------------------------------------------//
+// SMassProcessingGraphView
+//----------------------------------------------------------------------//
+void SMassProcessingGraphView::Construct(const FArguments& InArgs, TSharedRef<FMassDebuggerModel> InDebuggerModel)
+{
+	Initialize(InDebuggerModel);
+
+	OffsetPerLevel = InArgs._OffsetPerLevel.Get();
+
+	ItemsBox = SNew(SVerticalBox);
+
+	ChildSlot
+	[
+		SNew(SScrollBox)
+		.Orientation(Orient_Horizontal)
+		+ SScrollBox::Slot()
+		[
+			SAssignNew(GraphNodesTree, STreeView<TSharedPtr<FMassDebuggerProcessingGraphNodeTreeItem>>)
+			.SelectionMode(ESelectionMode::Single)
+			.TreeItemsSource(&RootNodes)
+			.TreeViewStyle(&FAppStyle::Get().GetWidgetStyle<FTableViewStyle>("PropertyTable.InViewport.ListView"))
+			.OnGenerateRow_Lambda([this](TSharedPtr<FMassDebuggerProcessingGraphNodeTreeItem> Item, const TSharedPtr<STableViewBase>& OwnerTable)
+				{
+					return SNew(SMassProcessingGraphTableRow, OwnerTable, Item);
+				})
+			.OnGetChildren_Lambda([](TSharedPtr<FMassDebuggerProcessingGraphNodeTreeItem> InItem, TArray<TSharedPtr<FMassDebuggerProcessingGraphNodeTreeItem>>& OutChildren)
+				{
+					if (InItem->ChildItems.Num())
+					{
+						OutChildren.Append(InItem->ChildItems);
+					}
+				})
+				.OnSelectionChanged(this, &SMassProcessingGraphView::HandleSelectionChanged)
+		]
+	];
+}
+
+void SMassProcessingGraphView::Display(TSharedPtr<FMassDebuggerProcessingGraph> InProcessingGraphData)
+{
+	if (!InProcessingGraphData)
+	{
+		return;
+	}
+
+	ItemsBox->ClearChildren();
+	RootNodes.Reset();
+	AllNodes.Reset();
+
+	if (InProcessingGraphData->GraphNodes.Num() == 0)
+	{
+		GraphNodesTree->RebuildList();
+		return;
+	}
+
+	const FSlateBrush* Brush = FMassDebuggerStyle::GetBrush("MassDebug.Fragment");
+
+	TArray<int32> Offset;
+	Offset.Reserve(InProcessingGraphData->GraphNodes.Num());
+	AllNodes.Reserve(InProcessingGraphData->GraphNodes.Num());
+
+	for (const FMassDebuggerProcessingGraphNode& Node : InProcessingGraphData->GraphNodes)
+	{
+		const int32 ThisNodesIndex = AllNodes.Num();
+		TSharedPtr<FMassDebuggerProcessingGraphNodeTreeItem>& SharedNode = AllNodes.Add_GetRef(MakeShareable(new FMassDebuggerProcessingGraphNodeTreeItem(Node)));
+
+		int32 MyOffset = INDEX_NONE;
+		if (Node.WaitForNodes.Num())
+		{
+			int32 LastDependencyIndex = INDEX_NONE;
+			for (const int32 DependencyIndex : Node.WaitForNodes)
+			{
+				AllNodes[DependencyIndex]->Node.BlockNodes.Add(ThisNodesIndex);
+
+				// note that we're accepting == offsets as well to attach to the last found item
+				if (MyOffset <= Offset[DependencyIndex])
+				{
+					MyOffset = Offset[DependencyIndex];
+					LastDependencyIndex = DependencyIndex;
+				}
+			}
+			check(LastDependencyIndex != INDEX_NONE);
+			AllNodes[LastDependencyIndex]->ChildItems.Add(SharedNode);
+		}
+		else
+		{
+			RootNodes.Add(SharedNode);
+		}
+		MyOffset++;
+		Offset.Add(MyOffset);
+	}
+
+	GraphNodesTree->RebuildList();
+
+	for (TSharedPtr<FMassDebuggerProcessingGraphNodeTreeItem>& Node : AllNodes)
+	{
+		GraphNodesTree->SetItemExpansion(Node, true);
+	}
+}
+
+void SMassProcessingGraphView::HandleSelectionChanged(TSharedPtr<FMassDebuggerProcessingGraphNodeTreeItem> InNode, ESelectInfo::Type InSelectInfo)
+{
+	if (InSelectInfo == ESelectInfo::Direct)
+	{
+		return;
+	}
+
+	ClearSelection();
+
+	if (InNode)
+	{
+		DebuggerModel->SelectProcessor(InNode->Node.ProcessorData);
+		MarkDependencies(InNode->Node);
+	}
+}
+
+void SMassProcessingGraphView::OnRefresh()
+{
+	// nothing to do here, but we need to implement this function since it's pure virtual in the parent class.
+}
+
+void SMassProcessingGraphView::OnProcessorsSelected(TConstArrayView<TSharedPtr<FMassDebuggerProcessorData>> SelectedProcessors, ESelectInfo::Type)
+{
+	ClearSelection();
+	if (SelectedProcessors.Num() && SelectedProcessors[0].IsValid())
+	{
+		const uint32 ProcessorHash = SelectedProcessors[0]->ProcessorHash;
+		TSharedPtr<FMassDebuggerProcessingGraphNodeTreeItem>* MatchingItem = AllNodes.FindByPredicate([ProcessorHash](const TSharedPtr<FMassDebuggerProcessingGraphNodeTreeItem>& Element) -> bool
+		{
+			return Element.IsValid() && Element->Node.ProcessorData && Element->Node.ProcessorData->ProcessorHash == ProcessorHash;
+		});
+
+		if (MatchingItem)
+		{
+			GraphNodesTree->SetSelection(*MatchingItem);
+			GraphNodesTree->RequestScrollIntoView(*MatchingItem);
+			MarkDependencies((*MatchingItem)->Node);
+			GraphNodesTree->RebuildList();
+		}
+	}
+}
+
+void SMassProcessingGraphView::OnArchetypesSelected(TConstArrayView<TSharedPtr<FMassDebuggerArchetypeData>> SelectedArchetypes, ESelectInfo::Type)
+{
+	if (DebuggerModel)
+	{
+		OnProcessorsSelected(DebuggerModel->SelectedProcessors, ESelectInfo::Direct);
+	}
+}
+
+void SMassProcessingGraphView::ClearSelection()
+{
+	GraphNodesTree->ClearSelection();
+
+	for (TSharedPtr<FMassDebuggerProcessingGraphNodeTreeItem>& Node : AllNodes)
+	{
+		Node->Node.GraphNodeSelection = EMassDebuggerProcessingGraphNodeSelection::None;
+	}
+}
+
+void SMassProcessingGraphView::MarkDependencies(const FMassDebuggerProcessingGraphNode& Node)
+{
+	for (const int32 DependencyIndex : Node.WaitForNodes)
+	{
+		AllNodes[DependencyIndex]->Node.GraphNodeSelection = EMassDebuggerProcessingGraphNodeSelection::WaitFor;
+	}
+	for (const int32 DependencyIndex : Node.BlockNodes)
+	{
+		AllNodes[DependencyIndex]->Node.GraphNodeSelection = EMassDebuggerProcessingGraphNodeSelection::Block;
+	}
+}
+
+#undef LOCTEXT_NAMESPACE
diff --git a/Engine/Plugins/Runtime/MassEntity/Source/MassEntityDebugger/Private/SMassProcessingGraphView.h b/Engine/Plugins/Runtime/MassEntity/Source/MassEntityDebugger/Private/SMassProcessingGraphView.h
new file mode 100644
--- /dev/null
+++ b/Engine/Plugins/Runtime/MassEntity/Source/MassEntityDebugger/Private/SMassProcessingGraphView.h
@@ -0,0 +1,53 @@
+// Copyright Epic Games, Inc. All Rights Reserved.
+#pragma once
+
+#include "SMassDebuggerViewBase.h"
+#include "MassDebuggerModel.h"
+
+
+struct FMassDebuggerProcessingGraph;
+struct FMassDebuggerModel;
+template<typename T> class STreeView;
+
+class SHeaderRow;
+
+struct FMassDebuggerProcessingGraphNodeTreeItem
+{
+	FMassDebuggerProcessingGraphNodeTreeItem(const FMassDebuggerProcessingGraphNode& InNode);
+
+	FMassDebuggerProcessingGraphNode Node;
+	TArray<TSharedPtr<FMassDebuggerProcessingGraphNodeTreeItem>> ChildItems;
+};
+
+class SMassProcessingGraphView : public SMassDebuggerViewBase
+{
+public:
+	SLATE_BEGIN_ARGS(SMassProcessingGraphView)
+		: _OffsetPerLevel(10.f)
+		{}
+		SLATE_ATTRIBUTE(float, OffsetPerLevel)
+	SLATE_END_ARGS()
+
+	void Construct(const FArguments& InArgs, TSharedRef<FMassDebuggerModel> InDebuggerModel);
+
+	void Display(TSharedPtr<FMassDebuggerProcessingGraph> InProcessingGraphData);
+
+protected:
+	void HandleSelectionChanged(TSharedPtr<FMassDebuggerProcessingGraphNodeTreeItem> InNode, ESelectInfo::Type InSelectInfo);
+
+	virtual void OnRefresh() override;
+	virtual void OnProcessorsSelected(TConstArrayView<TSharedPtr<FMassDebuggerProcessorData>> SelectedProcessors, ESelectInfo::Type) override;
+	virtual void OnArchetypesSelected(TConstArrayView<TSharedPtr<FMassDebuggerArchetypeData>> SelectedArchetypes, ESelectInfo::Type) override;
+
+	void ClearSelection();
+	void MarkDependencies(const FMassDebuggerProcessingGraphNode& Node);
+
+	TSharedPtr<FMassDebuggerProcessingGraph> ProcessingGraphData;
+	TSharedPtr<SVerticalBox> ItemsBox;
+	float OffsetPerLevel;
+
+	TArray<TSharedPtr<FMassDebuggerProcessingGraphNodeTreeItem>> AllNodes;
+	TArray<TSharedPtr<FMassDebuggerProcessingGraphNodeTreeItem>> RootNodes;
+
+	TSharedPtr<STreeView<TSharedPtr<FMassDebuggerProcessingGraphNodeTreeItem>>> GraphNodesTree;
+};
diff --git a/Engine/Plugins/Runtime/MassEntity/Source/MassEntityDebugger/Private/SMassProcessingView.cpp b/Engine/Plugins/Runtime/MassEntity/Source/MassEntityDebugger/Private/SMassProcessingView.cpp
new file mode 100644
--- /dev/null
+++ b/Engine/Plugins/Runtime/MassEntity/Source/MassEntityDebugger/Private/SMassProcessingView.cpp
@@ -0,0 +1,81 @@
+// Copyright Epic Games, Inc. All Rights Reserved.
+
+#include "SMassProcessingView.h"
+#include "MassDebuggerModel.h"
+#include "SMassProcessingGraphView.h"
+#include "Widgets/Views/SListView.h"
+
+#define LOCTEXT_NAMESPACE "SMassDebugger"
+
+//----------------------------------------------------------------------//
+// SMassProcessingGraphListTableRow
+//----------------------------------------------------------------------//
+using FMassDebuggerProcessingGraphPtr = TSharedPtr<FMassDebuggerProcessingGraph, ESPMode::ThreadSafe>;
+
+class SMassProcessingGraphListTableRow : public STableRow<FMassDebuggerProcessingGraphPtr>
+{
+public:
+	SLATE_BEGIN_ARGS(SMassProcessingGraphListTableRow) { }
+	SLATE_END_ARGS()
+
+	void Construct(const FArguments& InArgs, const TSharedRef<STableViewBase>& InOwnerTableView, const FMassDebuggerProcessingGraphPtr InEntryItem)
+	{
+		Item = InEntryItem;
+
+		STableRow<FMassDebuggerProcessingGraphPtr>::Construct(STableRow<FMassDebuggerProcessingGraphPtr>::FArguments(), InOwnerTableView);
+
+		ChildSlot
+			[
+				SNew(STextBlock)
+				.Text(FText::FromString(Item->Label))
+			];
+	}
+
+	FMassDebuggerProcessingGraphPtr Item;
+};
+
+//----------------------------------------------------------------------//
+// SMassProcessingView
+//----------------------------------------------------------------------//
+SMassProcessingView::~SMassProcessingView()
+{
+	DebuggerModel->OnRefreshDelegate.Remove(OnRefreshHandle);
+}
+
+void SMassProcessingView::Construct(const FArguments& InArgs, TSharedRef<FMassDebuggerModel> InDebuggerModel)
+{
+	DebuggerModel = InDebuggerModel;
+
+	ChildSlot
+	[
+		SNew(SHorizontalBox)
+		+ SHorizontalBox::Slot()
+		.AutoWidth()
+		[
+			SAssignNew(GraphsListWidget, SListView<TSharedPtr<FMassDebuggerProcessingGraph>>)
+			.ListItemsSource(&DebuggerModel->CachedProcessingGraphs)
+			.SelectionMode(ESelectionMode::Single)
+			.OnGenerateRow_Lambda([](TSharedPtr<FMassDebuggerProcessingGraph> Item, const TSharedPtr<STableViewBase>& OwnerTable)
+				{
+					return SNew(SMassProcessingGraphListTableRow, OwnerTable.ToSharedRef(), Item);
+				})
+			.OnSelectionChanged_Lambda([this](TSharedPtr<FMassDebuggerProcessingGraph> SelectedItem, ESelectInfo::Type SelectInfo)
+				{
+					ProcessingGraphWidget->Display(SelectedItem);
+				})
+		]
+		+ SHorizontalBox::Slot()
+		[
+			SAssignNew(ProcessingGraphWidget, SMassProcessingGraphView, InDebuggerModel)
+		]
+	];
+
+	OnRefreshHandle = DebuggerModel->OnRefreshDelegate.AddRaw(this, &SMassProcessingView::OnRefresh);
+}
+
+void SMassProcessingView::OnRefresh()
+{
+	GraphsListWidget->RequestListRefresh();
+}
+
+#undef LOCTEXT_NAMESPACE
diff --git a/Engine/Plugins/Runtime/MassEntity/Source/MassEntityDebugger/Private/SMassProcessingView.h b/Engine/Plugins/Runtime/MassEntity/Source/MassEntityDebugger/Private/SMassProcessingView.h
new file mode 100644
--- /dev/null
+++ b/Engine/Plugins/Runtime/MassEntity/Source/MassEntityDebugger/Private/SMassProcessingView.h
@@ -0,0 +1,33 @@
+// Copyright Epic Games, Inc. All Rights Reserved.
+#pragma once
+
+#include "Widgets/SCompoundWidget.h"
+
+
+struct FMassDebuggerModel;
+struct FMassDebuggerProcessingGraph;
+class SMassProcessingGraphView;
+template<typename T> class SListView;
+
+class SMassProcessingView : public SCompoundWidget
+{
+public:
+
+	SLATE_BEGIN_ARGS(SMassProcessingView){}
+	SLATE_END_ARGS()
+
+	~SMassProcessingView();
+
+	void Construct(const FArguments& InArgs, TSharedRef<FMassDebuggerModel> InDebuggerModel);
+
+protected:
+	void OnRefresh();
+
+	TSharedPtr<FMassDebuggerModel> DebuggerModel;
+
+	TSharedPtr<SListView<TSharedPtr<FMassDebuggerProcessingGraph>>> GraphsListWidget;
+
+	TSharedPtr<SMassProcessingGraphView> ProcessingGraphWidget;
+
+	FDelegateHandle OnRefreshHandle;
+};
diff --git a/Engine/Plugins/Runtime/MassEntity/Source/MassEntityDebugger/Private/SMassProcessor.cpp b/Engine/Plugins/Runtime/MassEntity/Source/MassEntityDebugger/Private/SMassProcessor.cpp
new file mode 100644
--- /dev/null
+++ b/Engine/Plugins/Runtime/MassEntity/Source/MassEntityDebugger/Private/SMassProcessor.cpp
@@ -0,0 +1,81 @@
+// Copyright Epic Games, Inc. All Rights Reserved.
+
+#include "SMassProcessor.h"
+#include "MassDebuggerModel.h"
+#include "MassDebuggerStyle.h"
+#include "SMassQuery.h"
+#include "Widgets/Text/SRichTextBlock.h"
+#include "Styling/AppStyle.h"
+
+
+#define LOCTEXT_NAMESPACE "SMassDebugger"
+
+
+//----------------------------------------------------------------------//
+// SMassQueryTableRow
+//----------------------------------------------------------------------//
+using FMassDebuggerQueryDataPtr = TSharedPtr<FMassDebuggerQueryData, ESPMode::ThreadSafe>;
+
+class SMassQueryTableRow : public STableRow<FMassDebuggerQueryDataPtr>
+{
+public:
+	SLATE_BEGIN_ARGS(SMassQueryTableRow) {}
+	SLATE_END_ARGS()
+
+	void Construct(const FArguments& InArgs, const TSharedRef<STableViewBase>& InOwnerTableView, const FMassDebuggerQueryDataPtr InEntryItem)
+	{
+		Item = InEntryItem;
+
+		STableRow<FMassDebuggerQueryDataPtr>::Construct(STableRow<FMassDebuggerQueryDataPtr>::FArguments(), InOwnerTableView);
+
+		ChildSlot
+			[
+				SNew(SMassQuery, Item)
+			];
+	}
+
+	FMassDebuggerQueryDataPtr Item;
+};
+
+//----------------------------------------------------------------------//
+// SMassProcessor
+//----------------------------------------------------------------------//
+void SMassProcessor::Construct(const FArguments& InArgs, TSharedPtr<FMassDebuggerProcessorData> InProcessorData)
+{
+	ProcessorData = InProcessorData;
+	if (!ProcessorData)
+	{
+		return;
+	}
+
+	TSharedRef<SVerticalBox> Box = SNew(SVerticalBox);
+	Box->AddSlot()
+	.AutoHeight()
+	[
+		SNew(SBorder)
+		.Padding(FMargin(10, 5))
+		[
+			SNew(SRichTextBlock)
+			.Text(FText::FromString(ProcessorData->Label))
+			.DecoratorStyleSet(&FAppStyle::Get())
+			.TextStyle(FAppStyle::Get(), "LargeText")
+		]
+	];
+	
+	for (TSharedPtr<FMassDebuggerQueryData>& QueryData : ProcessorData->Queries)
+	{
+		Box->AddSlot()
+		.AutoHeight()
+		[
+			SNew(SMassQuery, QueryData)
+		];
+	}
+
+	ChildSlot
+	[
+		Box
+	];
+}
+
+#undef LOCTEXT_NAMESPACE
+
diff --git a/Engine/Plugins/Runtime/MassEntity/Source/MassEntityDebugger/Private/SMassProcessor.h b/Engine/Plugins/Runtime/MassEntity/Source/MassEntityDebugger/Private/SMassProcessor.h
new file mode 100644
--- /dev/null
+++ b/Engine/Plugins/Runtime/MassEntity/Source/MassEntityDebugger/Private/SMassProcessor.h
@@ -0,0 +1,20 @@
+// Copyright Epic Games, Inc. All Rights Reserved.
+#pragma once
+
+#include "Widgets/SCompoundWidget.h"
+
+
+struct FMassDebuggerProcessorData;
+
+class SMassProcessor : public SCompoundWidget
+{
+public:
+
+	SLATE_BEGIN_ARGS(SMassProcessor){}
+	SLATE_END_ARGS()
+
+	void Construct(const FArguments& InArgs, TSharedPtr<FMassDebuggerProcessorData> InProcessorData);
+
+protected:
+	TSharedPtr<FMassDebuggerProcessorData> ProcessorData;
+};
diff --git a/Engine/Plugins/Runtime/MassEntity/Source/MassEntityDebugger/Private/SMassProcessorsView.cpp b/Engine/Plugins/Runtime/MassEntity/Source/MassEntityDebugger/Private/SMassProcessorsView.cpp
new file mode 100644
--- /dev/null
+++ b/Engine/Plugins/Runtime/MassEntity/Source/MassEntityDebugger/Private/SMassProcessorsView.cpp
@@ -0,0 +1,207 @@
+// Copyright Epic Games, Inc. All Rights Reserved.
+
+#include "SMassProcessorsView.h"
+#include "MassDebugger.h"
+#include "SMassDebugger.h"
+#include "SMassProcessor.h"
+#include "MassDebuggerModel.h"
+#include "MassEntitySettings.h"
+#include "Widgets/Layout/SScrollBox.h"
+#include "Widgets/Text/SRichTextBlock.h"
+#include "Widgets/Views/SListView.h"
+#include "Widgets/SBoxPanel.h"
+
+#define LOCTEXT_NAMESPACE "SMassDebugger"
+
+namespace UE::Mass::Debug::UI::Private
+{
+	const FText PickProcessorLabel = FText::FromString(TEXT("Pick a processor from the list"));
+	const FText MissingDebugData = FText::FromString(TEXT("Missing debug data"));
+}
+
+//----------------------------------------------------------------------//
+// SMassProcessorListTableRow
+//----------------------------------------------------------------------//
+using FMassDebuggerProcessorDataPtr = TSharedPtr<FMassDebuggerProcessorData, ESPMode::ThreadSafe>;
+
+class SMassProcessorListTableRow : public STableRow<FMassDebuggerProcessorDataPtr>
+{
+public:
+	SLATE_BEGIN_ARGS(SMassProcessorListTableRow) { }
+	SLATE_END_ARGS()
+
+	void Construct(const FArguments& InArgs, const TSharedRef<STableViewBase>& InOwnerTableView, const FMassDebuggerProcessorDataPtr InEntryItem)
+	{
+		if (!InEntryItem.IsValid())
+		{
+			return;
+		}
+
+		Item = InEntryItem;
+		
+		STableRow<FMassDebuggerProcessorDataPtr>::Construct(STableRow<FMassDebuggerProcessorDataPtr>::FArguments(), InOwnerTableView);
+		
+		ChildSlot
+		[
+			SNew(STextBlock)
+			.Text(FText::FromString(Item->Label))
+		];
+	}
+
+	FMassDebuggerProcessorDataPtr Item;
+};
+
+class SMassProcessorTableRow : public STableRow<FMassDebuggerProcessorDataPtr>
+{
+public:
+	SLATE_BEGIN_ARGS(SMassProcessorTableRow) { }
+	SLATE_END_ARGS()
+
+	void Construct(const FArguments& InArgs, const TSharedRef<STableViewBase>& InOwnerTableView, const FMassDebuggerProcessorDataPtr InEntryItem)
+	{
+		Item = InEntryItem;
+
+		STableRow<FMassDebuggerProcessorDataPtr>::Construct(STableRow<FMassDebuggerProcessorDataPtr>::FArguments(), InOwnerTableView);
+
+		ChildSlot
+		[
+			SNew(SMassProcessor, Item)
+		];
+	}
+
+	FMassDebuggerProcessorDataPtr Item;
+};
+
+//----------------------------------------------------------------------//
+// SMassProcessorsView
+//----------------------------------------------------------------------//
+void SMassProcessorsView::Construct(const FArguments& InArgs, TSharedRef<FMassDebuggerModel> InDebuggerModel)
+{
+	Initialize(InDebuggerModel);
+
+	ProcessorsBox = SNew(SVerticalBox);
+	for (TSharedPtr<FMassDebuggerProcessorData>& ProcessorData : DebuggerModel->SelectedProcessors)
+	{
+		ProcessorsBox->AddSlot()
+		[
+			SNew(SMassProcessor, ProcessorData)
+		];
+	}
+
+	ChildSlot
+		[
+			SNew(SHorizontalBox)
+			+ SHorizontalBox::Slot()
+			.FillWidth(0.3)
+			[
+				SAssignNew(ProcessorsListWidget, SListView<TSharedPtr<FMassDebuggerProcessorData> >)
+					.ItemHeight(20)
+					.ListItemsSource(&DebuggerModel->CachedProcessors)
+					.SelectionMode(ESelectionMode::Multi)
+					.OnSelectionChanged(this, &SMassProcessorsView::ProcessorListSelectionChanged)
+					.OnGenerateRow_Lambda([](TSharedPtr<FMassDebuggerProcessorData> Item, const TSharedRef<STableViewBase>& OwnerTable)
+						{
+							return SNew(SMassProcessorListTableRow, OwnerTable, Item);
+						})
+			]
+			+ SHorizontalBox::Slot()
+			.FillWidth(0.7)
+			[
+				SNew(SScrollBox)
+				.Orientation(Orient_Vertical)
+				+SScrollBox::Slot()
+				[
+					ProcessorsBox.ToSharedRef()
+				]
+			]
+		];
+
+	PopulateProcessorList();
+}
+
+
+void SMassProcessorsView::ProcessorListSelectionChanged(TSharedPtr<FMassDebuggerProcessorData> SelectedItem, ESelectInfo::Type SelectInfo)
+{
+	if (SelectInfo == ESelectInfo::Direct)
+	{
+		return;
+	}
+	
+	check(DebuggerModel);
+
+	TArray<TSharedPtr<FMassDebuggerProcessorData>> CurrentlySelectedProcessors;
+	ProcessorsListWidget->GetSelectedItems(CurrentlySelectedProcessors);
+	DebuggerModel->SelectProcessors(CurrentlySelectedProcessors, SelectInfo);
+}
+
+void SMassProcessorsView::OnProcessorsSelected(TConstArrayView<TSharedPtr<FMassDebuggerProcessorData>> SelectedProcessors, ESelectInfo::Type SelectInfo)
+{
+	using namespace UE::Mass::Debug::UI::Private;
+	using namespace UE::Mass::Debug;
+	
+	if (!DebuggerModel)
+	{
+		return;
+	}
+
+	if (SelectInfo == ESelectInfo::Direct)
+	{
+		ProcessorsListWidget->ClearSelection();
+	}
+	ProcessorsBox->ClearChildren();
+
+	if (SelectedProcessors.Num())
+	{		
+		if (SelectInfo == ESelectInfo::Direct)
+		{
+			ProcessorsListWidget->SetItemSelection(SelectedProcessors, /*bSelected=*/true, ESelectInfo::OnMouseClick);
+			// scroll to the first item to make sure there's anything in view
+			ProcessorsListWidget->RequestScrollIntoView(SelectedProcessors[0]);
+		}
+
+		for (TSharedPtr<FMassDebuggerProcessorData>& ProcessorData : DebuggerModel->SelectedProcessors)
+		{
+			ProcessorsBox->AddSlot()
+				.AutoHeight()
+				[
+					SNew(SMassProcessor, ProcessorData)
+				];
+		}
+	}
+	ProcessorsListWidget->RequestListRefresh();
+}
+
+void SMassProcessorsView::OnArchetypesSelected(TConstArrayView<TSharedPtr<FMassDebuggerArchetypeData>> SelectedArchetypes, ESelectInfo::Type SelectInfo)
+{
+	OnProcessorsSelected(DebuggerModel->SelectedProcessors, ESelectInfo::Direct);
+}
+
+void SMassProcessorsView::OnRefresh()
+{
+	PopulateProcessorList();
+}
+
+void SMassProcessorsView::PopulateProcessorList()
+{
+	check(DebuggerModel.IsValid());
+	ProcessorsListWidget->ClearSelection();
+	DebuggerModel->ClearProcessorSelection();
+
+	ProcessorsListWidget->RequestListRefresh();
+
+	if (ProcessorsBox)
+	{
+		ProcessorsBox->ClearChildren();
+		for (TSharedPtr<FMassDebuggerProcessorData>& ProcessorData : DebuggerModel->SelectedProcessors)
+		{
+			ProcessorsBox->AddSlot()
+				.AutoHeight()
+				[
+					SNew(SMassProcessor, ProcessorData)
+				];
+		}
+	}
+}
+
+#undef LOCTEXT_NAMESPACE
+
diff --git a/Engine/Plugins/Runtime/MassEntity/Source/MassEntityDebugger/Private/SMassProcessorsView.h b/Engine/Plugins/Runtime/MassEntity/Source/MassEntityDebugger/Private/SMassProcessorsView.h
new file mode 100644
--- /dev/null
+++ b/Engine/Plugins/Runtime/MassEntity/Source/MassEntityDebugger/Private/SMassProcessorsView.h
@@ -0,0 +1,38 @@
+// Copyright Epic Games, Inc. All Rights Reserved.
+#pragma once
+
+#include "SMassDebuggerViewBase.h"
+
+
+struct FMassDebuggerProcessorData;
+struct FMassDebuggerArchetypeData;
+class SRichTextBlock;
+struct FMassDebuggerModel;
+class SMassProcessor;
+template<typename T> class SListView;
+class SVerticalBox;
+
+class SMassProcessorsView : public SMassDebuggerViewBase
+{
+public:
+	SLATE_BEGIN_ARGS(SMassProcessorsView)
+	{
+	}
+	SLATE_END_ARGS()
+
+	void Construct(const FArguments& InArgs, TSharedRef<FMassDebuggerModel> InDebuggerModel);
+
+	void PopulateProcessorList();
+
+protected:
+	virtual void OnRefresh() override;
+	virtual void OnProcessorsSelected(TConstArrayView<TSharedPtr<FMassDebuggerProcessorData>> SelectedProcessors, ESelectInfo::Type SelectInfo) override;
+	virtual void OnArchetypesSelected(TConstArrayView<TSharedPtr<FMassDebuggerArchetypeData>> SelectedArchetypes, ESelectInfo::Type SelectInfo) override;
+
+	void ProcessorListSelectionChanged(TSharedPtr<FMassDebuggerProcessorData> SelectedItem, ESelectInfo::Type SelectInfo);
+
+	TSharedPtr<SListView<TSharedPtr<FMassDebuggerProcessorData>>> ProcessorsListWidget;	
+	TSharedPtr<SListView<TSharedPtr<FMassDebuggerProcessorData>>> SelectedProcessorsListWidget;
+	TSharedPtr<SMassProcessor> ProcessorWidget;
+	TSharedPtr<SVerticalBox> ProcessorsBox;
+};
diff --git a/Engine/Plugins/Runtime/MassEntity/Source/MassEntityDebugger/Private/SMassQuery.cpp b/Engine/Plugins/Runtime/MassEntity/Source/MassEntityDebugger/Private/SMassQuery.cpp
new file mode 100644
--- /dev/null
+++ b/Engine/Plugins/Runtime/MassEntity/Source/MassEntityDebugger/Private/SMassQuery.cpp
@@ -0,0 +1,82 @@
+// Copyright Epic Games, Inc. All Rights Reserved.
+
+#include "SMassQuery.h"
+#include "SMassBitSet.h"
+#include "MassEntityTypes.h"
+#include "MassDebuggerModel.h"
+#include "Widgets/Text/SRichTextBlock.h"
+#include "Styling/AppStyle.h"
+
+#define LOCTEXT_NAMESPACE "SMassDebugger"
+
+
+static_assert(EMassAccessOperation::Read == uint32(EMassBitSetDisplayMode::ReadOnly) && EMassAccessOperation::Write == uint32(EMassBitSetDisplayMode::ReadWrite)
+	, "EMassAccessOperation and EMassBitSetDisplayMode need to be kept in sync");
+
+namespace UE::Mass::Debugger::UI::Private
+{
+	template<typename TBitSet>
+	void AddBitSetPair(TSharedRef<SVerticalBox>& Box, const TMassExecutionAccess<TBitSet>& BitSetAccess, const FString& Label)
+	{
+		if (BitSetAccess.IsEmpty() == false)
+		{
+			static const FSlateBrush* RequiredAccessBrushes[] = {
+				FMassDebuggerStyle::GetBrush("MassDebug.Fragment.ReadOnly")
+				, FMassDebuggerStyle::GetBrush("MassDebug.Fragment.ReadWrite")
+			};
+
+			Box->AddSlot()
+			.AutoHeight()
+			[
+				SNew(SMassBitSet<TBitSet>, Label, BitSetAccess.AsArrayView(), RequiredAccessBrushes)
+				.SlotPadding(5)
+			];
+		}
+	}
+} // namespace UE::Mass::Debugger::UI::Private
+
+
+void SMassQuery::Construct(const FArguments& InArgs, TSharedPtr<FMassDebuggerQueryData> InQueryData)
+{
+	using UE::Mass::Debugger::UI::Private::AddBitSetPair;
+	using UE::Mass::Debugger::UI::AddBitSet;
+
+	const FSlateColor ReadOnlyColor = FLinearColor(0.4f, 0.4f, 0.4f, 0.9f);
+	const FSlateColor ReadWriteColor = FLinearColor(0.1f, 0.8f, 0.1f, 0.9f);
+
+	QueryData = InQueryData;
+
+	FMassExecutionRequirements& ExecutionRequirements = InQueryData->ExecutionRequirements;
+
+	TSharedRef<SVerticalBox> Box = SNew(SVerticalBox);
+	static const FSlateBrush* ReadOnlyBrush = FMassDebuggerStyle::GetBrush("MassDebug.Fragment.ReadOnly");
+
+	Box->AddSlot()
+	.AutoHeight()
+	[
+		SNew(SRichTextBlock)
+		.Text(LOCTEXT("MassEntityQueryLabel", "Query"))
+		.DecoratorStyleSet(&FAppStyle::Get())
+		.TextStyle(FAppStyle::Get(), "LargeText")
+	];
+
+	AddBitSetPair(Box, ExecutionRequirements.Fragments, TEXT("Fragments"));
+	AddBitSet(Box, ExecutionRequirements.RequiredAllTags, TEXT("RequiredAllTags"), ReadOnlyBrush);
+	AddBitSet(Box, ExecutionRequirements.RequiredAnyTags, TEXT("RequiredAnyTags"), ReadOnlyBrush);
+	AddBitSet(Box, ExecutionRequirements.RequiredNoneTags,TEXT("RequiredNoneTags"), ReadOnlyBrush);
+	AddBitSetPair(Box, ExecutionRequirements.ChunkFragments, TEXT("Chunk Fragments"));
+	AddBitSetPair(Box, ExecutionRequirements.SharedFragments, TEXT("Shared Fragments"));
+	AddBitSetPair(Box, ExecutionRequirements.RequiredSubsystems, TEXT("Required Subsystems"));
+
+	ChildSlot
+	[
+		SNew(SBorder)
+		.Padding(10.0f)
+		[
+			Box
+		]
+	];
+}
+
+#undef LOCTEXT_NAMESPACE
+
diff --git a/Engine/Plugins/Runtime/MassEntity/Source/MassEntityDebugger/Private/SMassQuery.h b/Engine/Plugins/Runtime/MassEntity/Source/MassEntityDebugger/Private/SMassQuery.h
new file mode 100644
--- /dev/null
+++ b/Engine/Plugins/Runtime/MassEntity/Source/MassEntityDebugger/Private/SMassQuery.h
@@ -0,0 +1,20 @@
+// Copyright Epic Games, Inc. All Rights Reserved.
+#pragma once
+
+#include "Widgets/SCompoundWidget.h"
+
+
+struct FMassDebuggerQueryData;
+
+class SMassQuery : public SCompoundWidget
+{
+public:
+
+	SLATE_BEGIN_ARGS(SMassQuery){}
+	SLATE_END_ARGS()
+
+	void Construct(const FArguments& InArgs, TSharedPtr<FMassDebuggerQueryData> InQueryData);
+
+protected:
+	TSharedPtr<FMassDebuggerQueryData> QueryData;
+};
diff --git a/Engine/Plugins/Runtime/MassEntity/Source/MassEntityDebugger/Public/MassDebuggerSettings.h b/Engine/Plugins/Runtime/MassEntity/Source/MassEntityDebugger/Public/MassDebuggerSettings.h
--- a/Engine/Plugins/Runtime/MassEntity/Source/MassEntityDebugger/Public/MassDebuggerSettings.h
+++ b/Engine/Plugins/Runtime/MassEntity/Source/MassEntityDebugger/Public/MassDebuggerSettings.h
@@ -6,10 +6,14 @@
 #include "MassDebuggerSettings.generated.h"
 
 
-UCLASS(config = Mass, defaultconfig, DisplayName = "Mass Debugger")
+#define GET_MASSDEBUGGER_CONFIG_VALUE(a) (GetMutableDefault<UMassDebuggerSettings>()->a)
+
+UCLASS(config = EditorPerProjectUserSettings, DisplayName = "Mass Debugger")
 class MASSENTITYDEBUGGER_API UMassDebuggerSettings : public UMassModuleSettings
 {
 	GENERATED_BODY()
 
 public:
+	UPROPERTY(EditAnywhere, config, Category=Debugger)
+	bool bStripMassPrefix = true;
 };
