MassDebugger tool non-unity fixes.
Compilation issue #191829

diff --git a/Engine/Plugins/Runtime/MassEntity/Source/MassEntityDebugger/Private/SMassProcessingGraphView.cpp b/Engine/Plugins/Runtime/MassEntity/Source/MassEntityDebugger/Private/SMassProcessingGraphView.cpp
--- a/Engine/Plugins/Runtime/MassEntity/Source/MassEntityDebugger/Private/SMassProcessingGraphView.cpp
+++ b/Engine/Plugins/Runtime/MassEntity/Source/MassEntityDebugger/Private/SMassProcessingGraphView.cpp
@@ -4,6 +4,7 @@
 #include "MassDebuggerModel.h"
 #include "Widgets/Layout/SScrollBox.h"
 #include "Widgets/Layout/SBorder.h"
+#include "Widgets/SBoxPanel.h"
 #include "Styling/SlateBrush.h"
 #include "MassDebuggerStyle.h"
 #include "Widgets/Views/STreeView.h"
diff --git a/Engine/Plugins/Runtime/MassEntity/Source/MassEntityDebugger/Private/SMassProcessingGraphView.h b/Engine/Plugins/Runtime/MassEntity/Source/MassEntityDebugger/Private/SMassProcessingGraphView.h
--- a/Engine/Plugins/Runtime/MassEntity/Source/MassEntityDebugger/Private/SMassProcessingGraphView.h
+++ b/Engine/Plugins/Runtime/MassEntity/Source/MassEntityDebugger/Private/SMassProcessingGraphView.h
@@ -8,8 +8,8 @@
 struct FMassDebuggerProcessingGraph;
 struct FMassDebuggerModel;
 template<typename T> class STreeView;
+class SVerticalBox;
 
-class SHeaderRow;
 
 struct FMassDebuggerProcessingGraphNodeTreeItem
 {
diff --git a/Engine/Plugins/Runtime/MassEntity/Source/MassEntityDebugger/Private/SMassProcessor.cpp b/Engine/Plugins/Runtime/MassEntity/Source/MassEntityDebugger/Private/SMassProcessor.cpp
--- a/Engine/Plugins/Runtime/MassEntity/Source/MassEntityDebugger/Private/SMassProcessor.cpp
+++ b/Engine/Plugins/Runtime/MassEntity/Source/MassEntityDebugger/Private/SMassProcessor.cpp
@@ -5,6 +5,10 @@
 #include "MassDebuggerStyle.h"
 #include "SMassQuery.h"
 #include "Widgets/Text/SRichTextBlock.h"
+#include "Widgets/Views/STableViewBase.h"
+#include "Widgets/Views/STableRow.h"
+#include "Widgets/SBoxPanel.h"
+#include "Widgets/Layout/SBorder.h"
 #include "Styling/AppStyle.h"
 
 
