Fix warnings calling virtual functions from constructor.

diff --git a/Engine/Plugins/Runtime/MassEntity/Source/MassEntityDebugger/Private/MassDebuggerStyle.cpp b/Engine/Plugins/Runtime/MassEntity/Source/MassEntityDebugger/Private/MassDebuggerStyle.cpp
--- a/Engine/Plugins/Runtime/MassEntity/Source/MassEntityDebugger/Private/MassDebuggerStyle.cpp
+++ b/Engine/Plugins/Runtime/MassEntity/Source/MassEntityDebugger/Private/MassDebuggerStyle.cpp
@@ -16,7 +16,7 @@
 #define TTF_CORE_FONT(RelativePath, ...) FSlateFontInfo(RootToCoreContentDir(RelativePath, TEXT(".ttf")), __VA_ARGS__)
 #define FONT(...) FSlateFontInfo(FCoreStyle::GetDefaultFont(), __VA_ARGS__)
 
-class FMassDebuggerStyleSet : public FSlateStyleSet
+class FMassDebuggerStyleSet final : public FSlateStyleSet
 {
 public:
 	static FString InContent(const FString& RelativePath, const ANSICHAR* Extension)
