MassDebugger tool non-unity fixes.
Compilation issue #191829
Compilation issue #192138

diff --git a/Engine/Plugins/Runtime/MassEntity/Source/MassEntityDebugger/Private/MassDebuggerModel.h b/Engine/Plugins/Runtime/MassEntity/Source/MassEntityDebugger/Private/MassDebuggerModel.h
--- a/Engine/Plugins/Runtime/MassEntity/Source/MassEntityDebugger/Private/MassDebuggerModel.h
+++ b/Engine/Plugins/Runtime/MassEntity/Source/MassEntityDebugger/Private/MassDebuggerModel.h
@@ -5,6 +5,7 @@
 #include "MassEntityQuery.h"
 #include "MassProcessorDependencySolver.h"
 #include "MassProcessingTypes.h"
+#include "Types/SlateEnums.h"
 #include "MassDebugger.h"
 
 
diff --git a/Engine/Plugins/Runtime/MassEntity/Source/MassEntityDebugger/Private/SMassDebuggerViewBase.h b/Engine/Plugins/Runtime/MassEntity/Source/MassEntityDebugger/Private/SMassDebuggerViewBase.h
--- a/Engine/Plugins/Runtime/MassEntity/Source/MassEntityDebugger/Private/SMassDebuggerViewBase.h
+++ b/Engine/Plugins/Runtime/MassEntity/Source/MassEntityDebugger/Private/SMassDebuggerViewBase.h
@@ -2,7 +2,7 @@
 #pragma once
 
 #include "Widgets/SCompoundWidget.h"
-
+#include "Types/SlateEnums.h"
 
 struct FMassDebuggerProcessorData;
 struct FMassDebuggerArchetypeData;
diff --git a/Engine/Plugins/Runtime/MassEntity/Source/MassEntityDebugger/Private/SMassProcessingGraphView.cpp b/Engine/Plugins/Runtime/MassEntity/Source/MassEntityDebugger/Private/SMassProcessingGraphView.cpp
--- a/Engine/Plugins/Runtime/MassEntity/Source/MassEntityDebugger/Private/SMassProcessingGraphView.cpp
+++ b/Engine/Plugins/Runtime/MassEntity/Source/MassEntityDebugger/Private/SMassProcessingGraphView.cpp
@@ -4,7 +4,6 @@
 #include "MassDebuggerModel.h"
 #include "Widgets/Layout/SScrollBox.h"
 #include "Widgets/Layout/SBorder.h"
-#include "Widgets/SBoxPanel.h"
 #include "Styling/SlateBrush.h"
 #include "MassDebuggerStyle.h"
 #include "Widgets/Views/STreeView.h"
diff --git a/Engine/Plugins/Runtime/MassEntity/Source/MassEntityDebugger/Private/SMassProcessingGraphView.h b/Engine/Plugins/Runtime/MassEntity/Source/MassEntityDebugger/Private/SMassProcessingGraphView.h
--- a/Engine/Plugins/Runtime/MassEntity/Source/MassEntityDebugger/Private/SMassProcessingGraphView.h
+++ b/Engine/Plugins/Runtime/MassEntity/Source/MassEntityDebugger/Private/SMassProcessingGraphView.h
@@ -3,13 +3,11 @@
 
 #include "SMassDebuggerViewBase.h"
 #include "MassDebuggerModel.h"
-
+#include "Widgets/SBoxPanel.h"
 
 struct FMassDebuggerProcessingGraph;
 struct FMassDebuggerModel;
 template<typename T> class STreeView;
-class SVerticalBox;
-
 
 struct FMassDebuggerProcessingGraphNodeTreeItem
 {
diff --git a/Engine/Plugins/Runtime/MassEntity/Source/MassEntityDebugger/Private/SMassProcessor.cpp b/Engine/Plugins/Runtime/MassEntity/Source/MassEntityDebugger/Private/SMassProcessor.cpp
--- a/Engine/Plugins/Runtime/MassEntity/Source/MassEntityDebugger/Private/SMassProcessor.cpp
+++ b/Engine/Plugins/Runtime/MassEntity/Source/MassEntityDebugger/Private/SMassProcessor.cpp
@@ -5,42 +5,12 @@
 #include "MassDebuggerStyle.h"
 #include "SMassQuery.h"
 #include "Widgets/Text/SRichTextBlock.h"
-#include "Widgets/Views/STableViewBase.h"
-#include "Widgets/Views/STableRow.h"
 #include "Widgets/SBoxPanel.h"
 #include "Widgets/Layout/SBorder.h"
 #include "Styling/AppStyle.h"
 
-
 #define LOCTEXT_NAMESPACE "SMassDebugger"
 
-
-//----------------------------------------------------------------------//
-// SMassQueryTableRow
-//----------------------------------------------------------------------//
-using FMassDebuggerQueryDataPtr = TSharedPtr<FMassDebuggerQueryData, ESPMode::ThreadSafe>;
-
-class SMassQueryTableRow : public STableRow<FMassDebuggerQueryDataPtr>
-{
-public:
-	SLATE_BEGIN_ARGS(SMassQueryTableRow) {}
-	SLATE_END_ARGS()
-
-	void Construct(const FArguments& InArgs, const TSharedRef<STableViewBase>& InOwnerTableView, const FMassDebuggerQueryDataPtr InEntryItem)
-	{
-		Item = InEntryItem;
-
-		STableRow<FMassDebuggerQueryDataPtr>::Construct(STableRow<FMassDebuggerQueryDataPtr>::FArguments(), InOwnerTableView);
-
-		ChildSlot
-			[
-				SNew(SMassQuery, Item)
-			];
-	}
-
-	FMassDebuggerQueryDataPtr Item;
-};
-
 //----------------------------------------------------------------------//
 // SMassProcessor
 //----------------------------------------------------------------------//
