EngineTest CookedEditor compilation fixes
Compilation Issue #192260

diff --git a/Engine/Plugins/Runtime/MassEntity/Source/MassEntityDebugger/Private/MassDebuggerModel.cpp b/Engine/Plugins/Runtime/MassEntity/Source/MassEntityDebugger/Private/MassDebuggerModel.cpp
--- a/Engine/Plugins/Runtime/MassEntity/Source/MassEntityDebugger/Private/MassDebuggerModel.cpp
+++ b/Engine/Plugins/Runtime/MassEntity/Source/MassEntityDebugger/Private/MassDebuggerModel.cpp
@@ -72,7 +72,9 @@ const UMassEntitySubsystem* FMassDebuggerEnvironment::GetEntitySubsystem() const
 //----------------------------------------------------------------------//
 FMassDebuggerQueryData::FMassDebuggerQueryData(const FMassEntityQuery& Query)
 {
+#if WITH_MASSENTITY_DEBUG
 	FMassDebugger::GetQueryExecutionRequirements(Query, ExecutionRequirements);
+#endif // WITH_MASSENTITY_DEBUG
 }
 
 int32 FMassDebuggerQueryData::GetTotalBitsUsedCount() 
@@ -86,7 +88,7 @@ int32 FMassDebuggerQueryData::GetTotalBitsUsedCount()
 FMassDebuggerProcessorData::FMassDebuggerProcessorData(const UMassProcessor& InProcessor)
 {
 	SetProcessor(InProcessor);
-
+#if WITH_MASSENTITY_DEBUG
 	TConstArrayView<FMassEntityQuery*> ProcessorQueries = FMassDebugger::GetProcessorQueries(InProcessor);
 	Queries.Reserve(ProcessorQueries.Num());
 	for (const FMassEntityQuery* Query : ProcessorQueries)
@@ -94,12 +96,13 @@ FMassDebuggerProcessorData::FMassDebuggerProcessorData(const UMassProcessor& InP
 		check(Query);
 		Queries.Add(MakeShareable(new FMassDebuggerQueryData(*Query)));
 	}
+#endif // WITH_MASSENTITY_DEBUG
 }
 
 FMassDebuggerProcessorData::FMassDebuggerProcessorData(const UMassEntitySubsystem& EntitySubsystem, UMassProcessor& InProcessor, const TMap<FMassArchetypeHandle, TSharedPtr<FMassDebuggerArchetypeData>>& InTransientArchetypesMap)
 {
 	SetProcessor(InProcessor);
-
+#if WITH_MASSENTITY_DEBUG
 	TConstArrayView<FMassEntityQuery*> ProcessorQueries = FMassDebugger::GetUpToDateProcessorQueries(EntitySubsystem, InProcessor);
 	Queries.Reserve(ProcessorQueries.Num());
 
@@ -113,6 +116,7 @@ FMassDebuggerProcessorData::FMassDebuggerProcessorData(const UMassEntitySubsyste
 			ValidArchetypes.Add(InTransientArchetypesMap.FindChecked(ArchetypeHandle));
 		}
 	}
+#endif // WITH_MASSENTITY_DEBUG
 }
 
 void FMassDebuggerProcessorData::SetProcessor(const UMassProcessor& InProcessor)
@@ -128,6 +132,7 @@ void FMassDebuggerProcessorData::SetProcessor(const UMassProcessor& InProcessor)
 //----------------------------------------------------------------------//
 FMassDebuggerArchetypeData::FMassDebuggerArchetypeData(const FMassArchetypeHandle& ArchetypeHandle)
 {
+#if WITH_MASSENTITY_DEBUG
 	Composition = FMassDebugger::GetArchetypeComposition(ArchetypeHandle);
 	SharedFragments = FMassDebugger::GetArchetypeSharedFragmentValues(ArchetypeHandle);
 
@@ -138,6 +143,7 @@ FMassDebuggerArchetypeData::FMassDebuggerArchetypeData(const FMassArchetypeHandl
 	BytesToHexLower(reinterpret_cast<const uint8*>(&FullHash), sizeof(FullHash), Label);
 
 	FMassDebugger::GetArchetypeEntityStats(ArchetypeHandle, EntitiesCount, EntitiesCountPerChunk, ChunksCount);
+#endif // WITH_MASSENTITY_DEBUG
 }
 
 int32 FMassDebuggerArchetypeData::GetTotalBitsUsedCount() const
@@ -175,7 +181,7 @@ FText FMassDebuggerProcessingGraphNode::GetLabel() const
 FMassDebuggerProcessingGraph::FMassDebuggerProcessingGraph(const FMassDebuggerModel& DebuggerModel, UMassCompositeProcessor& InGraphOwner)
 {
 	Label = InGraphOwner.GetProcessorName();
-
+#if WITH_MASSENTITY_DEBUG
 	TConstArrayView<UMassCompositeProcessor::FDependencyNode> ProcessingGraph = FMassDebugger::GetProcessingGraph(InGraphOwner);
 
 	GraphNodes.Reserve(ProcessingGraph.Num());
@@ -186,6 +192,7 @@ FMassDebuggerProcessingGraph::FMassDebuggerProcessingGraph(const FMassDebuggerMo
 		check(ProcessorData.IsValid());
 		GraphNodes.Add(FMassDebuggerProcessingGraphNode(ProcessorData, Node));
 	}
+#endif // WITH_MASSENTITY_DEBUG
 }
 
 //----------------------------------------------------------------------//
@@ -372,6 +379,7 @@ float FMassDebuggerModel::MinDistanceToSelectedArchetypes(const TSharedPtr<FMass
 
 void FMassDebuggerModel::StoreArchetypes(const UMassEntitySubsystem& EntitySubsystem, TMap<FMassArchetypeHandle, TSharedPtr<FMassDebuggerArchetypeData>>& OutTransientArchetypesMap)
 {
+#if WITH_MASSENTITY_DEBUG
 	TArray<FMassArchetypeHandle> ArchetypeHandles = FMassDebugger::GetAllArchetypes(EntitySubsystem);
 
 	CachedArchetypes.Reset(ArchetypeHandles.Num());
@@ -387,7 +395,8 @@ void FMassDebuggerModel::StoreArchetypes(const UMassEntitySubsystem& EntitySubsy
 
 		MaxBitsUsed = FMath::Max(MaxBitsUsed, ArchetypeDataPtr->GetTotalBitsUsedCount());
 	}
-	
+#endif // WITH_MASSENTITY_DEBUG
+
 	// calculate distances
 	ArchetypeDistances.Reset();
 	ArchetypeDistances.AddDefaulted(CachedArchetypes.Num());
diff --git a/Engine/Plugins/Runtime/MassEntity/Source/MassEntityDebugger/Private/MassDebuggerModel.h b/Engine/Plugins/Runtime/MassEntity/Source/MassEntityDebugger/Private/MassDebuggerModel.h
--- a/Engine/Plugins/Runtime/MassEntity/Source/MassEntityDebugger/Private/MassDebuggerModel.h
+++ b/Engine/Plugins/Runtime/MassEntity/Source/MassEntityDebugger/Private/MassDebuggerModel.h
@@ -5,6 +5,7 @@
 #include "MassEntityQuery.h"
 #include "MassProcessorDependencySolver.h"
 #include "MassProcessingTypes.h"
+#include "MassProcessor.h"
 #include "Types/SlateEnums.h"
 #include "MassDebugger.h"
 
diff --git a/Engine/Plugins/Runtime/MassEntity/Source/MassEntityDebugger/Private/SMassBitSet.h b/Engine/Plugins/Runtime/MassEntity/Source/MassEntityDebugger/Private/SMassBitSet.h
--- a/Engine/Plugins/Runtime/MassEntity/Source/MassEntityDebugger/Private/SMassBitSet.h
+++ b/Engine/Plugins/Runtime/MassEntity/Source/MassEntityDebugger/Private/SMassBitSet.h
@@ -8,7 +8,7 @@
 #include "Widgets/Layout/SWrapBox.h"
 #include "Widgets/Layout/SBorder.h"
 #include "Widgets/SBoxPanel.h"
-
+#include "MassProcessingTypes.h"
 
 struct FMassDebuggerQueryData;
 
@@ -69,6 +69,7 @@ public:
 protected:
 	void AddBitSet(const SMassBitSet::FArguments& InArgs, TSharedRef<SWrapBox>& ButtonBox, const TBitSet& BitSet, const FSlateBrush* Brush)
 	{
+#if WITH_MASSENTITY_DEBUG
 		TArray<FName> TypeNames;
 		BitSet.DebugGetIndividualNames(TypeNames);
 
@@ -86,6 +87,7 @@ protected:
 				]
 			];
 		}
+#endif // WITH_MASSENTITY_DEBUG
 	}
 };
 
