Added support for declaring subsystem requirements and fetching using TSubclassOf

diff --git a/Engine/Plugins/Runtime/MassEntity/Source/MassEntity/Public/MassEntityQuery.h b/Engine/Plugins/Runtime/MassEntity/Source/MassEntity/Public/MassEntityQuery.h
--- a/Engine/Plugins/Runtime/MassEntity/Source/MassEntity/Public/MassEntityQuery.h
+++ b/Engine/Plugins/Runtime/MassEntity/Source/MassEntity/Public/MassEntityQuery.h
@@ -349,6 +349,27 @@ public:
 		return *this;
 	}
 
+	FMassEntityQuery& AddSubsystemRequirement(const TSubclassOf<USubsystem> SubsystemClass, const EMassFragmentAccess AccessMode, const bool bGameThreadOnly = true)
+	{
+		check(AccessMode != EMassFragmentAccess::None && AccessMode != EMassFragmentAccess::MAX);
+
+		switch (AccessMode)
+		{
+		case EMassFragmentAccess::ReadOnly:
+			RequiredConstSubsystems.Add(**SubsystemClass);
+			bRequiresGameThreadExecution |= bGameThreadOnly;
+			break;
+		case EMassFragmentAccess::ReadWrite:
+			RequiredMutableSubsystems.Add(**SubsystemClass);
+			bRequiresGameThreadExecution |= bGameThreadOnly;
+			break;
+		default:
+			check(false);
+		}
+
+		return *this;
+	}
+
 	void Clear()
 	{
 		Requirements.Empty();
diff --git a/Engine/Plugins/Runtime/MassEntity/Source/MassEntity/Public/MassExecutionContext.h b/Engine/Plugins/Runtime/MassEntity/Source/MassEntity/Public/MassExecutionContext.h
--- a/Engine/Plugins/Runtime/MassEntity/Source/MassEntity/Public/MassExecutionContext.h
+++ b/Engine/Plugins/Runtime/MassEntity/Source/MassEntity/Public/MassExecutionContext.h
@@ -288,6 +288,46 @@ public:
 
 	void SetSubsystemRequirements(const FMassExternalSubystemBitSet& RequiredConstSubsystems, const FMassExternalSubystemBitSet& RequiredMutableSubsystems);
 
+	template<typename T>
+	T* GetMutableSubsystem(const UWorld* World, const TSubclassOf<UWorldSubsystem> SubsystemClass)
+	{
+		// @todo consider getting this directly from entity subsystem - it should cache all the used system
+		const uint32 SystemIndex = FMassExternalSubystemBitSet::GetTypeIndex(**SubsystemClass);
+		if (ensure(MutableSubsystemsBitSet.IsBitSet(SystemIndex)))
+		{
+			return GetSubsystemInternal<T>(World, SystemIndex, SubsystemClass);
+		}
+
+		return nullptr;
+	}
+
+	template<typename T>
+	T& GetMutableSubsystemChecked(const UWorld* World, const TSubclassOf<UWorldSubsystem> SubsystemClass)
+	{
+		T* InstancePtr = GetMutableSubsystem<T>(World, SubsystemClass);
+		check(InstancePtr);
+		return *InstancePtr;
+	}
+
+	template<typename T>
+	const T* GetSubsystem(const UWorld* World, const TSubclassOf<UWorldSubsystem> SubsystemClass)
+	{
+		const uint32 SystemIndex = FMassExternalSubystemBitSet::GetTypeIndex(**SubsystemClass);
+		if (ensure(ConstSubsystemsBitSet.IsBitSet(SystemIndex) || MutableSubsystemsBitSet.IsBitSet(SystemIndex)))
+		{
+			return GetSubsystemInternal<T>(World, SystemIndex, SubsystemClass);
+		}
+		return nullptr;
+	}
+
+	template<typename T>
+	const T& GetSubsystemChecked(const UWorld* World, const TSubclassOf<UWorldSubsystem> SubsystemClass)
+	{
+		const T* InstancePtr = GetSubsystem<T>(World, SubsystemClass);
+		check(InstancePtr);
+		return *InstancePtr;
+	}
+
 	/** Sparse chunk related operation */
 	const FMassArchetypeEntityCollection& GetEntityCollection() const { return EntityCollection; }
 
@@ -351,4 +391,21 @@ protected:
 		}
 		return SystemInstance;
 	}
+
+	template<typename T>
+	T* GetSubsystemInternal(const UWorld* World, const uint32 SystemIndex, const TSubclassOf<UWorldSubsystem> SubsystemClass)
+	{
+		if (UNLIKELY(Subsystems.IsValidIndex(SystemIndex) == false))
+		{
+			Subsystems.AddZeroed(Subsystems.Num() - SystemIndex + 1);
+		}
+
+		T* SystemInstance = (T*)Subsystems[SystemIndex].Get();
+		if (SystemInstance == nullptr)
+		{
+			SystemInstance = FMassExternalSubsystemTraits::GetInstance<typename TRemoveConst<T>::Type>(World, SubsystemClass);
+			Subsystems[SystemIndex] = SystemInstance;
+		}
+		return SystemInstance;
+	}
 };
diff --git a/Engine/Plugins/Runtime/MassEntity/Source/MassEntity/Public/MassExternalSubsystemTraits.h b/Engine/Plugins/Runtime/MassEntity/Source/MassEntity/Public/MassExternalSubsystemTraits.h
--- a/Engine/Plugins/Runtime/MassEntity/Source/MassEntity/Public/MassExternalSubsystemTraits.h
+++ b/Engine/Plugins/Runtime/MassEntity/Source/MassEntity/Public/MassExternalSubsystemTraits.h
@@ -43,11 +43,17 @@ namespace FMassExternalSubsystemTraits
 	This is a getter function that given a UWorld* fetches an instance.
 	*/
 	template<typename T, typename = typename TEnableIf<TIsDerivedFrom<T, UWorldSubsystem>::IsDerived>::Type>
-	static T* GetInstance(const UWorld* World)
+	T* GetInstance(const UWorld* World)
 	{ 
 		// note that the default implementation works only for UWorldSubsystems
 		return UWorld::GetSubsystem<T>(World);
 	}
+
+	template<typename T>
+	T* GetInstance(const UWorld* World, const TSubclassOf<UWorldSubsystem> SubsystemClass)
+	{
+		return World ? Cast<T>(World->GetSubsystemBase(SubsystemClass)) : (T*)nullptr;
+	}
 }
 
 /** 
