Updated MassDebugger's Processor display to clearly indicate the processor's dependencies (as declared by MassProcessor.ProcessorRequirements).

diff --git a/Engine/Plugins/Runtime/MassEntity/Source/MassEntity/Private/MassProcessorDependencySolver.cpp b/Engine/Plugins/Runtime/MassEntity/Source/MassEntity/Private/MassProcessorDependencySolver.cpp
--- a/Engine/Plugins/Runtime/MassEntity/Source/MassEntity/Private/MassProcessorDependencySolver.cpp
+++ b/Engine/Plugins/Runtime/MassEntity/Source/MassEntity/Private/MassProcessorDependencySolver.cpp
@@ -68,6 +68,13 @@ int32 FMassExecutionRequirements::GetTotalBitsUsedCount()
 		+ RequiredAnyTags.CountStoredTypes() + RequiredNoneTags.CountStoredTypes();
 }
 
+bool FMassExecutionRequirements::IsEmpty() const
+{
+	return Fragments.IsEmpty() && ChunkFragments.IsEmpty() 
+		&& SharedFragments.IsEmpty() && RequiredSubsystems.IsEmpty()
+		&& RequiredAllTags.IsEmpty() && RequiredAnyTags.IsEmpty() && RequiredNoneTags.IsEmpty();
+}
+
 //----------------------------------------------------------------------//
 //  FProcessorDependencySolver::FResourceUsage
 //----------------------------------------------------------------------//
diff --git a/Engine/Plugins/Runtime/MassEntity/Source/MassEntity/Public/MassProcessorDependencySolver.h b/Engine/Plugins/Runtime/MassEntity/Source/MassEntity/Public/MassProcessorDependencySolver.h
--- a/Engine/Plugins/Runtime/MassEntity/Source/MassEntity/Public/MassProcessorDependencySolver.h
+++ b/Engine/Plugins/Runtime/MassEntity/Source/MassEntity/Public/MassProcessorDependencySolver.h
@@ -52,6 +52,7 @@ struct MASSENTITY_API FMassExecutionRequirements
 	void Append(const FMassExecutionRequirements& Other);
 	void CountResourcesUsed();
 	int32 GetTotalBitsUsedCount();
+	bool IsEmpty() const;
 
 	TMassExecutionAccess<FMassFragmentBitSet> Fragments;
 	TMassExecutionAccess<FMassChunkFragmentBitSet> ChunkFragments;
diff --git a/Engine/Plugins/Runtime/MassEntity/Source/MassEntityDebugger/Private/MassDebuggerModel.cpp b/Engine/Plugins/Runtime/MassEntity/Source/MassEntityDebugger/Private/MassDebuggerModel.cpp
--- a/Engine/Plugins/Runtime/MassEntity/Source/MassEntityDebugger/Private/MassDebuggerModel.cpp
+++ b/Engine/Plugins/Runtime/MassEntity/Source/MassEntityDebugger/Private/MassDebuggerModel.cpp
@@ -70,7 +70,8 @@ const UMassEntitySubsystem* FMassDebuggerEnvironment::GetEntitySubsystem() const
 //----------------------------------------------------------------------//
 // FMassDebuggerQueryData
 //----------------------------------------------------------------------//
-FMassDebuggerQueryData::FMassDebuggerQueryData(const FMassEntityQuery& Query)
+FMassDebuggerQueryData::FMassDebuggerQueryData(const FMassEntityQuery& Query, const FText& InLabel)
+	: Label(InLabel)
 {
 #if WITH_MASSENTITY_DEBUG
 	FMassDebugger::GetQueryExecutionRequirements(Query, ExecutionRequirements);
@@ -82,6 +83,11 @@ int32 FMassDebuggerQueryData::GetTotalBitsUsedCount()
 	return ExecutionRequirements.GetTotalBitsUsedCount();
 }
 
+bool FMassDebuggerQueryData::IsEmpty() const
+{
+	return ExecutionRequirements.IsEmpty();
+}
+
 //----------------------------------------------------------------------//
 // FMassDebuggerProcessorData
 //----------------------------------------------------------------------//
@@ -90,11 +96,18 @@ FMassDebuggerProcessorData::FMassDebuggerProcessorData(const UMassProcessor& InP
 	SetProcessor(InProcessor);
 #if WITH_MASSENTITY_DEBUG
 	TConstArrayView<FMassEntityQuery*> ProcessorQueries = FMassDebugger::GetProcessorQueries(InProcessor);
-	Queries.Reserve(ProcessorQueries.Num());
-	for (const FMassEntityQuery* Query : ProcessorQueries)
+
+	if (ensureMsgf(ProcessorQueries.Num() >= 1, TEXT("We expect at least one query to be there - every processor should have their ProcessorRequirements registered in OwnedQueries")))
 	{
-		check(Query);
-		Queries.Add(MakeShareable(new FMassDebuggerQueryData(*Query)));
+		ProcessorRequirements = MakeShareable(new FMassDebuggerQueryData(*ProcessorQueries[0], LOCTEXT("MassProcessorRequirementsLabel", "Processor Requirements")));
+
+		Queries.Reserve(ProcessorQueries.Num() - 1);
+		for (int i = 1; i < ProcessorQueries.Num(); ++i)
+		{
+			const FMassEntityQuery* Query = ProcessorQueries[i];
+			check(Query);
+			Queries.Add(MakeShareable(new FMassDebuggerQueryData(*Query, LOCTEXT("MassEntityQueryLabel", "Query"))));
+		}
 	}
 #endif // WITH_MASSENTITY_DEBUG
 }
@@ -104,16 +117,21 @@ FMassDebuggerProcessorData::FMassDebuggerProcessorData(const UMassEntitySubsyste
 	SetProcessor(InProcessor);
 #if WITH_MASSENTITY_DEBUG
 	TConstArrayView<FMassEntityQuery*> ProcessorQueries = FMassDebugger::GetUpToDateProcessorQueries(EntitySubsystem, InProcessor);
-	Queries.Reserve(ProcessorQueries.Num());
-
-	for (const FMassEntityQuery* Query : ProcessorQueries)
+	if (ensureMsgf(ProcessorQueries.Num() >= 1, TEXT("We expect at least one query to be there - every processor should have their ProcessorRequirements registered in OwnedQueries")))
 	{
-		check(Query);
-		Queries.Add(MakeShareable(new FMassDebuggerQueryData(*Query)));
+		ProcessorRequirements = MakeShareable(new FMassDebuggerQueryData(*ProcessorQueries[0], LOCTEXT("MassProcessorRequirementsLabel", "Processor Requirements")));
 
-		for (const FMassArchetypeHandle& ArchetypeHandle : Query->GetArchetypes())
+		Queries.Reserve(ProcessorQueries.Num() - 1);
+		for (int i = 1; i < ProcessorQueries.Num(); ++i)
 		{
-			ValidArchetypes.Add(InTransientArchetypesMap.FindChecked(ArchetypeHandle));
+			const FMassEntityQuery* Query = ProcessorQueries[i];
+			check(Query);
+			Queries.Add(MakeShareable(new FMassDebuggerQueryData(*Query, LOCTEXT("MassEntityQueryLabel", "Query"))));
+
+			for (const FMassArchetypeHandle& ArchetypeHandle : Query->GetArchetypes())
+			{
+				ValidArchetypes.Add(InTransientArchetypesMap.FindChecked(ArchetypeHandle));
+			}
 		}
 	}
 #endif // WITH_MASSENTITY_DEBUG
diff --git a/Engine/Plugins/Runtime/MassEntity/Source/MassEntityDebugger/Private/MassDebuggerModel.h b/Engine/Plugins/Runtime/MassEntity/Source/MassEntityDebugger/Private/MassDebuggerModel.h
--- a/Engine/Plugins/Runtime/MassEntity/Source/MassEntityDebugger/Private/MassDebuggerModel.h
+++ b/Engine/Plugins/Runtime/MassEntity/Source/MassEntityDebugger/Private/MassDebuggerModel.h
@@ -43,11 +43,13 @@ enum class EMassDebuggerProcessingGraphNodeSelection : uint8
 
 struct FMassDebuggerQueryData
 {
-	FMassDebuggerQueryData(const FMassEntityQuery& Query);
+	FMassDebuggerQueryData(const FMassEntityQuery& Query, const FText& InLabel);
 
 	FMassExecutionRequirements ExecutionRequirements;
+	FText Label;
 
 	int32 GetTotalBitsUsedCount();
+	bool IsEmpty() const;
 }; 
 
 struct FMassDebuggerArchetypeData
@@ -85,6 +87,7 @@ public:
 	FString Name;
 	FString Label;
 	uint32 ProcessorHash = 0; 
+	TSharedPtr<FMassDebuggerQueryData> ProcessorRequirements;
 	TArray<TSharedPtr<FMassDebuggerQueryData>> Queries;
 	TArray<TSharedPtr<FMassDebuggerArchetypeData>> ValidArchetypes;
 	
diff --git a/Engine/Plugins/Runtime/MassEntity/Source/MassEntityDebugger/Private/SMassProcessor.cpp b/Engine/Plugins/Runtime/MassEntity/Source/MassEntityDebugger/Private/SMassProcessor.cpp
--- a/Engine/Plugins/Runtime/MassEntity/Source/MassEntityDebugger/Private/SMassProcessor.cpp
+++ b/Engine/Plugins/Runtime/MassEntity/Source/MassEntityDebugger/Private/SMassProcessor.cpp
@@ -5,6 +5,7 @@
 #include "MassDebuggerStyle.h"
 #include "SMassQuery.h"
 #include "Widgets/Text/SRichTextBlock.h"
+#include "Widgets/Text/STextBlock.h"
 #include "Widgets/SBoxPanel.h"
 #include "Widgets/Layout/SBorder.h"
 #include "Styling/AppStyle.h"
@@ -36,6 +37,28 @@ void SMassProcessor::Construct(const FArguments& InArgs, TSharedPtr<FMassDebugge
 		]
 	];
 	
+	if (ProcessorData->ProcessorRequirements->IsEmpty())
+	{
+		Box->AddSlot()
+		.AutoHeight()
+		[
+			SNew(SBorder)
+			.Padding(10.0f)
+			[
+				SNew(STextBlock)
+				.Text(LOCTEXT("NoProcessorRequirements", "No Processor Requirements"))
+			]
+		];
+	}
+	else
+	{
+		Box->AddSlot()
+		.AutoHeight()
+		[
+			SNew(SMassQuery, ProcessorData->ProcessorRequirements)
+		];
+	}
+
 	for (TSharedPtr<FMassDebuggerQueryData>& QueryData : ProcessorData->Queries)
 	{
 		Box->AddSlot()
diff --git a/Engine/Plugins/Runtime/MassEntity/Source/MassEntityDebugger/Private/SMassQuery.cpp b/Engine/Plugins/Runtime/MassEntity/Source/MassEntityDebugger/Private/SMassQuery.cpp
--- a/Engine/Plugins/Runtime/MassEntity/Source/MassEntityDebugger/Private/SMassQuery.cpp
+++ b/Engine/Plugins/Runtime/MassEntity/Source/MassEntityDebugger/Private/SMassQuery.cpp
@@ -55,7 +55,7 @@ void SMassQuery::Construct(const FArguments& InArgs, TSharedPtr<FMassDebuggerQue
 	.AutoHeight()
 	[
 		SNew(SRichTextBlock)
-		.Text(LOCTEXT("MassEntityQueryLabel", "Query"))
+		.Text(InQueryData->Label)
 		.DecoratorStyleSet(&FAppStyle::Get())
 		.TextStyle(FAppStyle::Get(), "LargeText")
 	];
