Build Issue 193799
Disabled overzelous warning in UMassProcessor::PostInitProperties that was triggering during CitySample cook

diff --git a/Engine/Plugins/Runtime/MassEntity/Source/MassEntity/Private/MassProcessor.cpp b/Engine/Plugins/Runtime/MassEntity/Source/MassEntity/Private/MassProcessor.cpp
--- a/Engine/Plugins/Runtime/MassEntity/Source/MassEntity/Private/MassProcessor.cpp
+++ b/Engine/Plugins/Runtime/MassEntity/Source/MassEntity/Private/MassProcessor.cpp
@@ -134,8 +134,9 @@ void UMassProcessor::PostInitProperties()
 			CA_ASSUME(QueryPtr);
 			bNeedsGameThread = (bNeedsGameThread || QueryPtr->DoesRequireGameThreadExecution());
 		}
-		UE_CLOG(bRequiresGameThreadExecution != bNeedsGameThread, LogMass, Warning, TEXT("%s is marked bRequiresGameThreadExecution = %s, while the registered quries' requirement indicate the opposite")
-			, *GetProcessorName(), bRequiresGameThreadExecution ? TEXT("TRUE") : TEXT("FALSE"));
+		// @todo warning temporarily disabled until details cleared
+		/*UE_CLOG(bRequiresGameThreadExecution != bNeedsGameThread, LogMass, Warning, TEXT("%s is marked bRequiresGameThreadExecution = %s, while the registered quries' requirement indicate the opposite")
+			, *GetProcessorName(), bRequiresGameThreadExecution ? TEXT("TRUE") : TEXT("FALSE"));*/
 
 		// better safe than sorry - if queries indicate the game thread execution is required then we marked the whole processor as such
 		bRequiresGameThreadExecution = bRequiresGameThreadExecution || bNeedsGameThread;
