Extracted requirement information out of MassEntityQuery and into the newly created MassEntityRequirements that then became MassEntityQuery's parent class.
In addition:
* Significantly reduced number of FMassArchetypeHandle's "friends"
  - Added helper functions to MassArchetypeData.h to supply the functionality.
  - kept MassEntitySubsystem as a friend for readability and convenience - I originally removed that one as friend as well, but some updated code looked awkward
* Moved some of the "DebugGetDescription" functions over to FMassDebugger
* Moved contents of MassEntityDebug.h/cpp to MassDebugger
* Renamed a couple of MassEntitySubsystem function's archetype handle parameters from "Archetype" to more precise "ArchetypeHandle".
  - also made these parameters const refs to avoid copying FMassArchetypeData shared pointer

diff --git a/Engine/Plugins/AI/MassAI/Source/MassAIDebug/Private/GameplayDebuggerCategory_Mass.cpp b/Engine/Plugins/AI/MassAI/Source/MassAIDebug/Private/GameplayDebuggerCategory_Mass.cpp
--- a/Engine/Plugins/AI/MassAI/Source/MassAIDebug/Private/GameplayDebuggerCategory_Mass.cpp
+++ b/Engine/Plugins/AI/MassAI/Source/MassAIDebug/Private/GameplayDebuggerCategory_Mass.cpp
@@ -25,6 +25,7 @@
 #include "MassSimulationLOD.h"
 #include "CanvasItem.h"
 #include "Engine/World.h"
+#include "MassDebugger.h"
 
 namespace UE::Mass::Debug
 {
@@ -288,7 +289,7 @@ void FGameplayDebuggerCategory_Mass::CollectData(APlayerController* OwnerPC, AAc
 					FStringOutputDevice FragmentsDesc;
 					FragmentsDesc.SetAutoEmitLineTerminator(true);
 					const TCHAR* PrefixToRemove = TEXT("DataFragment_");
-					EntitySystem->DebugPrintEntity(CachedEntity, FragmentsDesc, PrefixToRemove);
+					FMassDebugger::OutputEntityDescription(FragmentsDesc, *EntitySystem, CachedEntity, PrefixToRemove);
 					AddTextLine(FString::Printf(TEXT("{Green}Fragments:\n{White}%s"), *FragmentsDesc));
 				}
 				else
diff --git a/Engine/Plugins/Runtime/MassEntity/Source/MassEntity/Private/MassArchetypeData.cpp b/Engine/Plugins/Runtime/MassEntity/Source/MassEntity/Private/MassArchetypeData.cpp
--- a/Engine/Plugins/Runtime/MassEntity/Source/MassEntity/Private/MassArchetypeData.cpp
+++ b/Engine/Plugins/Runtime/MassEntity/Source/MassEntity/Private/MassArchetypeData.cpp
@@ -521,10 +521,10 @@ void FMassArchetypeData::CompactEntities(const double TimeAllowed)
 	}
 }
 
-void FMassArchetypeData::GetRequirementsFragmentMapping(TConstArrayView<FMassFragmentRequirement> Requirements, FMassFragmentIndicesMapping& OutFragmentIndices) const
+void FMassArchetypeData::GetRequirementsFragmentMapping(TConstArrayView<FMassFragmentRequirementDescription> Requirements, FMassFragmentIndicesMapping& OutFragmentIndices) const
 {
 	OutFragmentIndices.Reset(Requirements.Num());
-	for (const FMassFragmentRequirement& Requirement : Requirements)
+	for (const FMassFragmentRequirementDescription& Requirement : Requirements)
 	{
 		if (Requirement.RequiresBinding())
 		{
@@ -535,11 +535,11 @@ void FMassArchetypeData::GetRequirementsFragmentMapping(TConstArrayView<FMassFra
 	}
 }
 
-void FMassArchetypeData::GetRequirementsChunkFragmentMapping(TConstArrayView<FMassFragmentRequirement> ChunkRequirements, FMassFragmentIndicesMapping& OutFragmentIndices) const
+void FMassArchetypeData::GetRequirementsChunkFragmentMapping(TConstArrayView<FMassFragmentRequirementDescription> ChunkRequirements, FMassFragmentIndicesMapping& OutFragmentIndices) const
 {
 	int32 LastFoundFragmentIndex = -1;
 	OutFragmentIndices.Reset(ChunkRequirements.Num());
-	for (const FMassFragmentRequirement& Requirement : ChunkRequirements)
+	for (const FMassFragmentRequirementDescription& Requirement : ChunkRequirements)
 	{
 		if (Requirement.RequiresBinding())
 		{
@@ -561,10 +561,10 @@ void FMassArchetypeData::GetRequirementsChunkFragmentMapping(TConstArrayView<FMa
 	}
 }
 
-void FMassArchetypeData::GetRequirementsConstSharedFragmentMapping(TConstArrayView<FMassFragmentRequirement> Requirements, FMassFragmentIndicesMapping& OutFragmentIndices) const
+void FMassArchetypeData::GetRequirementsConstSharedFragmentMapping(TConstArrayView<FMassFragmentRequirementDescription> Requirements, FMassFragmentIndicesMapping& OutFragmentIndices) const
 {
 	OutFragmentIndices.Reset(Requirements.Num());
-	for (const FMassFragmentRequirement& Requirement : Requirements)
+	for (const FMassFragmentRequirementDescription& Requirement : Requirements)
 	{
 		if (Requirement.RequiresBinding())
 		{
@@ -575,10 +575,10 @@ void FMassArchetypeData::GetRequirementsConstSharedFragmentMapping(TConstArrayVi
 	}
 }
 
-void FMassArchetypeData::GetRequirementsSharedFragmentMapping(TConstArrayView<FMassFragmentRequirement> Requirements, FMassFragmentIndicesMapping& OutFragmentIndices) const
+void FMassArchetypeData::GetRequirementsSharedFragmentMapping(TConstArrayView<FMassFragmentRequirementDescription> Requirements, FMassFragmentIndicesMapping& OutFragmentIndices) const
 {
 	OutFragmentIndices.Reset(Requirements.Num());
-	for (const FMassFragmentRequirement& Requirement : Requirements)
+	for (const FMassFragmentRequirementDescription& Requirement : Requirements)
 	{
 		if (Requirement.RequiresBinding())
 		{
diff --git a/Engine/Plugins/Runtime/MassEntity/Source/MassEntity/Private/MassArchetypeData.h b/Engine/Plugins/Runtime/MassEntity/Source/MassEntity/Private/MassArchetypeData.h
--- a/Engine/Plugins/Runtime/MassEntity/Source/MassEntity/Private/MassArchetypeData.h
+++ b/Engine/Plugins/Runtime/MassEntity/Source/MassEntity/Private/MassArchetypeData.h
@@ -244,16 +244,16 @@ public:
 	void SetFragmentData(FMassArchetypeEntityCollection::FConstEntityRangeArrayView EntityRangeContainer, const FInstancedStruct& FragmentSource);
 
 	/** Returns conversion from given Requirements to archetype's fragment indices */
-	void GetRequirementsFragmentMapping(TConstArrayView<FMassFragmentRequirement> Requirements, FMassFragmentIndicesMapping& OutFragmentIndices) const;
+	void GetRequirementsFragmentMapping(TConstArrayView<FMassFragmentRequirementDescription> Requirements, FMassFragmentIndicesMapping& OutFragmentIndices) const;
 
 	/** Returns conversion from given ChunkRequirements to archetype's chunk fragment indices */
-	void GetRequirementsChunkFragmentMapping(TConstArrayView<FMassFragmentRequirement> ChunkRequirements, FMassFragmentIndicesMapping& OutFragmentIndices) const;
+	void GetRequirementsChunkFragmentMapping(TConstArrayView<FMassFragmentRequirementDescription> ChunkRequirements, FMassFragmentIndicesMapping& OutFragmentIndices) const;
 
 	/** Returns conversion from given const shared requirements to archetype's const shared fragment indices */
-	void GetRequirementsConstSharedFragmentMapping(TConstArrayView<FMassFragmentRequirement> Requirements, FMassFragmentIndicesMapping& OutFragmentIndices) const;
+	void GetRequirementsConstSharedFragmentMapping(TConstArrayView<FMassFragmentRequirementDescription> Requirements, FMassFragmentIndicesMapping& OutFragmentIndices) const;
 
 	/** Returns conversion from given shared requirements to archetype's shared fragment indices */
-	void GetRequirementsSharedFragmentMapping(TConstArrayView<FMassFragmentRequirement> Requirements, FMassFragmentIndicesMapping& OutFragmentIndices) const;
+	void GetRequirementsSharedFragmentMapping(TConstArrayView<FMassFragmentRequirementDescription> Requirements, FMassFragmentIndicesMapping& OutFragmentIndices) const;
 
 	SIZE_T GetAllocatedSize() const;
 
@@ -336,3 +336,18 @@ private:
 	int32 AddEntityInternal(FMassEntityHandle Entity);
 	void RemoveEntityInternal(const int32 AbsoluteIndex);
 };
+
+
+struct FMassArchetypeHelper
+{
+	FORCEINLINE static FMassArchetypeData* ArchetypeDataFromHandle(const FMassArchetypeHandle& ArchetypeHandle) { return ArchetypeHandle.DataPtr.Get(); }
+	FORCEINLINE static FMassArchetypeData& ArchetypeDataFromHandleChecked(const FMassArchetypeHandle& ArchetypeHandle)
+	{
+		check(ArchetypeHandle.IsValid());
+		return *ArchetypeHandle.DataPtr.Get();
+	}
+	FORCEINLINE static FMassArchetypeHandle ArchetypeHandleFromData(const TSharedPtr<FMassArchetypeData>& Archetype)
+	{
+		return FMassArchetypeHandle(Archetype);
+	}
+};
diff --git a/Engine/Plugins/Runtime/MassEntity/Source/MassEntity/Private/MassArchetypeTypes.cpp b/Engine/Plugins/Runtime/MassEntity/Source/MassEntity/Private/MassArchetypeTypes.cpp
--- a/Engine/Plugins/Runtime/MassEntity/Source/MassEntity/Private/MassArchetypeTypes.cpp
+++ b/Engine/Plugins/Runtime/MassEntity/Source/MassEntity/Private/MassArchetypeTypes.cpp
@@ -25,7 +25,7 @@ FMassArchetypeEntityCollection::FMassArchetypeEntityCollection(const FMassArchet
 		return;
 	}
 
-	const FMassArchetypeData* ArchetypeData = InArchetype.DataPtr.Get();
+	const FMassArchetypeData* ArchetypeData = FMassArchetypeHelper::ArchetypeDataFromHandle(InArchetype);
 	const int32 NumEntitiesPerChunk = ArchetypeData ? ArchetypeData->GetNumEntitiesPerChunk() : MAX_int32;
 
 	// InEntities has a real chance of not being sorted by AbsoluteIndex. We gotta fix that to optimize how we process the data 
@@ -97,7 +97,8 @@ void FMassArchetypeEntityCollection::BuildEntityRanges(TStridedView<const int32>
 {
 	checkf(Ranges.Num() == 0, TEXT("Calling %s is valid only for initial configuration"), ANSI_TO_TCHAR(__FUNCTION__));
 
-	const int32 NumEntitiesPerChunk = Archetype.IsValid() ? Archetype.DataPtr->GetNumEntitiesPerChunk() : MAX_int32;
+	const FMassArchetypeData* ArchetypeData = FMassArchetypeHelper::ArchetypeDataFromHandle(Archetype);
+	const int32 NumEntitiesPerChunk = ArchetypeData ? ArchetypeData->GetNumEntitiesPerChunk() : MAX_int32;
 
 	// the following block of code is splitting up sorted AbsoluteIndices into 
 	// continuous chunks
@@ -132,13 +133,13 @@ FMassArchetypeEntityCollection::FMassArchetypeEntityCollection(const FMassArchet
 {
 	if (Initialization == EInitializationType::GatherAll)
 	{
-		check(InArchetypeHandle.DataPtr.IsValid());
+		check(InArchetypeHandle.IsValid());
 		GatherChunksFromArchetype();
 	}
 }
 
 FMassArchetypeEntityCollection::FMassArchetypeEntityCollection(TSharedPtr<FMassArchetypeData>& InArchetype, const EInitializationType Initialization)
-	: Archetype(InArchetype)
+	: Archetype(FMassArchetypeHelper::ArchetypeHandleFromData(InArchetype))
 {	
 	if (Initialization == EInitializationType::GatherAll)
 	{
@@ -149,7 +150,7 @@ FMassArchetypeEntityCollection::FMassArchetypeEntityCollection(TSharedPtr<FMassA
 
 void FMassArchetypeEntityCollection::GatherChunksFromArchetype()
 {
-	if (const FMassArchetypeData* ArchetypePtr = Archetype.DataPtr.Get())
+	if (const FMassArchetypeData* ArchetypePtr = FMassArchetypeHelper::ArchetypeDataFromHandle(Archetype))
 	{
 		const int32 ChunkCount = ArchetypePtr->GetChunkCount();
 		Ranges.Reset(ChunkCount);
@@ -180,11 +181,6 @@ bool FMassArchetypeEntityCollection::IsSame(const FMassArchetypeEntityCollection
 //////////////////////////////////////////////////////////////////////
 // FMassArchetypeEntityCollectionWithPayload
 
-FMassArchetypeData* FMassArchetypeEntityCollection::GetArchetypePtr(const FMassArchetypeHandle& ArchetypeHandle)
-{
-	return ArchetypeHandle.DataPtr.Get();
-}
-
 void FMassArchetypeEntityCollectionWithPayload::CreateEntityRangesWithPayload(const UMassEntitySubsystem& EntitySystem, const TConstArrayView<FMassEntityHandle> Entities
 	, const FMassArchetypeEntityCollection::EDuplicatesHandling DuplicatesHandling, FMassGenericPayloadView Payload
 	, TArray<FMassArchetypeEntityCollectionWithPayload>& OutEntityCollections)
@@ -232,7 +228,7 @@ void FMassArchetypeEntityCollectionWithPayload::CreateEntityRangesWithPayload(co
 	{
 		const FMassEntityHandle& Entity = Entities[i];
 		const FMassArchetypeHandle ArchetypeHandle = EntitySystem.GetArchetypeForEntity(Entity);
-		const FMassArchetypeData* ArchetypePtr = FMassArchetypeEntityCollection::GetArchetypePtr(ArchetypeHandle);
+		const FMassArchetypeData* ArchetypePtr = FMassArchetypeHelper::ArchetypeDataFromHandle(ArchetypeHandle);
 		
 		// @todo if FMassArchetypeHandle used indices the look up would be a lot faster
 		int32 ArchetypeIndex = INDEX_NONE;
diff --git a/Engine/Plugins/Runtime/MassEntity/Source/MassEntity/Private/MassDebugger.cpp b/Engine/Plugins/Runtime/MassEntity/Source/MassEntity/Private/MassDebugger.cpp
--- a/Engine/Plugins/Runtime/MassEntity/Source/MassEntity/Private/MassDebugger.cpp
+++ b/Engine/Plugins/Runtime/MassEntity/Source/MassEntity/Private/MassDebugger.cpp
@@ -6,11 +6,200 @@
 #include "MassEntitySubsystem.h"
 #include "MassArchetypeTypes.h"
 #include "MassArchetypeData.h"
+#include "MassRequirements.h"
 #include "MassEntityQuery.h"
+#include "StructTypeBitSet.h"
+#include "Misc/OutputDevice.h"
+#include "Engine/World.h"
+#include "Engine/Engine.h"
 
 
 DEFINE_ENUM_TO_STRING(EMassFragmentAccess, "/Script/MassEntity");
 DEFINE_ENUM_TO_STRING(EMassFragmentPresence, "/Script/MassEntity");
+DEFINE_ENUM_TO_STRING(EMassProcessingPhase, "/Script/MassEntity");
+
+namespace UE::Mass::Debug
+{
+
+	FString DebugGetFragmentAccessString(EMassFragmentAccess Access)
+	{
+#if WITH_MASSENTITY_DEBUG
+		switch (Access)
+		{
+		case EMassFragmentAccess::None:	return TEXT("--");
+		case EMassFragmentAccess::ReadOnly:	return TEXT("RO");
+		case EMassFragmentAccess::ReadWrite:	return TEXT("RW");
+		default:
+			ensureMsgf(false, TEXT("Missing string conversion for EMassFragmentAccess=%d"), Access);
+			break;
+		}
+#endif // WITH_MASSENTITY_DEBUG
+		return TEXT("Missing string conversion");
+	}
+
+	void DebugOutputDescription(TConstArrayView<UMassProcessor*> Processors, FOutputDevice& Ar)
+	{
+#if WITH_MASSENTITY_DEBUG
+		const bool bAutoLineEnd = Ar.GetAutoEmitLineTerminator();
+		Ar.SetAutoEmitLineTerminator(false);
+		for (const UMassProcessor* Proc : Processors)
+		{
+			if (Proc)
+			{
+				Proc->DebugOutputDescription(Ar);
+				Ar.Logf(TEXT("\n"));
+			}
+			else
+			{
+				Ar.Logf(TEXT("NULL\n"));
+			}
+		}
+		Ar.SetAutoEmitLineTerminator(bAutoLineEnd);
+#endif // WITH_MASSENTITY_DEBUG
+	}
+
+#if WITH_MASSENTITY_DEBUG && WITH_MASSENTITY_DEBUG
+	FAutoConsoleCommandWithWorldArgsAndOutputDevice PrintEntityFragmentsCmd(
+		TEXT("mass.PrintEntityFragments"),
+		TEXT("Prints all fragment types and values (uproperties) for the specified Entity index"),
+		FConsoleCommandWithWorldArgsAndOutputDeviceDelegate::CreateLambda(
+			[](const TArray<FString>& Params, UWorld* World, FOutputDevice& Ar)
+			{
+				check(World);
+				if (UMassEntitySubsystem* EntitySystem = World->GetSubsystem<UMassEntitySubsystem>())
+				{
+					int32 Index = INDEX_NONE;
+					if (LexTryParseString<int32>(Index, *Params[0]))
+					{
+						FMassDebugger::OutputEntityDescription(Ar, *EntitySystem, Index);
+					}
+					else
+					{
+						Ar.Logf(ELogVerbosity::Error, TEXT("Entity index parameter must be an integer"));
+					}
+				}
+				else
+				{
+					Ar.Logf(ELogVerbosity::Error, TEXT("Failed to find MassEntitySubsystem for world %s"), *GetPathNameSafe(World));
+				}
+			})
+	);
+
+	FAutoConsoleCommandWithWorldArgsAndOutputDevice LogArchetypesCmd(
+		TEXT("mass.LogArchetypes"),
+		TEXT("Dumps description of archetypes to log. Optional parameter controls whether to include or exclude non-occupied archetypes. Defaults to 'include'."),
+		FConsoleCommandWithWorldArgsAndOutputDeviceDelegate::CreateLambda([](const TArray<FString>& Params, UWorld*, FOutputDevice& Ar)
+			{
+				const TIndirectArray<FWorldContext>& WorldContexts = GEngine->GetWorldContexts();
+				for (const FWorldContext& Context : WorldContexts)
+				{
+					UWorld* World = Context.World();
+					if (World == nullptr || World->IsPreviewWorld())
+					{
+						continue;
+					}
+
+					Ar.Logf(ELogVerbosity::Log, TEXT("Dumping description of archetypes for world: %s (%s - %s)"),
+						*GetPathNameSafe(World),
+						LexToString(World->WorldType),
+						*ToString(World->GetNetMode()));
+
+					if (UMassEntitySubsystem* EntitySystem = World->GetSubsystem<UMassEntitySubsystem>())
+					{
+						bool bIncludeEmpty = true;
+						if (Params.Num())
+						{
+							LexTryParseString(bIncludeEmpty, *Params[0]);
+						}
+						Ar.Logf(ELogVerbosity::Log, TEXT("Include empty archetypes: %s"), bIncludeEmpty ? TEXT("TRUE") : TEXT("FALSE"));
+						EntitySystem->DebugGetArchetypesStringDetails(Ar, bIncludeEmpty);
+					}
+					else
+					{
+						Ar.Logf(ELogVerbosity::Error, TEXT("Failed to find MassEntitySubsystem for world: %s (%s - %s)"),
+							*GetPathNameSafe(World),
+							LexToString(World->WorldType),
+							*ToString(World->GetNetMode()));
+					}
+				}
+			})
+	);
+
+	// @todo these console commands will be reparented to "massentities" domain once we rename and shuffle the modules around 
+	FAutoConsoleCommandWithWorld RecacheQueries(
+		TEXT("mass.RecacheQueries"),
+		TEXT("Forces EntityQueries to recache their valid archetypes"),
+		FConsoleCommandWithWorldDelegate::CreateLambda([](UWorld* InWorld)
+			{
+				check(InWorld);
+				if (UMassEntitySubsystem* System = InWorld->GetSubsystem<UMassEntitySubsystem>())
+				{
+					System->DebugForceArchetypeDataVersionBump();
+				}
+			}
+	));
+
+	FAutoConsoleCommandWithWorldArgsAndOutputDevice LogFragmentSizes(
+		TEXT("mass.LogFragmentSizes"),
+		TEXT("Logs all the fragment types being used along with their sizes."),
+		FConsoleCommandWithWorldArgsAndOutputDeviceDelegate::CreateLambda([](const TArray<FString>& Params, UWorld* World, FOutputDevice& Ar)
+			{
+				for (const TWeakObjectPtr<const UScriptStruct>& WeakStruct : FMassFragmentBitSet::DebugGetAllStructTypes())
+				{
+					if (const UScriptStruct* StructType = WeakStruct.Get())
+					{
+						Ar.Logf(ELogVerbosity::Log, TEXT("%s, size: %d"), *StructType->GetName(), StructType->GetStructureSize());
+					}
+				}
+			})
+	);
+
+	FAutoConsoleCommandWithWorldArgsAndOutputDevice LogMemoryUsage(
+		TEXT("mass.LogMemoryUsage"),
+		TEXT("Logs how much memory the mass entity system uses"),
+		FConsoleCommandWithWorldArgsAndOutputDeviceDelegate::CreateLambda([](const TArray<FString>& Params, UWorld* World, FOutputDevice& Ar)
+			{
+				check(World);
+				if (UMassEntitySubsystem* System = World->GetSubsystem<UMassEntitySubsystem>())
+				{
+					FResourceSizeEx CumulativeResourceSize;
+					System->GetResourceSizeEx(CumulativeResourceSize);
+					Ar.Logf(ELogVerbosity::Log, TEXT("MassEntity system uses: %d bytes"), CumulativeResourceSize.GetDedicatedSystemMemoryBytes());
+				}
+			}));
+
+	FAutoConsoleCommandWithOutputDevice LogFragments(
+		TEXT("mass.LogKnownFragments"),
+		TEXT("Logs all the known tags and fragments along with their \"index\" as stored via bitsets."),
+		FConsoleCommandWithOutputDeviceDelegate::CreateStatic([](FOutputDevice& OutputDevice)
+			{
+				auto PrintKnownTypes = [&OutputDevice](TConstArrayView<TWeakObjectPtr<const UScriptStruct>> AllStructs) {
+					int i = 0;
+					for (TWeakObjectPtr<const UScriptStruct> Struct : AllStructs)
+					{
+						if (Struct.IsValid())
+						{
+							OutputDevice.Logf(TEXT("\t%d. %s"), i++, *Struct->GetName());
+						}
+					}
+				};
+
+				OutputDevice.Logf(TEXT("Known tags:"));
+				PrintKnownTypes(FMassTagBitSet::DebugGetAllStructTypes());
+
+				OutputDevice.Logf(TEXT("Known Fragments:"));
+				PrintKnownTypes(FMassFragmentBitSet::DebugGetAllStructTypes());
+
+				OutputDevice.Logf(TEXT("Known Shared Fragments:"));
+				PrintKnownTypes(FMassSharedFragmentBitSet::DebugGetAllStructTypes());
+
+				OutputDevice.Logf(TEXT("Known Chunk Fragments:"));
+				PrintKnownTypes(FMassChunkFragmentBitSet::DebugGetAllStructTypes());
+			}));
+
+#endif // WITH_MASSENTITY_DEBUG
+
+} // namespace UE::Mass::Debug
 
 //----------------------------------------------------------------------//
 // FMassDebugger
@@ -35,7 +224,7 @@ TConstArrayView<FMassEntityQuery*> FMassDebugger::GetUpToDateProcessorQueries(co
 
 UE::Mass::Debug::FQueryRequirementsView FMassDebugger::GetQueryRequirements(const FMassEntityQuery& Query)
 {
-	UE::Mass::Debug::FQueryRequirementsView View = { Query.Requirements, Query.ChunkRequirements, Query.ConstSharedRequirements, Query.SharedRequirements
+	UE::Mass::Debug::FQueryRequirementsView View = { Query.FragmentRequirements, Query.ChunkFragmentRequirements, Query.ConstSharedFragmentRequirements, Query.SharedFragmentRequirements
 		, Query.RequiredAllTags, Query.RequiredAnyTags, Query.RequiredNoneTags
 		, Query.RequiredConstSubsystems, Query.RequiredMutableSubsystems };
 
@@ -55,7 +244,7 @@ TArray<FMassArchetypeHandle> FMassDebugger::GetAllArchetypes(const UMassEntitySu
 	{
 		for (const TSharedPtr<FMassArchetypeData>& Archetype : KVP.Value)
 		{
-			Archetypes.Add(FMassArchetypeHandle(Archetype));
+			Archetypes.Add(FMassArchetypeHelper::ArchetypeHandleFromData(Archetype));
 		}
 	}
 
@@ -64,23 +253,23 @@ TArray<FMassArchetypeHandle> FMassDebugger::GetAllArchetypes(const UMassEntitySu
 
 const FMassArchetypeCompositionDescriptor& FMassDebugger::GetArchetypeComposition(const FMassArchetypeHandle& ArchetypeHandle)
 {
-	check(ArchetypeHandle.IsValid());
-	return ArchetypeHandle.DataPtr->CompositionDescriptor;
+	const FMassArchetypeData& ArchetypeData = FMassArchetypeHelper::ArchetypeDataFromHandleChecked(ArchetypeHandle);
+	return ArchetypeData.CompositionDescriptor;
 }
 
 const FMassArchetypeSharedFragmentValues& FMassDebugger::GetArchetypeSharedFragmentValues(const FMassArchetypeHandle& ArchetypeHandle)
 {
-	check(ArchetypeHandle.IsValid());
-	return ArchetypeHandle.DataPtr->SharedFragmentValues;
+	const FMassArchetypeData& ArchetypeData = FMassArchetypeHelper::ArchetypeDataFromHandleChecked(ArchetypeHandle);
+	return ArchetypeData.SharedFragmentValues;
 }
 
 void FMassDebugger::GetArchetypeEntityStats(const FMassArchetypeHandle& ArchetypeHandle
 	, int32& OutEntitiesCount, int32& OutNumEntitiesPerChunk, int32& OutChunksCount)
 {
-	check(ArchetypeHandle.IsValid());
-	OutEntitiesCount = ArchetypeHandle.DataPtr->GetNumEntities();
-	OutNumEntitiesPerChunk = ArchetypeHandle.DataPtr->GetNumEntitiesPerChunk();
-	OutChunksCount = ArchetypeHandle.DataPtr->GetChunkCount();
+	const FMassArchetypeData& ArchetypeData = FMassArchetypeHelper::ArchetypeDataFromHandleChecked(ArchetypeHandle);
+	OutEntitiesCount = ArchetypeData.GetNumEntities();
+	OutNumEntitiesPerChunk = ArchetypeData.GetNumEntitiesPerChunk();
+	OutChunksCount = ArchetypeData.GetChunkCount();
 }
 
 TConstArrayView<struct UMassCompositeProcessor::FDependencyNode> FMassDebugger::GetProcessingGraph(const UMassCompositeProcessor& GraphOwner)
@@ -88,4 +277,147 @@ TConstArrayView<struct UMassCompositeProcessor::FDependencyNode> FMassDebugger::
 	return GraphOwner.ProcessingFlatGraph;
 }
 
+FString FMassDebugger::GetRequirementsDescription(const FMassFragmentRequirements& Requirements)
+{
+	TStringBuilder<256> StringBuilder;
+	StringBuilder.Append(TEXT("<"));
+
+	bool bNeedsComma = false;
+	for (const FMassFragmentRequirementDescription& Requirement : Requirements.FragmentRequirements)
+	{
+		if (bNeedsComma)
+		{
+			StringBuilder.Append(TEXT(","));
+		}
+		StringBuilder.Append(*FMassDebugger::GetSingleRequirementDescription(Requirement));
+		bNeedsComma = true;
+	}
+
+	StringBuilder.Append(TEXT(">"));
+	return StringBuilder.ToString();
+}
+
+FString FMassDebugger::GetArchetypeRequirementCompatibilityDescription(const FMassFragmentRequirements& Requirements, const FMassArchetypeHandle& ArchetypeHandle)
+{
+	if (ArchetypeHandle.IsValid() == false)
+	{
+		return TEXT("Invalid");
+	}
+
+	const FMassArchetypeData& ArchetypeData = FMassArchetypeHelper::ArchetypeDataFromHandleChecked(ArchetypeHandle);
+	FStringOutputDevice OutDescription;
+
+	const FMassArchetypeCompositionDescriptor& ArchetypeComposition = ArchetypeData.GetCompositionDescriptor();
+
+	if (ArchetypeComposition.Fragments.HasAll(Requirements.RequiredAllFragments) == false)
+	{
+		// missing one of the strictly required fragments
+		OutDescription += TEXT("\nMissing required fragments: ");
+		(Requirements.RequiredAllFragments - ArchetypeComposition.Fragments).DebugGetStringDesc(OutDescription);
+	}
+
+	if (Requirements.RequiredAnyFragments.IsEmpty() == false && ArchetypeComposition.Fragments.HasAny(Requirements.RequiredAnyFragments) == false)
+	{
+		// missing all of the "any" fragments
+		OutDescription += TEXT("\nMissing all \'any\' fragments: ");
+		Requirements.RequiredAnyFragments.DebugGetStringDesc(OutDescription);
+	}
+
+	if (ArchetypeComposition.Fragments.HasNone(Requirements.RequiredNoneFragments) == false)
+	{
+		// has some of the fragments required absent
+		OutDescription += TEXT("\nHas fragments required absent: ");
+		Requirements.RequiredNoneFragments.DebugGetStringDesc(OutDescription);
+	}
+
+	if (ArchetypeComposition.Tags.HasAll(Requirements.RequiredAllTags) == false)
+	{
+		// missing one of the strictly required tags
+		OutDescription += TEXT("\nMissing required tags: ");
+		(Requirements.RequiredAllTags - ArchetypeComposition.Tags).DebugGetStringDesc(OutDescription);
+	}
+
+	if (Requirements.RequiredAnyTags.IsEmpty() == false && ArchetypeComposition.Tags.HasAny(Requirements.RequiredAnyTags) == false)
+	{
+		// missing all of the "any" tags
+		OutDescription += TEXT("\nMissing all \'any\' tags: ");
+		Requirements.RequiredAnyTags.DebugGetStringDesc(OutDescription);
+	}
+
+	if (ArchetypeComposition.Tags.HasNone(Requirements.RequiredNoneTags) == false)
+	{
+		// has some of the tags required absent
+		OutDescription += TEXT("\nHas tags required absent: ");
+		Requirements.RequiredNoneTags.DebugGetStringDesc(OutDescription);
+	}
+
+	if (ArchetypeComposition.ChunkFragments.HasAll(Requirements.RequiredAllChunkFragments) == false)
+	{
+		// missing one of the strictly required chunk fragments
+		OutDescription += TEXT("\nMissing required chunk fragments: ");
+		(Requirements.RequiredAllChunkFragments - ArchetypeComposition.ChunkFragments).DebugGetStringDesc(OutDescription);
+	}
+
+	if (ArchetypeComposition.ChunkFragments.HasNone(Requirements.RequiredNoneChunkFragments) == false)
+	{
+		// has some of the chunk fragments required absent
+		OutDescription += TEXT("\nHas chunk fragments required absent: ");
+		Requirements.RequiredNoneChunkFragments.DebugGetStringDesc(OutDescription);
+	}
+
+	return OutDescription.Len() > 0 ? static_cast<FString>(OutDescription) : TEXT("Match");
+}
+
+FString FMassDebugger::GetSingleRequirementDescription(const FMassFragmentRequirementDescription& Requirement)
+{
+	return FString::Printf(TEXT("%s%s[%s]"), Requirement.IsOptional() ? TEXT("?") : (Requirement.Presence == EMassFragmentPresence::None ? TEXT("-") : TEXT("+"))
+		, *GetNameSafe(Requirement.StructType), *UE::Mass::Debug::DebugGetFragmentAccessString(Requirement.AccessMode));
+}
+
+void FMassDebugger::OutputArchetypeDescription(FOutputDevice& Ar, const FMassArchetypeHandle& ArchetypeHandle)
+{
+	Ar.Logf(TEXT("%s"), ArchetypeHandle.IsValid() ? *FMassArchetypeHelper::ArchetypeDataFromHandleChecked(ArchetypeHandle).DebugGetDescription() : TEXT("INVALID"));
+}
+
+void FMassDebugger::OutputEntityDescription(FOutputDevice& Ar, const UMassEntitySubsystem& EntitySubsystem, const int32 EntityIndex, const TCHAR* InPrefix)
+{
+	if (EntityIndex >= EntitySubsystem.Entities.Num())
+	{
+		Ar.Logf(ELogVerbosity::Log, TEXT("Unable to list fragments values for out of range index in %s"), *GetPathNameSafe(&EntitySubsystem));
+		return;
+	}
+
+	const UMassEntitySubsystem::FEntityData& EntityData = EntitySubsystem.Entities[EntityIndex];
+	if (!EntityData.IsValid())
+	{
+		Ar.Logf(ELogVerbosity::Log, TEXT("Unable to list fragments values for invalid entity in %s"), *GetPathNameSafe(&EntitySubsystem));
+	}
+
+	FMassEntityHandle Entity;
+	Entity.Index = EntityIndex;
+	Entity.SerialNumber = EntityData.SerialNumber;
+	OutputEntityDescription(Ar, EntitySubsystem, Entity, InPrefix);
+}
+
+void FMassDebugger::OutputEntityDescription(FOutputDevice& Ar, const UMassEntitySubsystem& EntitySubsystem, const FMassEntityHandle Entity, const TCHAR* InPrefix)
+{
+	if (!EntitySubsystem.IsEntityActive(Entity))
+	{
+		Ar.Logf(ELogVerbosity::Log, TEXT("Unable to list fragments values for invalid entity in %s"), *GetPathNameSafe(&EntitySubsystem));
+	}
+
+	Ar.Logf(ELogVerbosity::Log, TEXT("Listing fragments values for Entity[%s] in %s"), *Entity.DebugGetDescription(), *GetPathNameSafe(&EntitySubsystem));
+
+	const UMassEntitySubsystem::FEntityData& EntityData = EntitySubsystem.Entities[Entity.Index];
+	FMassArchetypeData* Archetype = EntityData.CurrentArchetype.Get();
+	if (Archetype == nullptr)
+	{
+		Ar.Logf(ELogVerbosity::Log, TEXT("Unable to list fragments values for invalid entity in %s"), *GetPathNameSafe(&EntitySubsystem));
+	}
+	else
+	{
+		Archetype->DebugPrintEntity(Entity, Ar, InPrefix);
+	}
+}
+
 #endif // WITH_MASSENTITY_DEBUG
\ No newline at end of file
diff --git a/Engine/Plugins/Runtime/MassEntity/Source/MassEntity/Private/MassEntityDebug.cpp b/Engine/Plugins/Runtime/MassEntity/Source/MassEntity/Private/MassEntityDebug.cpp
deleted file mode 100644
diff --git a/Engine/Plugins/Runtime/MassEntity/Source/MassEntity/Private/MassEntityQuery.cpp b/Engine/Plugins/Runtime/MassEntity/Source/MassEntity/Private/MassEntityQuery.cpp
--- a/Engine/Plugins/Runtime/MassEntity/Source/MassEntity/Private/MassEntityQuery.cpp
+++ b/Engine/Plugins/Runtime/MassEntity/Source/MassEntity/Private/MassEntityQuery.cpp
@@ -1,7 +1,7 @@
 // Copyright Epic Games, Inc. All Rights Reserved.
 
 #include "MassEntityQuery.h"
-#include "MassEntityDebug.h"
+#include "MassDebugger.h"
 #include "MassEntitySubsystem.h"
 #include "MassArchetypeData.h"
 #include "MassCommandBuffer.h"
@@ -15,33 +15,9 @@
 #endif // WITH_MASSENTITY_DEBUG
 
 
-namespace UE::Mass::Private
-{
-	template<typename TContainer>
-	void ExportRequirements(TConstArrayView<FMassFragmentRequirement> Requirements, TMassExecutionAccess<TContainer>& Out)
-	{
-		for (const FMassFragmentRequirement& Requirement : Requirements)
-		{
-			if (Requirement.Presence != EMassFragmentPresence::None)
-			{
-				check(Requirement.StructType);
-				if (Requirement.AccessMode == EMassFragmentAccess::ReadOnly)
-				{
-					Out.Read.Add(*Requirement.StructType);
-				}
-				else if (Requirement.AccessMode == EMassFragmentAccess::ReadWrite)
-				{
-					Out.Write.Add(*Requirement.StructType);
-				}
-			}
-		}
-	}
-}
-
 //////////////////////////////////////////////////////////////////////
 // FMassEntityQuery
 
-
 FMassEntityQuery::FMassEntityQuery()
 {
 	bAllowParallelExecution = false;
@@ -92,23 +68,13 @@ void FMassEntityQuery::ReadCommandlineParams()
 	}
 }
 
-void FMassEntityQuery::SortRequirements()
-{
-	// we're sorting the Requirements the same way ArchetypeData's FragmentConfig is sorted (see FMassArchetypeData::Initialize)
-	// so that when we access ArchetypeData.FragmentConfigs in FMassArchetypeData::BindRequirementsWithMapping
-	// (via GetFragmentData call) the access is sequential (i.e. not random) and there's a higher chance the memory
-	// FragmentConfigs we want to access have already been fetched and are available in processor cache.
-	Requirements.Sort(FScriptStructSortOperator());
-	ChunkRequirements.Sort(FScriptStructSortOperator());
-	ConstSharedRequirements.Sort(FScriptStructSortOperator());
-	SharedRequirements.Sort(FScriptStructSortOperator());
-}
-
 void FMassEntityQuery::CacheArchetypes(const UMassEntitySubsystem& InEntitySubsystem)
 {
 	const uint32 InEntitySubsystemHash = PointerHash(&InEntitySubsystem);
-	if (EntitySubsystemHash != InEntitySubsystemHash || InEntitySubsystem.GetArchetypeDataVersion() != ArchetypeDataVersion)
+	if (EntitySubsystemHash != InEntitySubsystemHash || InEntitySubsystem.GetArchetypeDataVersion() != ArchetypeDataVersion || IncrementalChangesCount)
 	{
+		IncrementalChangesCount = 0;
+
 		if (CheckValidity())
 		{
 			SortRequirements();
@@ -119,58 +85,34 @@ void FMassEntityQuery::CacheArchetypes(const UMassEntitySubsystem& InEntitySubsy
 			ArchetypeDataVersion = InEntitySubsystem.GetArchetypeDataVersion();
 
 			TRACE_CPUPROFILER_EVENT_SCOPE_STR("Mass RequirementsBinding")
-			const TConstArrayView<FMassFragmentRequirement> LocalRequirements = GetRequirements();
+			const TConstArrayView<FMassFragmentRequirementDescription> LocalRequirements = GetFragmentRequirements();
 			ArchetypeFragmentMapping.Reset(ValidArchetypes.Num());
 			ArchetypeFragmentMapping.AddDefaulted(ValidArchetypes.Num());
 			for (int i = 0; i < ValidArchetypes.Num(); ++i)
 			{
-				ValidArchetypes[i].DataPtr->GetRequirementsFragmentMapping(LocalRequirements, ArchetypeFragmentMapping[i].EntityFragments);
-				if (ChunkRequirements.Num())
+				FMassArchetypeData& ArchetypeData = FMassArchetypeHelper::ArchetypeDataFromHandleChecked(ValidArchetypes[i]);
+				ArchetypeData.GetRequirementsFragmentMapping(LocalRequirements, ArchetypeFragmentMapping[i].EntityFragments);
+				if (ChunkFragmentRequirements.Num())
 				{
-					ValidArchetypes[i].DataPtr->GetRequirementsChunkFragmentMapping(ChunkRequirements, ArchetypeFragmentMapping[i].ChunkFragments);
+					ArchetypeData.GetRequirementsChunkFragmentMapping(ChunkFragmentRequirements, ArchetypeFragmentMapping[i].ChunkFragments);
 				}
-				if (ConstSharedRequirements.Num())
+				if (ConstSharedFragmentRequirements.Num())
 				{
-					ValidArchetypes[i].DataPtr->GetRequirementsConstSharedFragmentMapping(ConstSharedRequirements, ArchetypeFragmentMapping[i].ConstSharedFragments);
+					ArchetypeData.GetRequirementsConstSharedFragmentMapping(ConstSharedFragmentRequirements, ArchetypeFragmentMapping[i].ConstSharedFragments);
 				}
-				if (SharedRequirements.Num())
+				if (SharedFragmentRequirements.Num())
 				{
-					ValidArchetypes[i].DataPtr->GetRequirementsSharedFragmentMapping(SharedRequirements, ArchetypeFragmentMapping[i].SharedFragments);
+					ArchetypeData.GetRequirementsSharedFragmentMapping(SharedFragmentRequirements, ArchetypeFragmentMapping[i].SharedFragments);
 				}
 			}
 		}
 		else
 		{
-			UE_VLOG_UELOG(&InEntitySubsystem, LogMass, Error, TEXT("FMassEntityQuery::CacheArchetypes: requirements not valid: %s"), *DebugGetDescription());
+			UE_VLOG_UELOG(&InEntitySubsystem, LogMass, Error, TEXT("FMassEntityQuery::CacheArchetypes: requirements not valid: %s"), *FMassDebugger::GetRequirementsDescription(*this));
 		}
 	}
 }
 
-bool FMassEntityQuery::CheckValidity() const
-{
-	return RequiredAllFragments.IsEmpty() == false || RequiredAnyFragments.IsEmpty() == false || RequiredOptionalFragments.IsEmpty() == false;
-}
-
-bool FMassEntityQuery::DoesArchetypeMatchRequirements(const FMassArchetypeHandle& ArchetypeHandle) const
-{
-	check(ArchetypeHandle.IsValid());
-	const FMassArchetypeData* Archetype = ArchetypeHandle.DataPtr.Get();
-	CA_ASSUME(Archetype);
-	
-	const FMassArchetypeCompositionDescriptor& ArchetypeComposition = Archetype->GetCompositionDescriptor();
-
-	return ArchetypeComposition.Fragments.HasAll(RequiredAllFragments)
-		&& (RequiredAnyFragments.IsEmpty() || ArchetypeComposition.Fragments.HasAny(RequiredAnyFragments))
-		&& ArchetypeComposition.Fragments.HasNone(RequiredNoneFragments)
-		&& ArchetypeComposition.Tags.HasAll(RequiredAllTags)
-		&& (RequiredAnyTags.IsEmpty() || ArchetypeComposition.Tags.HasAny(RequiredAnyTags))
-		&& ArchetypeComposition.Tags.HasNone(RequiredNoneTags)
-		&& ArchetypeComposition.ChunkFragments.HasAll(RequiredAllChunkFragments)
-		&& ArchetypeComposition.ChunkFragments.HasNone(RequiredNoneChunkFragments)
-		&& ArchetypeComposition.SharedFragments.HasAll(RequiredAllSharedFragments)
-		&& ArchetypeComposition.SharedFragments.HasNone(RequiredNoneSharedFragments);
-}
-
 void FMassEntityQuery::ForEachEntityChunk(const FMassArchetypeEntityCollection& Collection, UMassEntitySubsystem& EntitySubsystem, FMassExecutionContext& ExecutionContext, const FMassExecuteFunction& ExecuteFunction)
 {
 	// mz@todo I don't like that we're copying data here.
@@ -200,16 +142,17 @@ void FMassEntityQuery::ForEachEntityChunk(UMassEntitySubsystem& EntitySubsystem,
 		if (DoesArchetypeMatchRequirements(ArchetypeHandle) == false)
 		{
 			UE_VLOG_UELOG(&EntitySubsystem, LogMass, Log, TEXT("Attempted to execute FMassEntityQuery with an incompatible Archetype: %s")
-				, *DebugGetArchetypeCompatibilityDescription(ArchetypeHandle));
+				, *FMassDebugger::GetArchetypeRequirementCompatibilityDescription(*this, ArchetypeHandle));
 
 #if WITH_MASSENTITY_DEBUG
 			EntitySubsystem.GetRequirementAccessDetector().ReleaseAccess(*this);
 #endif // WITH_MASSENTITY_DEBUG
 			return;
 		}
-		ExecutionContext.SetRequirements(Requirements, ChunkRequirements, ConstSharedRequirements, SharedRequirements);
+		ExecutionContext.SetRequirements(FragmentRequirements, ChunkFragmentRequirements, ConstSharedFragmentRequirements, SharedFragmentRequirements);
 		
-		ArchetypeHandle.DataPtr->ExecuteFunction(ExecutionContext, ExecuteFunction
+		FMassArchetypeData& ArchetypeData = FMassArchetypeHelper::ArchetypeDataFromHandleChecked(ArchetypeHandle);
+		ArchetypeData.ExecuteFunction(ExecutionContext, ExecuteFunction
 			, GetFragmentMappingForArchetype(ArchetypeHandle)
 			, ExecutionContext.GetEntityCollection().GetRanges());
 #if WITH_MASSENTITY_DEBUG
@@ -220,13 +163,13 @@ void FMassEntityQuery::ForEachEntityChunk(UMassEntitySubsystem& EntitySubsystem,
 	{
 		CacheArchetypes(EntitySubsystem);
 		// it's important to set requirements after caching archetypes due to that call potentially sorting the requirements and the order is relevant here.
-		ExecutionContext.SetRequirements(Requirements, ChunkRequirements, ConstSharedRequirements, SharedRequirements);
+		ExecutionContext.SetRequirements(FragmentRequirements, ChunkFragmentRequirements, ConstSharedFragmentRequirements, SharedFragmentRequirements);
 
 		for (int i = 0; i < ValidArchetypes.Num(); ++i)
 		{
-			FMassArchetypeHandle& Archetype = ValidArchetypes[i];
-			check(Archetype.IsValid());
-			Archetype.DataPtr->ExecuteFunction(ExecutionContext, ExecuteFunction, ArchetypeFragmentMapping[i], ArchetypeCondition, ChunkCondition);
+			const FMassArchetypeHandle& ArchetypeHandle = ValidArchetypes[i];
+			FMassArchetypeData& ArchetypeData = FMassArchetypeHelper::ArchetypeDataFromHandleChecked(ArchetypeHandle);
+			ArchetypeData.ExecuteFunction(ExecutionContext, ExecuteFunction, ArchetypeFragmentMapping[i], ArchetypeCondition, ChunkCondition);
 			ExecutionContext.ClearFragmentViews();
 #if WITH_MASSENTITY_DEBUG
 			NumEntitiesToProcess += ExecutionContext.GetNumEntities();
@@ -270,14 +213,13 @@ void FMassEntityQuery::ParallelForEachEntityChunk(UMassEntitySubsystem& EntitySu
 		if (DoesArchetypeMatchRequirements(ArchetypeHandle) == false)
 		{
 			UE_VLOG_UELOG(&EntitySubsystem, LogMass, Log, TEXT("Attempted to execute FMassEntityQuery with an incompatible Archetype: %s")
-				, *DebugGetArchetypeCompatibilityDescription(ExecutionContext.GetEntityCollection().GetArchetype()));
+				, *FMassDebugger::GetArchetypeRequirementCompatibilityDescription(*this, ExecutionContext.GetEntityCollection().GetArchetype()));
 			return;
 		}
 
-		ExecutionContext.SetRequirements(Requirements, ChunkRequirements, ConstSharedRequirements, SharedRequirements);
-		check(ArchetypeHandle.IsValid());
-		FMassArchetypeData& ArchetypeRef = *ArchetypeHandle.DataPtr.Get();
-		const FMassArchetypeEntityCollection AsEntityCollection(ArchetypeHandle.DataPtr);
+		ExecutionContext.SetRequirements(FragmentRequirements, ChunkFragmentRequirements, ConstSharedFragmentRequirements, SharedFragmentRequirements);
+		FMassArchetypeData& ArchetypeRef = FMassArchetypeHelper::ArchetypeDataFromHandleChecked(ArchetypeHandle);
+		const FMassArchetypeEntityCollection AsEntityCollection(ArchetypeHandle);
 		for (const FMassArchetypeEntityCollection::FArchetypeEntityRange& EntityRange : AsEntityCollection.GetRanges())
 		{
 			Jobs.Add({ ArchetypeRef, INDEX_NONE, EntityRange });
@@ -286,13 +228,12 @@ void FMassEntityQuery::ParallelForEachEntityChunk(UMassEntitySubsystem& EntitySu
 	else
 	{
 		CacheArchetypes(EntitySubsystem);
-		ExecutionContext.SetRequirements(Requirements, ChunkRequirements, ConstSharedRequirements, SharedRequirements);
+		ExecutionContext.SetRequirements(FragmentRequirements, ChunkFragmentRequirements, ConstSharedFragmentRequirements, SharedFragmentRequirements);
 		for (int ArchetypeIndex = 0; ArchetypeIndex < ValidArchetypes.Num(); ++ArchetypeIndex)
 		{
-			FMassArchetypeHandle& Archetype = ValidArchetypes[ArchetypeIndex];
-			check(Archetype.IsValid());
-			FMassArchetypeData& ArchetypeRef = *Archetype.DataPtr.Get();
-			const FMassArchetypeEntityCollection AsEntityCollection(Archetype.DataPtr);
+			FMassArchetypeHandle& ArchetypeHandle = ValidArchetypes[ArchetypeIndex];
+			FMassArchetypeData& ArchetypeRef = FMassArchetypeHelper::ArchetypeDataFromHandleChecked(ArchetypeHandle);
+			const FMassArchetypeEntityCollection AsEntityCollection(ArchetypeHandle);
 			for (const FMassArchetypeEntityCollection::FArchetypeEntityRange& EntityRange : AsEntityCollection.GetRanges())
 			{
 				Jobs.Add({ArchetypeRef, ArchetypeIndex, EntityRange});
@@ -315,7 +256,7 @@ void FMassEntityQuery::ParallelForEachEntityChunk(UMassEntitySubsystem& EntitySu
 void FMassEntityQuery::ApplyQueryRequirementsToContext(FMassExecutionContext& ExecutionContext)
 {
 	ExecutionContext.SetSubsystemRequirements(RequiredConstSubsystems, RequiredMutableSubsystems);
-	ExecutionContext.SetRequirements(Requirements, ChunkRequirements, ConstSharedRequirements, SharedRequirements);
+	ExecutionContext.SetRequirements(FragmentRequirements, ChunkFragmentRequirements, ConstSharedFragmentRequirements, SharedFragmentRequirements);
 }
 
 int32 FMassEntityQuery::GetNumMatchingEntities(UMassEntitySubsystem& InEntitySubsystem)
@@ -324,7 +265,7 @@ int32 FMassEntityQuery::GetNumMatchingEntities(UMassEntitySubsystem& InEntitySub
 	int32 TotalEntities = 0;
 	for (FMassArchetypeHandle& ArchetypeHandle : ValidArchetypes)
 	{
-		if (const FMassArchetypeData* Archetype = ArchetypeHandle.DataPtr.Get())
+		if (const FMassArchetypeData* Archetype = FMassArchetypeHelper::ArchetypeDataFromHandle(ArchetypeHandle))
 		{
 			TotalEntities += Archetype->GetNumEntities();
 		}
@@ -338,7 +279,7 @@ bool FMassEntityQuery::HasMatchingEntities(UMassEntitySubsystem& InEntitySubsyst
 
 	for (FMassArchetypeHandle& ArchetypeHandle : ValidArchetypes)
 	{
-		const FMassArchetypeData* Archetype = ArchetypeHandle.DataPtr.Get();
+		const FMassArchetypeData* Archetype = FMassArchetypeHelper::ArchetypeDataFromHandle(ArchetypeHandle);
 		if (Archetype && Archetype->GetNumEntities() > 0)
 		{
 			return true;
@@ -356,126 +297,6 @@ const FMassQueryRequirementIndicesMapping& FMassEntityQuery::GetFragmentMappingF
 
 void FMassEntityQuery::ExportRequirements(FMassExecutionRequirements& OutRequirements) const
 {
-	OutRequirements.RequiredSubsystems.Read += RequiredConstSubsystems;
-	OutRequirements.RequiredSubsystems.Write += RequiredMutableSubsystems;
-
-	using UE::Mass::Private::ExportRequirements;
-	ExportRequirements<FMassFragmentBitSet>(Requirements, OutRequirements.Fragments);
-	ExportRequirements<FMassChunkFragmentBitSet>(ChunkRequirements, OutRequirements.ChunkFragments);
-	ExportRequirements<FMassSharedFragmentBitSet>(ConstSharedRequirements, OutRequirements.SharedFragments);
-	ExportRequirements<FMassSharedFragmentBitSet>(SharedRequirements, OutRequirements.SharedFragments);
-
-	OutRequirements.RequiredAllTags = RequiredAllTags;
-	OutRequirements.RequiredAnyTags = RequiredAnyTags;
-	OutRequirements.RequiredNoneTags = RequiredNoneTags;
-}
-
-FString FMassEntityQuery::DebugGetDescription() const
-{
-#if WITH_MASSENTITY_DEBUG
-	TStringBuilder<256> StringBuilder;
-	StringBuilder.Append(TEXT("<"));
-
-	bool bNeedsComma = false;
-	for (const FMassFragmentRequirement& Requirement : Requirements)
-	{
-		if (bNeedsComma)
-		{
-			StringBuilder.Append(TEXT(","));
-		}
-		StringBuilder.Append(*Requirement.DebugGetDescription());
-		bNeedsComma = true;
-	}
-
-	StringBuilder.Append(TEXT(">"));
-	return StringBuilder.ToString();
-#else
-	return {};
-#endif
-}
-
-FString FMassEntityQuery::DebugGetArchetypeCompatibilityDescription(const FMassArchetypeHandle& ArchetypeHandle) const
-{
-	if (ArchetypeHandle.IsValid() == false)
-	{
-		return TEXT("Invalid");
-	}
-	
-	const FMassArchetypeData& Archetype = *ArchetypeHandle.DataPtr;
-	FStringOutputDevice OutDescription;
-#if WITH_MASSENTITY_DEBUG	
-	const FMassArchetypeCompositionDescriptor& ArchetypeComposition = Archetype.GetCompositionDescriptor();
-
-	if (ArchetypeComposition.Fragments.HasAll(RequiredAllFragments) == false)
-	{
-		// missing one of the strictly required fragments
-		OutDescription += TEXT("\nMissing required fragments: ");
-		(RequiredAllFragments - ArchetypeComposition.Fragments).DebugGetStringDesc(OutDescription);
-	}
-
-	if (RequiredAnyFragments.IsEmpty() == false && ArchetypeComposition.Fragments.HasAny(RequiredAnyFragments) == false)
-	{
-		// missing all of the "any" fragments
-		OutDescription += TEXT("\nMissing all \'any\' fragments: ");
-		RequiredAnyFragments.DebugGetStringDesc(OutDescription);
-	}
-
-	if (ArchetypeComposition.Fragments.HasNone(RequiredNoneFragments) == false)
-	{
-		// has some of the fragments required absent
-		OutDescription += TEXT("\nHas fragments required absent: ");
-		RequiredNoneFragments.DebugGetStringDesc(OutDescription);
-	}
-
-	if (ArchetypeComposition.Tags.HasAll(RequiredAllTags) == false)
-	{
-		// missing one of the strictly required tags
-		OutDescription += TEXT("\nMissing required tags: ");
-		(RequiredAllTags - ArchetypeComposition.Tags).DebugGetStringDesc(OutDescription);
-	}
-
-	if (RequiredAnyTags.IsEmpty() == false && ArchetypeComposition.Tags.HasAny(RequiredAnyTags) == false)
-	{
-		// missing all of the "any" tags
-		OutDescription += TEXT("\nMissing all \'any\' tags: ");
-		RequiredAnyTags.DebugGetStringDesc(OutDescription);
-	}
-
-	if (ArchetypeComposition.Tags.HasNone(RequiredNoneTags) == false)
-	{
-		// has some of the tags required absent
-		OutDescription += TEXT("\nHas tags required absent: ");
-		RequiredNoneTags.DebugGetStringDesc(OutDescription);
-	}
-	
-	if (ArchetypeComposition.ChunkFragments.HasAll(RequiredAllChunkFragments) == false)
-    {
-    	// missing one of the strictly required chunk fragments
-    	OutDescription += TEXT("\nMissing required chunk fragments: ");
-		(RequiredAllChunkFragments - ArchetypeComposition.ChunkFragments).DebugGetStringDesc(OutDescription);
-    }
-
-    if (ArchetypeComposition.ChunkFragments.HasNone(RequiredNoneChunkFragments) == false)
-	{
-		// has some of the chunk fragments required absent
-		OutDescription += TEXT("\nHas chunk fragments required absent: ");
-		RequiredNoneChunkFragments.DebugGetStringDesc(OutDescription);
-	}
-#endif // WITH_MASSENTITY_DEBUG
-	
-	return OutDescription.Len() > 0 ? static_cast<FString>(OutDescription) : TEXT("Match");
-}
-
-
-//////////////////////////////////////////////////////////////////////
-// FMassFragmentRequirement
-
-FString FMassFragmentRequirement::DebugGetDescription() const
-{
-#if WITH_MASSENTITY_DEBUG
-	return FString::Printf(TEXT("%s%s[%s]"), IsOptional() ? TEXT("?") : (Presence == EMassFragmentPresence::None ? TEXT("-") : TEXT("+"))
-		, *GetNameSafe(StructType), *UE::Mass::Debug::DebugGetFragmentAccessString(AccessMode));
-#else
-	return {};
-#endif
+	FMassSubsystemRequirements::ExportRequirements(OutRequirements);
+	FMassFragmentRequirements::ExportRequirements(OutRequirements);
 }
diff --git a/Engine/Plugins/Runtime/MassEntity/Source/MassEntity/Private/MassEntitySettings.cpp b/Engine/Plugins/Runtime/MassEntity/Source/MassEntity/Private/MassEntitySettings.cpp
--- a/Engine/Plugins/Runtime/MassEntity/Source/MassEntity/Private/MassEntitySettings.cpp
+++ b/Engine/Plugins/Runtime/MassEntity/Source/MassEntity/Private/MassEntitySettings.cpp
@@ -4,7 +4,6 @@
 #include "MassProcessingPhase.h"
 #include "VisualLogger/VisualLogger.h"
 #include "UObject/UObjectHash.h"
-#include "MassEntityDebug.h"
 #include "Misc/CoreDelegates.h"
 #if WITH_EDITOR
 #include "ObjectEditorUtils.h"
@@ -18,7 +17,7 @@ UMassEntitySettings::UMassEntitySettings(const FObjectInitializer& ObjectInitial
 {
 	for (int i = 0; i < (int)EMassProcessingPhase::MAX; ++i)
 	{
-		ProcessingPhasesConfig[i].PhaseName = *EnumToString(EMassProcessingPhase(i));
+		ProcessingPhasesConfig[i].PhaseName = *UEnum::GetValueAsString(EMassProcessingPhase(i));
 	}
 
 	FCoreDelegates::OnPostEngineInit.AddUObject(this, &UMassEntitySettings::BuildProcessorListAndPhases);
diff --git a/Engine/Plugins/Runtime/MassEntity/Source/MassEntity/Private/MassEntitySubsystem.cpp b/Engine/Plugins/Runtime/MassEntity/Source/MassEntity/Private/MassEntitySubsystem.cpp
--- a/Engine/Plugins/Runtime/MassEntity/Source/MassEntity/Private/MassEntitySubsystem.cpp
+++ b/Engine/Plugins/Runtime/MassEntity/Source/MassEntity/Private/MassEntitySubsystem.cpp
@@ -173,35 +173,34 @@ FMassArchetypeHandle UMassEntitySubsystem::CreateArchetype(const FMassArchetypeC
 
 	TArray<TSharedPtr<FMassArchetypeData>>& HashRow = FragmentHashToArchetypeMap.FindOrAdd(TypeHash);
 
-	FMassArchetypeHandle Result;
+	TSharedPtr<FMassArchetypeData> ArchetypeDataPtr;
 	for (const TSharedPtr<FMassArchetypeData>& Ptr : HashRow)
 	{
 		if (Ptr->IsEquivalent(Composition, SharedFragmentValues))
 		{
-			Result.DataPtr = Ptr;
+			ArchetypeDataPtr = Ptr;
 			break;
 		}
 	}
 
-	if (!Result.DataPtr.IsValid())
+	if (!ArchetypeDataPtr.IsValid())
 	{
 		// Create a new archetype
-		const TSharedPtr<FMassArchetypeData> NewArchetype = MakeShareable(new FMassArchetypeData);
+		FMassArchetypeData* NewArchetype = new FMassArchetypeData();
 		NewArchetype->Initialize(Composition, SharedFragmentValues);
-		HashRow.Add(NewArchetype);
+		
+		ArchetypeDataPtr = HashRow.Add_GetRef(MakeShareable(NewArchetype));
 
 		for (const FMassArchetypeFragmentConfig& FragmentConfig : NewArchetype->GetFragmentConfigs())
 		{
 			checkSlow(FragmentConfig.FragmentType)
-			FragmentTypeToArchetypeMap.FindOrAdd(FragmentConfig.FragmentType).Add(NewArchetype);
+			FragmentTypeToArchetypeMap.FindOrAdd(FragmentConfig.FragmentType).Add(ArchetypeDataPtr);
 		}
 
-		Result.DataPtr = NewArchetype;
-
 		++ArchetypeDataVersion;
 	}
 
-	return MoveTemp(Result);
+	return FMassArchetypeHelper::ArchetypeHandleFromData(ArchetypeDataPtr);
 }
 
 FMassArchetypeHandle UMassEntitySubsystem::InternalCreateSimilarArchetype(const TSharedPtr<FMassArchetypeData>& SourceArchetype, const FMassTagBitSet& OverrideTags)
@@ -226,57 +225,54 @@ FMassArchetypeHandle UMassEntitySubsystem::InternalCreateSimilarArchetype(const
 
 	TArray<TSharedPtr<FMassArchetypeData>>& HashRow = FragmentHashToArchetypeMap.FindOrAdd(TypeHash);
 
-	FMassArchetypeHandle Result;
+	TSharedPtr<FMassArchetypeData> ArchetypeDataPtr;
 	for (const TSharedPtr<FMassArchetypeData>& Ptr : HashRow)
 	{
 		if (Ptr->IsEquivalent(NewComposition, SourceArchetypeRef.GetSharedFragmentValues()))
 		{
-			Result.DataPtr = Ptr;
+			ArchetypeDataPtr = Ptr;
 			break;
 		}
 	}
 
-	if (!Result.DataPtr.IsValid())
+	if (!ArchetypeDataPtr.IsValid())
 	{
 		// Create a new archetype
-		const TSharedPtr<FMassArchetypeData> NewArchetype = MakeShareable(new FMassArchetypeData);
+		FMassArchetypeData* NewArchetype = new FMassArchetypeData();
 		NewArchetype->InitializeWithSimilar(SourceArchetypeRef, MoveTemp(NewComposition));
-		HashRow.Add(NewArchetype);
+
+		ArchetypeDataPtr = HashRow.Add_GetRef(MakeShareable(NewArchetype));
 
 		for (const FMassArchetypeFragmentConfig& FragmentConfig : NewArchetype->GetFragmentConfigs())
 		{
 			checkSlow(FragmentConfig.FragmentType)
-			FragmentTypeToArchetypeMap.FindOrAdd(FragmentConfig.FragmentType).Add(NewArchetype);
+			FragmentTypeToArchetypeMap.FindOrAdd(FragmentConfig.FragmentType).Add(ArchetypeDataPtr);
 		}
 
-		Result.DataPtr = NewArchetype;
-
 		++ArchetypeDataVersion;
 	}
 
-	return MoveTemp(Result);
+	return FMassArchetypeHelper::ArchetypeHandleFromData(ArchetypeDataPtr);
 }
 
 FMassArchetypeHandle UMassEntitySubsystem::GetArchetypeForEntity(FMassEntityHandle Entity) const
 {
-	FMassArchetypeHandle Result;
 	if (IsEntityValid(Entity))
 	{
-		Result.DataPtr = Entities[Entity.Index].CurrentArchetype;
+		return FMassArchetypeHelper::ArchetypeHandleFromData(Entities[Entity.Index].CurrentArchetype);
 	}
-	return Result;
+	return FMassArchetypeHandle();
 }
 
 FMassArchetypeHandle UMassEntitySubsystem::GetArchetypeForEntityUnsafe(FMassEntityHandle Entity) const
 {
 	check(Entities.IsValidIndex(Entity.Index));
-	return FMassArchetypeHandle(Entities[Entity.Index].CurrentArchetype);
+	return FMassArchetypeHelper::ArchetypeHandleFromData(Entities[Entity.Index].CurrentArchetype);
 }
 
-void UMassEntitySubsystem::ForEachArchetypeFragmentType(const FMassArchetypeHandle Archetype, TFunction< void(const UScriptStruct* /*FragmentType*/)> Function)
+void UMassEntitySubsystem::ForEachArchetypeFragmentType(const FMassArchetypeHandle& ArchetypeHandle, TFunction< void(const UScriptStruct* /*FragmentType*/)> Function)
 {
-	check(Archetype.DataPtr.IsValid());
-	const FMassArchetypeData& ArchetypeData = *Archetype.DataPtr.Get();
+	const FMassArchetypeData& ArchetypeData = FMassArchetypeHelper::ArchetypeDataFromHandleChecked(ArchetypeHandle);
 	ArchetypeData.ForEachFragmentType(Function);
 }
 
@@ -304,12 +300,12 @@ void UMassEntitySubsystem::DoEntityCompaction(const double TimeAllowed)
 	}
 }
 
-FMassEntityHandle UMassEntitySubsystem::CreateEntity(const FMassArchetypeHandle Archetype)
+FMassEntityHandle UMassEntitySubsystem::CreateEntity(const FMassArchetypeHandle& ArchetypeHandle)
 {
-	check(Archetype.DataPtr.IsValid());
+	check(ArchetypeHandle.IsValid());
 
 	const FMassEntityHandle Entity = ReserveEntity();
-	InternalBuildEntity(Entity, Archetype);
+	InternalBuildEntity(Entity, ArchetypeHandle);
 	return Entity;
 }
 
@@ -317,11 +313,11 @@ FMassEntityHandle UMassEntitySubsystem::CreateEntity(TConstArrayView<FInstancedS
 {
 	check(FragmentInstanceList.Num() > 0);
 
-	const FMassArchetypeHandle Archetype = CreateArchetype(FMassArchetypeCompositionDescriptor(FragmentInstanceList, FMassTagBitSet(), FMassChunkFragmentBitSet(), FMassSharedFragmentBitSet()), FMassArchetypeSharedFragmentValues());
-	check(Archetype.DataPtr.IsValid());
+	const FMassArchetypeHandle& ArchetypeHandle = CreateArchetype(FMassArchetypeCompositionDescriptor(FragmentInstanceList, FMassTagBitSet(), FMassChunkFragmentBitSet(), FMassSharedFragmentBitSet()), FMassArchetypeSharedFragmentValues());
+	check(ArchetypeHandle.IsValid());
 
 	const FMassEntityHandle Entity = ReserveEntity();
-	InternalBuildEntity(Entity, Archetype);
+	InternalBuildEntity(Entity, ArchetypeHandle);
 
 	const FEntityData& EntityData = Entities[Entity.Index];
 	EntityData.CurrentArchetype->SetFragmentsData(Entity, FragmentInstanceList);
@@ -347,12 +343,12 @@ void UMassEntitySubsystem::ReleaseReservedEntity(FMassEntityHandle Entity)
 	InternalReleaseEntity(Entity);
 }
 
-void UMassEntitySubsystem::BuildEntity(FMassEntityHandle Entity, FMassArchetypeHandle Archetype)
+void UMassEntitySubsystem::BuildEntity(FMassEntityHandle Entity, const FMassArchetypeHandle& ArchetypeHandle)
 {
 	checkf(!IsEntityBuilt(Entity), TEXT("Expecting an entity that is not already built"));
-	check(Archetype.DataPtr.IsValid());
+	check(ArchetypeHandle.IsValid());
 
-	InternalBuildEntity(Entity, Archetype);
+	InternalBuildEntity(Entity, ArchetypeHandle);
 }
 
 void UMassEntitySubsystem::BuildEntity(FMassEntityHandle Entity, TConstArrayView<FInstancedStruct> FragmentInstanceList, FMassArchetypeSharedFragmentValues SharedFragmentValues)
@@ -372,10 +368,10 @@ void UMassEntitySubsystem::BuildEntity(FMassEntityHandle Entity, TConstArrayView
 		Composition.SharedFragments.Add(*SharedFragment.GetScriptStruct());
 	}
 
-	const FMassArchetypeHandle Archetype = CreateArchetype(Composition, SharedFragmentValues);
-	check(Archetype.DataPtr.IsValid());
+	const FMassArchetypeHandle& ArchetypeHandle = CreateArchetype(Composition, SharedFragmentValues);
+	check(ArchetypeHandle.IsValid());
 
-	InternalBuildEntity(Entity, Archetype);
+	InternalBuildEntity(Entity, ArchetypeHandle);
 
 	const FEntityData& EntityData = Entities[Entity.Index];
 	EntityData.CurrentArchetype->SetFragmentsData(Entity, FragmentInstanceList);
@@ -406,7 +402,7 @@ void UMassEntitySubsystem::BatchBuildEntities(const FMassArchetypeEntityCollecti
 
 	// "built" entities case, this is verified during FMassArchetypeEntityCollectionWithPayload construction
 	FMassArchetypeHandle TargetArchetypeHandle = CreateArchetype(Composition, SharedFragmentValues);
-	checkSlow(TargetArchetypeHandle.IsValid());
+	check(TargetArchetypeHandle.IsValid());
 
 	// there are some extra steps in creating EncodedEntities from the original given entity handles and then back
 	// to handles here, but this way we're consistent in how stuff is handled, and there are some slight benefits 
@@ -445,12 +441,11 @@ void UMassEntitySubsystem::BatchBuildEntities(const FMassArchetypeEntityCollecti
 	}
 }
 
-TSharedRef<UMassEntitySubsystem::FEntityCreationContext> UMassEntitySubsystem::BatchCreateEntities(const FMassArchetypeHandle Archetype, const int32 Count, TArray<FMassEntityHandle>& OutEntities)
+TSharedRef<UMassEntitySubsystem::FEntityCreationContext> UMassEntitySubsystem::BatchCreateEntities(const FMassArchetypeHandle& ArchetypeHandle, const int32 Count, TArray<FMassEntityHandle>& OutEntities)
 {
 	TRACE_CPUPROFILER_EVENT_SCOPE(Mass_BatchCreateEntities);
 
-	FMassArchetypeData* ArchetypePtr = Archetype.DataPtr.Get();
-	check(ArchetypePtr);
+	FMassArchetypeData& ArchetypeData = FMassArchetypeHelper::ArchetypeDataFromHandleChecked(ArchetypeHandle);
 	check(Count > 0);
 	
 	int32 Index = OutEntities.Num();
@@ -464,17 +459,17 @@ TSharedRef<UMassEntitySubsystem::FEntityCreationContext> UMassEntitySubsystem::B
 		Result.SerialNumber = SerialNumberGenerator++;
 
 		FEntityData& EntityData = Entities[Result.Index];
-		EntityData.CurrentArchetype = Archetype.DataPtr;
+		EntityData.CurrentArchetype = ArchetypeHandle.DataPtr;
 		EntityData.SerialNumber = Result.SerialNumber;
 		
-		ArchetypePtr->AddEntity(Result);
+		ArchetypeData.AddEntity(Result);
 	}
 		
 	FEntityCreationContext* CreationContext = new FEntityCreationContext(Count);
 	// @todo this could probably be optimized since one would assume we're adding elements to OutEntities in order.
 	// Then again, if that's the case, the sorting will be almost instant
-	new (&CreationContext->EntityCollection)FMassArchetypeEntityCollection(Archetype, MakeArrayView(&OutEntities[OutEntities.Num() - Count], Count), FMassArchetypeEntityCollection::NoDuplicates);
-	if (ObserverManager.HasObserversForBitSet(ArchetypePtr->GetCompositionDescriptor().Fragments, EMassObservedOperation::Add))
+	new (&CreationContext->EntityCollection)FMassArchetypeEntityCollection(ArchetypeHandle, MakeArrayView(&OutEntities[OutEntities.Num() - Count], Count), FMassArchetypeEntityCollection::NoDuplicates);
+	if (ObserverManager.HasObserversForBitSet(ArchetypeData.GetCompositionDescriptor().Fragments, EMassObservedOperation::Add))
 	{
 		CreationContext->OnSpawningFinished = [this](FEntityCreationContext& Context){
 			ObserverManager.OnPostEntitiesCreated(Context.EntityCollection);
@@ -553,8 +548,8 @@ void UMassEntitySubsystem::BatchDestroyEntityChunks(const FMassArchetypeEntityCo
 		ProcessingContext.CommandBuffer = MakeShareable(new FMassCommandBuffer());
 		ObserverManager.OnPreEntitiesDestroyed(ProcessingContext, EntityCollection);
 
-		check(EntityCollection.GetArchetype().IsValid());		
-		EntityCollection.GetArchetype().DataPtr->BatchDestroyEntityChunks(EntityCollection.GetRanges(), EntitiesRemoved);
+		FMassArchetypeData& ArchetypeData = FMassArchetypeHelper::ArchetypeDataFromHandleChecked(EntityCollection.GetArchetype());
+		ArchetypeData.BatchDestroyEntityChunks(EntityCollection.GetRanges(), EntitiesRemoved);
 
 		EntityFreeIndexList.Reserve(EntityFreeIndexList.Num() + EntitiesRemoved.Num());
 	}
@@ -617,9 +612,8 @@ void UMassEntitySubsystem::AddCompositionToEntity_GetDelta(FMassEntityHandle Ent
 		if (ensure(NewArchetypeHandle.DataPtr != EntityData.CurrentArchetype))
 		{
 			// Move the entity over
-			FMassArchetypeData* NewArchetype = NewArchetypeHandle.DataPtr.Get();
-			check(NewArchetype);
-			EntityData.CurrentArchetype->MoveEntityToAnotherArchetype(Entity, *NewArchetype);
+			FMassArchetypeData& NewArchetype = FMassArchetypeHelper::ArchetypeDataFromHandleChecked(NewArchetypeHandle);
+			EntityData.CurrentArchetype->MoveEntityToAnotherArchetype(Entity, NewArchetype);
 			EntityData.CurrentArchetype = NewArchetypeHandle.DataPtr;
 
 			ObserverManager.OnPostCompositionAdded(Entity, InDescriptor);
@@ -654,9 +648,8 @@ void UMassEntitySubsystem::RemoveCompositionFromEntity(FMassEntityHandle Entity,
 			if (ensure(NewArchetypeHandle.DataPtr != EntityData.CurrentArchetype))
 			{
 				// Move the entity over
-				FMassArchetypeData* NewArchetype = NewArchetypeHandle.DataPtr.Get();
-				check(NewArchetype);
-				EntityData.CurrentArchetype->MoveEntityToAnotherArchetype(Entity, *NewArchetype);
+				FMassArchetypeData& NewArchetype = FMassArchetypeHelper::ArchetypeDataFromHandleChecked(NewArchetypeHandle);
+				EntityData.CurrentArchetype->MoveEntityToAnotherArchetype(Entity, NewArchetype);
 				EntityData.CurrentArchetype = NewArchetypeHandle.DataPtr;
 			}
 		}
@@ -665,13 +658,14 @@ void UMassEntitySubsystem::RemoveCompositionFromEntity(FMassEntityHandle Entity,
 
 const FMassArchetypeCompositionDescriptor& UMassEntitySubsystem::GetArchetypeComposition(const FMassArchetypeHandle& ArchetypeHandle) const
 {
-	return ArchetypeHandle.DataPtr->GetCompositionDescriptor();
+	const FMassArchetypeData& ArchetypeData = FMassArchetypeHelper::ArchetypeDataFromHandleChecked(ArchetypeHandle);
+	return ArchetypeData.GetCompositionDescriptor();
 }
 
-void UMassEntitySubsystem::InternalBuildEntity(FMassEntityHandle Entity, const FMassArchetypeHandle Archetype)
+void UMassEntitySubsystem::InternalBuildEntity(FMassEntityHandle Entity, const FMassArchetypeHandle& ArchetypeHandle)
 {
 	FEntityData& EntityData = Entities[Entity.Index];
-	EntityData.CurrentArchetype = Archetype.DataPtr;
+	EntityData.CurrentArchetype = ArchetypeHandle.DataPtr;
 	EntityData.CurrentArchetype->AddEntity(Entity);
 }
 
@@ -712,9 +706,8 @@ void UMassEntitySubsystem::InternalAddFragmentListToEntity(FMassEntityHandle Ent
 	if (NewArchetypeHandle.DataPtr != EntityData.CurrentArchetype)
 	{
 		// Move the entity over
-		FMassArchetypeData* NewArchetype = NewArchetypeHandle.DataPtr.Get();
-		check(NewArchetype);
-		EntityData.CurrentArchetype->MoveEntityToAnotherArchetype(Entity, *NewArchetype);
+		FMassArchetypeData& NewArchetype = FMassArchetypeHelper::ArchetypeDataFromHandleChecked(NewArchetypeHandle);
+		EntityData.CurrentArchetype->MoveEntityToAnotherArchetype(Entity, NewArchetype);
 		EntityData.CurrentArchetype = NewArchetypeHandle.DataPtr;
 	}
 }
@@ -756,9 +749,8 @@ void UMassEntitySubsystem::RemoveFragmentListFromEntity(FMassEntityHandle Entity
 		const FMassArchetypeHandle NewArchetypeHandle = CreateArchetype(NewComposition, OldArchetype->GetSharedFragmentValues());
 
 		// Move the entity over
-		FMassArchetypeData* NewArchetype = NewArchetypeHandle.DataPtr.Get();
-		check(NewArchetype);
-		OldArchetype->MoveEntityToAnotherArchetype(Entity, *NewArchetype);
+		FMassArchetypeData& NewArchetype = FMassArchetypeHelper::ArchetypeDataFromHandleChecked(NewArchetypeHandle);
+		OldArchetype->MoveEntityToAnotherArchetype(Entity, NewArchetype);
 		EntityData.CurrentArchetype = NewArchetypeHandle.DataPtr;
 	}
 }
@@ -1023,12 +1015,11 @@ void UMassEntitySubsystem::MoveEntityToAnotherArchetype(FMassEntityHandle Entity
 {
 	CheckIfEntityIsActive(Entity);
 
-	FMassArchetypeData* NewArchetype = NewArchetypeHandle.DataPtr.Get();
-	check(NewArchetype);
+	FMassArchetypeData& NewArchetype = FMassArchetypeHelper::ArchetypeDataFromHandleChecked(NewArchetypeHandle);
 
 	// Move the entity over
 	FEntityData& EntityData = Entities[Entity.Index];
-	EntityData.CurrentArchetype->MoveEntityToAnotherArchetype(Entity, *NewArchetype);
+	EntityData.CurrentArchetype->MoveEntityToAnotherArchetype(Entity, NewArchetype);
 	EntityData.CurrentArchetype = NewArchetypeHandle.DataPtr;
 }
 
@@ -1098,7 +1089,7 @@ void UMassEntitySubsystem::GetValidArchetypes(const FMassEntityQuery& Query, TAr
 
 	// First get set of all archetypes that contain *any* fragment
 	TSet<TSharedPtr<FMassArchetypeData>> AnyArchetypes;
-	for (const FMassFragmentRequirement& Requirement : Query.GetRequirements())
+	for (const FMassFragmentRequirementDescription& Requirement : Query.GetFragmentRequirements())
 	{
 		check(Requirement.StructType);
 		if (Requirement.Presence != EMassFragmentPresence::None)
@@ -1314,52 +1305,6 @@ void UMassEntitySubsystem::DebugPrintArchetypes(FOutputDevice& Ar, const bool bI
 		NumArchetypes, NumBuckets, LongestArchetypeBucket);
 }
 
-void UMassEntitySubsystem::DebugPrintEntity(int32 Index, FOutputDevice& Ar, const TCHAR* InPrefix) const
-{
-	if (Index >= Entities.Num())
-	{
-		Ar.Logf(ELogVerbosity::Log, TEXT("Unable to list fragments values for out of range index in %s"), *GetPathNameSafe(this));
-		return;
-	}
-
-	const FEntityData& EntityData = Entities[Index];
-	if (!EntityData.IsValid())
-	{
-		Ar.Logf(ELogVerbosity::Log, TEXT("Unable to list fragments values for invalid entity in %s"), *GetPathNameSafe(this));
-	}
-
-	FMassEntityHandle Entity;
-	Entity.Index = Index;
-	Entity.SerialNumber = EntityData.SerialNumber;
-	DebugPrintEntity(Entity, Ar, InPrefix);
-}
-
-void UMassEntitySubsystem::DebugPrintEntity(FMassEntityHandle Entity, FOutputDevice& Ar, const TCHAR* InPrefix) const
-{
-	if (!IsEntityActive(Entity))
-	{
-		Ar.Logf(ELogVerbosity::Log, TEXT("Unable to list fragments values for invalid entity in %s"), *GetPathNameSafe(this));
-	}
-
-	Ar.Logf(ELogVerbosity::Log, TEXT("Listing fragments values for Entity[%s] in %s"), *Entity.DebugGetDescription(), *GetPathNameSafe(this));
-
-	const FEntityData& EntityData = Entities[Entity.Index];
-	FMassArchetypeData* Archetype = EntityData.CurrentArchetype.Get();
-	if (Archetype == nullptr)
-	{
-		Ar.Logf(ELogVerbosity::Log, TEXT("Unable to list fragments values for invalid entity in %s"), *GetPathNameSafe(this));
-	}
-	else
-	{
-		Archetype->DebugPrintEntity(Entity, Ar, InPrefix);
-	}
-}
-
-void UMassEntitySubsystem::DebugGetStringDesc(const FMassArchetypeHandle& Archetype, FOutputDevice& Ar)
-{
-	Ar.Logf(TEXT("%s"), Archetype.IsValid() ? *Archetype.DataPtr->DebugGetDescription() : TEXT("INVALID"));
-}
-
 void UMassEntitySubsystem::DebugGetArchetypesStringDetails(FOutputDevice& Ar, const bool bIncludeEmpty)
 {
 #if WITH_MASSENTITY_DEBUG
@@ -1381,20 +1326,21 @@ void UMassEntitySubsystem::DebugGetArchetypesStringDetails(FOutputDevice& Ar, co
 
 void UMassEntitySubsystem::DebugGetArchetypeFragmentTypes(const FMassArchetypeHandle& Archetype, TArray<const UScriptStruct*>& InOutFragmentList) const
 {
-	if (Archetype.DataPtr.IsValid())
+	if (Archetype.IsValid())
 	{
-		Archetype.DataPtr->GetCompositionDescriptor().Fragments.ExportTypes(InOutFragmentList);
+		const FMassArchetypeData& ArchetypeData = FMassArchetypeHelper::ArchetypeDataFromHandleChecked(Archetype);
+		ArchetypeData.GetCompositionDescriptor().Fragments.ExportTypes(InOutFragmentList);
 	}
 }
 
 int32 UMassEntitySubsystem::DebugGetArchetypeEntitiesCount(const FMassArchetypeHandle& Archetype) const
 {
-	return Archetype.DataPtr.IsValid() ? Archetype.DataPtr->GetNumEntities() : 0;
+	return Archetype.IsValid() ? FMassArchetypeHelper::ArchetypeDataFromHandleChecked(Archetype).GetNumEntities() : 0;
 }
 
 int32 UMassEntitySubsystem::DebugGetArchetypeEntitiesCountPerChunk(const FMassArchetypeHandle& Archetype) const
 {
-	return Archetype.DataPtr.IsValid() ? Archetype.DataPtr->GetNumEntitiesPerChunk() : 0;
+	return Archetype.IsValid() ? FMassArchetypeHelper::ArchetypeDataFromHandleChecked(Archetype).GetNumEntitiesPerChunk() : 0;
 }
 
 void UMassEntitySubsystem::DebugRemoveAllEntities()
@@ -1425,7 +1371,7 @@ void UMassEntitySubsystem::DebugGetArchetypeStrings(const FMassArchetypeHandle&
 		return;
 	}
 
-	const FMassArchetypeData& ArchetypeRef = *Archetype.DataPtr.Get();
+	const FMassArchetypeData& ArchetypeRef = FMassArchetypeHelper::ArchetypeDataFromHandleChecked(Archetype);
 	
 	OutFragmentNames.Reserve(ArchetypeRef.GetFragmentConfigs().Num());
 	for (const FMassArchetypeFragmentConfig& FragmentConfig : ArchetypeRef.GetFragmentConfigs())
diff --git a/Engine/Plugins/Runtime/MassEntity/Source/MassEntity/Private/MassEntityView.cpp b/Engine/Plugins/Runtime/MassEntity/Source/MassEntity/Private/MassEntityView.cpp
--- a/Engine/Plugins/Runtime/MassEntity/Source/MassEntity/Private/MassEntityView.cpp
+++ b/Engine/Plugins/Runtime/MassEntity/Source/MassEntity/Private/MassEntityView.cpp
@@ -10,8 +10,7 @@
 FMassEntityView::FMassEntityView(const FMassArchetypeHandle& ArchetypeHandle, FMassEntityHandle InEntity)
 {
 	Entity = InEntity;
-	check(ArchetypeHandle.IsValid());
-	Archetype = ArchetypeHandle.DataPtr.Get();
+	Archetype = &FMassArchetypeHelper::ArchetypeDataFromHandleChecked(ArchetypeHandle);
 	EntityHandle = Archetype->MakeEntityHandle(Entity);
 }
 
@@ -19,8 +18,7 @@ FMassEntityView::FMassEntityView(const UMassEntitySubsystem& EntitySubsystem, FM
 {
 	Entity = InEntity;
 	const FMassArchetypeHandle ArchetypeHandle = EntitySubsystem.GetArchetypeForEntity(Entity);
-	check(ArchetypeHandle.IsValid());
-	Archetype = ArchetypeHandle.DataPtr.Get();
+	Archetype = &FMassArchetypeHelper::ArchetypeDataFromHandleChecked(ArchetypeHandle);
 	EntityHandle = Archetype->MakeEntityHandle(Entity);
 }
 
diff --git a/Engine/Plugins/Runtime/MassEntity/Source/MassEntity/Private/MassExecutionContext.cpp b/Engine/Plugins/Runtime/MassEntity/Source/MassEntity/Private/MassExecutionContext.cpp
--- a/Engine/Plugins/Runtime/MassEntity/Source/MassEntity/Private/MassExecutionContext.cpp
+++ b/Engine/Plugins/Runtime/MassEntity/Source/MassEntity/Private/MassExecutionContext.cpp
@@ -39,13 +39,13 @@ void FMassExecutionContext::SetEntityCollection(FMassArchetypeEntityCollection&&
 	EntityCollection = MoveTemp(InEntityCollection);
 }
 
-void FMassExecutionContext::SetRequirements(TConstArrayView<FMassFragmentRequirement> InRequirements, 
-	TConstArrayView<FMassFragmentRequirement> InChunkRequirements, 
-	TConstArrayView<FMassFragmentRequirement> InConstSharedRequirements, 
-	TConstArrayView<FMassFragmentRequirement> InSharedRequirements)
+void FMassExecutionContext::SetRequirements(TConstArrayView<FMassFragmentRequirementDescription> InRequirements, 
+	TConstArrayView<FMassFragmentRequirementDescription> InChunkRequirements, 
+	TConstArrayView<FMassFragmentRequirementDescription> InConstSharedRequirements, 
+	TConstArrayView<FMassFragmentRequirementDescription> InSharedRequirements)
 { 
 	FragmentViews.Reset();
-	for (const FMassFragmentRequirement& Requirement : InRequirements)
+	for (const FMassFragmentRequirementDescription& Requirement : InRequirements)
 	{
 		if (Requirement.RequiresBinding())
 		{
@@ -54,7 +54,7 @@ void FMassExecutionContext::SetRequirements(TConstArrayView<FMassFragmentRequire
 	}
 
 	ChunkFragmentViews.Reset();
-	for (const FMassFragmentRequirement& Requirement : InChunkRequirements)
+	for (const FMassFragmentRequirementDescription& Requirement : InChunkRequirements)
 	{
 		if (Requirement.RequiresBinding())
 		{
@@ -63,7 +63,7 @@ void FMassExecutionContext::SetRequirements(TConstArrayView<FMassFragmentRequire
 	}
 
 	ConstSharedFragmentViews.Reset();
-	for (const FMassFragmentRequirement& Requirement : InConstSharedRequirements)
+	for (const FMassFragmentRequirementDescription& Requirement : InConstSharedRequirements)
 	{
 		if (Requirement.RequiresBinding())
 		{
@@ -72,7 +72,7 @@ void FMassExecutionContext::SetRequirements(TConstArrayView<FMassFragmentRequire
 	}
 
 	SharedFragmentViews.Reset();
-	for (const FMassFragmentRequirement& Requirement : InSharedRequirements)
+	for (const FMassFragmentRequirementDescription& Requirement : InSharedRequirements)
 	{
 		if (Requirement.RequiresBinding())
 		{
diff --git a/Engine/Plugins/Runtime/MassEntity/Source/MassEntity/Private/MassObserverProcessor.cpp b/Engine/Plugins/Runtime/MassEntity/Source/MassEntity/Private/MassObserverProcessor.cpp
--- a/Engine/Plugins/Runtime/MassEntity/Source/MassEntity/Private/MassObserverProcessor.cpp
+++ b/Engine/Plugins/Runtime/MassEntity/Source/MassEntity/Private/MassObserverProcessor.cpp
@@ -2,7 +2,6 @@
 
 #include "MassObserverProcessor.h"
 #include "MassObserverRegistry.h"
-#include "MassEntityDebug.h"
 
 //----------------------------------------------------------------------//
 // UMassObserverProcessor
diff --git a/Engine/Plugins/Runtime/MassEntity/Source/MassEntity/Private/MassProcessingPhase.cpp b/Engine/Plugins/Runtime/MassEntity/Source/MassEntity/Private/MassProcessingPhase.cpp
--- a/Engine/Plugins/Runtime/MassEntity/Source/MassEntity/Private/MassProcessingPhase.cpp
+++ b/Engine/Plugins/Runtime/MassEntity/Source/MassEntity/Private/MassProcessingPhase.cpp
@@ -2,7 +2,7 @@
 
 #include "MassProcessingPhase.h"
 #include "MassProcessingTypes.h"
-#include "MassEntityDebug.h"
+#include "MassDebugger.h"
 #include "MassEntitySettings.h"
 #include "MassProcessor.h"
 #include "MassExecutor.h"
@@ -278,7 +278,7 @@ void UMassProcessingPhaseManager::SetPhaseProcessor(const EMassProcessingPhase P
 	{
 		REDIRECT_OBJECT_TO_VLOG(PhaseProcessor, this);
 		PhaseProcessor->SetProcessingPhase(Phase);
-		PhaseProcessor->SetGroupName(FName(FString::Printf(TEXT("%s-%s"), *PhaseProcessor->GetName(), *EnumToString(Phase))));
+		PhaseProcessor->SetGroupName(FName(FString::Printf(TEXT("%s-%s"), *PhaseProcessor->GetName(), *UEnum::GetValueAsString(Phase))));
 		
 		check(GetOuter());
 		PhaseProcessor->Initialize(*GetOuter());
@@ -292,7 +292,7 @@ void UMassProcessingPhaseManager::SetPhaseProcessor(const EMassProcessingPhase P
 	}
 	else
 	{
-		UE_VLOG_UELOG(this, LogMass, Log, TEXT("Nulling-out phase processor for phase %s"), *EnumToString(Phase));
+		UE_VLOG_UELOG(this, LogMass, Log, TEXT("Nulling-out phase processor for phase %s"), *UEnum::GetValueAsString(Phase));
 	}
 }
 
@@ -304,7 +304,7 @@ void UMassProcessingPhaseManager::CreatePhases()
 		ProcessingPhases[i].Phase = EMassProcessingPhase(i);
 		ProcessingPhases[i].TickGroup = UE::Mass::Private::PhaseToTickingGroup[i];
 		UMassCompositeProcessor* PhaseProcessor = NewObject<UMassCompositeProcessor>(this, UMassCompositeProcessor::StaticClass()
-			, *FString::Printf(TEXT("ProcessingPhase_%s"), *EnumToString(EMassProcessingPhase(i))));
+			, *FString::Printf(TEXT("ProcessingPhase_%s"), *UEnum::GetValueAsString(EMassProcessingPhase(i))));
 		SetPhaseProcessor(EMassProcessingPhase(i), PhaseProcessor);
 	}
 }
diff --git a/Engine/Plugins/Runtime/MassEntity/Source/MassEntity/Private/MassProcessingTypes.cpp b/Engine/Plugins/Runtime/MassEntity/Source/MassEntity/Private/MassProcessingTypes.cpp
--- a/Engine/Plugins/Runtime/MassEntity/Source/MassEntity/Private/MassProcessingTypes.cpp
+++ b/Engine/Plugins/Runtime/MassEntity/Source/MassEntity/Private/MassProcessingTypes.cpp
@@ -6,7 +6,7 @@
 #include "Misc/OutputDevice.h"
 #include "MassEntityUtils.h"
 #include "VisualLogger/VisualLogger.h"
-#include "MassEntityDebug.h"
+#include "MassDebugger.h"
 
 DEFINE_LOG_CATEGORY(LogMass);
 
@@ -230,11 +230,6 @@ UMassCompositeProcessor* FMassRuntimePipeline::FindTopLevelGroupByName(FName Gro
 	return nullptr;
 }
 
-void FMassRuntimePipeline::DebugOutputDescription(FOutputDevice& Ar) const
-{
-	UE::Mass::Debug::DebugOutputDescription(Processors, Ar);
-}
-
 uint32 GetTypeHash(const FMassRuntimePipeline& Instance)
 { 
 	uint32 Hash = 0;
diff --git a/Engine/Plugins/Runtime/MassEntity/Source/MassEntity/Private/MassRequirementAccessDetector.cpp b/Engine/Plugins/Runtime/MassEntity/Source/MassEntity/Private/MassRequirementAccessDetector.cpp
--- a/Engine/Plugins/Runtime/MassEntity/Source/MassEntity/Private/MassRequirementAccessDetector.cpp
+++ b/Engine/Plugins/Runtime/MassEntity/Source/MassEntity/Private/MassRequirementAccessDetector.cpp
@@ -39,10 +39,10 @@ void FMassRequirementAccessDetector::RequireAccess(const FMassEntityQuery& Query
 		Operation(Query.RequiredConstSubsystems, &FRWAccessDetector::AcquireReadAccess);
 		Operation(Query.RequiredMutableSubsystems, &FRWAccessDetector::AcquireWriteAccess);
 		
-		Aquire(Query.Requirements);
-		Aquire(Query.ChunkRequirements);
-		Aquire(Query.ConstSharedRequirements);
-		Aquire(Query.SharedRequirements);
+		Aquire(Query.FragmentRequirements);
+		Aquire(Query.ChunkFragmentRequirements);
+		Aquire(Query.ConstSharedFragmentRequirements);
+		Aquire(Query.SharedFragmentRequirements);
 	}
 }
 
@@ -53,10 +53,10 @@ void FMassRequirementAccessDetector::ReleaseAccess(const FMassEntityQuery& Query
 		Operation(Query.RequiredConstSubsystems, &FRWAccessDetector::ReleaseReadAccess);
 		Operation(Query.RequiredMutableSubsystems, &FRWAccessDetector::ReleaseWriteAccess);
 
-		Release(Query.Requirements);
-		Release(Query.ChunkRequirements);
-		Release(Query.ConstSharedRequirements);
-		Release(Query.SharedRequirements);
+		Release(Query.FragmentRequirements);
+		Release(Query.ChunkFragmentRequirements);
+		Release(Query.ConstSharedFragmentRequirements);
+		Release(Query.SharedFragmentRequirements);
 	}
 }
 
diff --git a/Engine/Plugins/Runtime/MassEntity/Source/MassEntity/Private/MassRequirements.cpp b/Engine/Plugins/Runtime/MassEntity/Source/MassEntity/Private/MassRequirements.cpp
new file mode 100644
--- /dev/null
+++ b/Engine/Plugins/Runtime/MassEntity/Source/MassEntity/Private/MassRequirements.cpp
@@ -0,0 +1,140 @@
+// Copyright Epic Games, Inc. All Rights Reserved.
+
+#include "MassRequirements.h"
+#include "MassArchetypeData.h"
+#include "MassProcessorDependencySolver.h"
+#if WITH_MASSENTITY_DEBUG
+#include "MassRequirementAccessDetector.h"
+#endif // WITH_MASSENTITY_DEBUG
+
+
+namespace UE::Mass::Private
+{
+	template<typename TContainer>
+	void ExportRequirements(TConstArrayView<FMassFragmentRequirementDescription> Requirements, TMassExecutionAccess<TContainer>& Out)
+	{
+		for (const FMassFragmentRequirementDescription& Requirement : Requirements)
+		{
+			if (Requirement.Presence != EMassFragmentPresence::None)
+			{
+				check(Requirement.StructType);
+				if (Requirement.AccessMode == EMassFragmentAccess::ReadOnly)
+				{
+					Out.Read.Add(*Requirement.StructType);
+				}
+				else if (Requirement.AccessMode == EMassFragmentAccess::ReadWrite)
+				{
+					Out.Write.Add(*Requirement.StructType);
+				}
+			}
+		}
+	}
+}
+
+//////////////////////////////////////////////////////////////////////
+// FMassSubsystemRequirements
+
+void FMassSubsystemRequirements::ExportRequirements(FMassExecutionRequirements& OutRequirements) const
+{
+	OutRequirements.RequiredSubsystems.Read += RequiredConstSubsystems;
+	OutRequirements.RequiredSubsystems.Write += RequiredMutableSubsystems;
+}
+
+void FMassSubsystemRequirements::Reset()
+{
+	RequiredConstSubsystems.Reset();
+	RequiredMutableSubsystems.Reset();
+	bRequiresGameThreadExecution = false;
+}
+
+//////////////////////////////////////////////////////////////////////
+// FMassFragmentRequirements
+
+FMassFragmentRequirements::FMassFragmentRequirements(std::initializer_list<UScriptStruct*> InitList)
+{
+	for (const UScriptStruct* FragmentType : InitList)
+	{
+		AddRequirement(FragmentType, EMassFragmentAccess::ReadWrite, EMassFragmentPresence::All);
+	}
+}
+
+FMassFragmentRequirements::FMassFragmentRequirements(TConstArrayView<const UScriptStruct*> InitList)
+{
+	for (const UScriptStruct* FragmentType : InitList)
+	{
+		AddRequirement(FragmentType, EMassFragmentAccess::ReadWrite, EMassFragmentPresence::All);
+	}
+}
+
+void FMassFragmentRequirements::SortRequirements()
+{
+	// we're sorting the Requirements the same way ArchetypeData's FragmentConfig is sorted (see FMassArchetypeData::Initialize)
+	// so that when we access ArchetypeData.FragmentConfigs in FMassArchetypeData::BindRequirementsWithMapping
+	// (via GetFragmentData call) the access is sequential (i.e. not random) and there's a higher chance the memory
+	// FragmentConfigs we want to access have already been fetched and are available in processor cache.
+	FragmentRequirements.Sort(FScriptStructSortOperator());
+	ChunkFragmentRequirements.Sort(FScriptStructSortOperator());
+	ConstSharedFragmentRequirements.Sort(FScriptStructSortOperator());
+	SharedFragmentRequirements.Sort(FScriptStructSortOperator());
+}
+
+bool FMassFragmentRequirements::CheckValidity() const
+{
+	return RequiredAllFragments.IsEmpty() == false || RequiredAnyFragments.IsEmpty() == false || RequiredOptionalFragments.IsEmpty() == false;
+}
+
+bool FMassFragmentRequirements::DoesArchetypeMatchRequirements(const FMassArchetypeHandle& ArchetypeHandle) const
+{
+	check(ArchetypeHandle.IsValid());
+	const FMassArchetypeData* Archetype = FMassArchetypeHelper::ArchetypeDataFromHandle(ArchetypeHandle);
+	CA_ASSUME(Archetype);
+	
+	const FMassArchetypeCompositionDescriptor& ArchetypeComposition = Archetype->GetCompositionDescriptor();
+
+	return ArchetypeComposition.Fragments.HasAll(RequiredAllFragments)
+		&& (RequiredAnyFragments.IsEmpty() || ArchetypeComposition.Fragments.HasAny(RequiredAnyFragments))
+		&& ArchetypeComposition.Fragments.HasNone(RequiredNoneFragments)
+		&& ArchetypeComposition.Tags.HasAll(RequiredAllTags)
+		&& (RequiredAnyTags.IsEmpty() || ArchetypeComposition.Tags.HasAny(RequiredAnyTags))
+		&& ArchetypeComposition.Tags.HasNone(RequiredNoneTags)
+		&& ArchetypeComposition.ChunkFragments.HasAll(RequiredAllChunkFragments)
+		&& ArchetypeComposition.ChunkFragments.HasNone(RequiredNoneChunkFragments)
+		&& ArchetypeComposition.SharedFragments.HasAll(RequiredAllSharedFragments)
+		&& ArchetypeComposition.SharedFragments.HasNone(RequiredNoneSharedFragments);
+}
+
+void FMassFragmentRequirements::ExportRequirements(FMassExecutionRequirements& OutRequirements) const
+{
+	using UE::Mass::Private::ExportRequirements;
+	ExportRequirements<FMassFragmentBitSet>(FragmentRequirements, OutRequirements.Fragments);
+	ExportRequirements<FMassChunkFragmentBitSet>(ChunkFragmentRequirements, OutRequirements.ChunkFragments);
+	ExportRequirements<FMassSharedFragmentBitSet>(ConstSharedFragmentRequirements, OutRequirements.SharedFragments);
+	ExportRequirements<FMassSharedFragmentBitSet>(SharedFragmentRequirements, OutRequirements.SharedFragments);
+
+	OutRequirements.RequiredAllTags = RequiredAllTags;
+	OutRequirements.RequiredAnyTags = RequiredAnyTags;
+	OutRequirements.RequiredNoneTags = RequiredNoneTags;
+}
+
+void FMassFragmentRequirements::Reset()
+{
+	FragmentRequirements.Reset();
+	ChunkFragmentRequirements.Reset();
+	ConstSharedFragmentRequirements.Reset();
+	SharedFragmentRequirements.Reset();
+	RequiredAllTags.Reset();
+	RequiredAnyTags.Reset();
+	RequiredNoneTags.Reset();
+	RequiredAllFragments.Reset();
+	RequiredAnyFragments.Reset();
+	RequiredOptionalFragments.Reset();
+	RequiredNoneFragments.Reset();
+	RequiredAllChunkFragments.Reset();
+	RequiredOptionalChunkFragments.Reset();
+	RequiredNoneChunkFragments.Reset();
+	RequiredAllSharedFragments.Reset();
+	RequiredOptionalSharedFragments.Reset();
+	RequiredNoneSharedFragments.Reset();
+
+	IncrementalChangesCount = 0;
+}
diff --git a/Engine/Plugins/Runtime/MassEntity/Source/MassEntity/Public/MassArchetypeTypes.h b/Engine/Plugins/Runtime/MassEntity/Source/MassEntity/Public/MassArchetypeTypes.h
--- a/Engine/Plugins/Runtime/MassEntity/Source/MassEntity/Public/MassArchetypeTypes.h
+++ b/Engine/Plugins/Runtime/MassEntity/Source/MassEntity/Public/MassArchetypeTypes.h
@@ -17,6 +17,7 @@ struct FMassEntityQuery;
 struct FMassArchetypeEntityCollection;
 struct FMassEntityView;
 struct FMassDebugger;
+struct FMassArchetypeHelper;
 
 typedef TFunction< void(FMassExecutionContext& /*ExecutionContext*/) > FMassExecuteFunction;
 typedef TFunction< bool(const FMassExecutionContext& /*ExecutionContext*/) > FMassChunkConditionFunction;
@@ -42,11 +43,8 @@ private:
 	{}
 	TSharedPtr<FMassArchetypeData> DataPtr;
 
+	friend FMassArchetypeHelper;
 	friend UMassEntitySubsystem;
-	friend FMassArchetypeEntityCollection;
-	friend FMassEntityQuery;
-	friend FMassEntityView;
-	friend FMassDebugger;
 };
 
 
@@ -126,9 +124,6 @@ public:
 
 protected:
 	friend struct FMassArchetypeEntityCollectionWithPayload;
-	// @todo temporary solution to FMassArchetypeHandle.DataPtr being private, will be fixed by switching that to an ArchetypeIndex
-	static FMassArchetypeData* GetArchetypePtr(const FMassArchetypeHandle& ArchetypeHandle);
-	
 	void BuildEntityRanges(TStridedView<const int32> TrueIndices);
 	
 private:
diff --git a/Engine/Plugins/Runtime/MassEntity/Source/MassEntity/Public/MassDebugger.h b/Engine/Plugins/Runtime/MassEntity/Source/MassEntity/Public/MassDebugger.h
--- a/Engine/Plugins/Runtime/MassEntity/Source/MassEntity/Public/MassDebugger.h
+++ b/Engine/Plugins/Runtime/MassEntity/Source/MassEntity/Public/MassDebugger.h
@@ -10,30 +10,37 @@
 #include "MassEntityQuery.h"
 #include "MassProcessor.h"
 
+
+class FOutputDevice;
 class UMassProcessor;
 struct FMassEntityQuery;
 class UMassEntitySubsystem;
 struct FMassArchetypeHandle;
+struct FMassFragmentRequirementDescription;
 enum class EMassFragmentAccess : uint8;
 enum class EMassFragmentPresence : uint8;
 
 MASSENTITY_API DECLARE_ENUM_TO_STRING(EMassFragmentAccess);
 MASSENTITY_API DECLARE_ENUM_TO_STRING(EMassFragmentPresence);
+MASSENTITY_API DECLARE_ENUM_TO_STRING(EMassProcessingPhase);
 
 namespace UE::Mass::Debug
 {
 struct MASSENTITY_API FQueryRequirementsView
 {
-	TConstArrayView<FMassFragmentRequirement> FragmentRequirements;
-	TConstArrayView<FMassFragmentRequirement> ChunkRequirements;
-	TConstArrayView<FMassFragmentRequirement> ConstSharedRequirements;
-	TConstArrayView<FMassFragmentRequirement> SharedRequirements;
+	TConstArrayView<FMassFragmentRequirementDescription> FragmentRequirements;
+	TConstArrayView<FMassFragmentRequirementDescription> ChunkRequirements;
+	TConstArrayView<FMassFragmentRequirementDescription> ConstSharedRequirements;
+	TConstArrayView<FMassFragmentRequirementDescription> SharedRequirements;
 	const FMassTagBitSet& RequiredAllTags;
 	const FMassTagBitSet& RequiredAnyTags;
 	const FMassTagBitSet& RequiredNoneTags;
 	const FMassExternalSubystemBitSet& RequiredConstSubsystems;
 	const FMassExternalSubystemBitSet& RequiredMutableSubsystems;
 };
+
+FString DebugGetFragmentAccessString(EMassFragmentAccess Access);
+MASSENTITY_API extern void DebugOutputDescription(TConstArrayView<UMassProcessor*> Processors, FOutputDevice& Ar);
 } // namespace UE::Mass::Debug
 
 struct MASSENTITY_API FMassDebugger
@@ -45,7 +52,7 @@ struct MASSENTITY_API FMassDebugger
 
 	static UE::Mass::Debug::FQueryRequirementsView GetQueryRequirements(const FMassEntityQuery& Query);
 	static void GetQueryExecutionRequirements(const FMassEntityQuery& Query, FMassExecutionRequirements& OutExecutionRequirements);
-	
+
 	static TArray<FMassArchetypeHandle> GetAllArchetypes(const UMassEntitySubsystem& EntitySubsystem);
 	static const FMassArchetypeCompositionDescriptor& GetArchetypeComposition(const FMassArchetypeHandle& ArchetypeHandle);
 	static const FMassArchetypeSharedFragmentValues& GetArchetypeSharedFragmentValues(const FMassArchetypeHandle& ArchetypeHandle);
@@ -53,6 +60,14 @@ struct MASSENTITY_API FMassDebugger
 	static void GetArchetypeEntityStats(const FMassArchetypeHandle& ArchetypeHandle, int32& OutEntitiesCount, int32& OutNumEntitiesPerChunk, int32& OutChunksCount);
 
 	static TConstArrayView<UMassCompositeProcessor::FDependencyNode> GetProcessingGraph(const UMassCompositeProcessor& GraphOwner);
+	
+	static FString GetSingleRequirementDescription(const FMassFragmentRequirementDescription& Requirement);
+	static FString GetRequirementsDescription(const FMassFragmentRequirements& Requirements);
+	static FString GetArchetypeRequirementCompatibilityDescription(const FMassFragmentRequirements& Requirements, const FMassArchetypeHandle& ArchetypeHandle);
+
+	static void OutputArchetypeDescription(FOutputDevice& Ar, const FMassArchetypeHandle& Archetype);
+	static void OutputEntityDescription(FOutputDevice& Ar, const UMassEntitySubsystem& EntitySubsystem, const int32 EntityIndex, const TCHAR* InPrefix = TEXT(""));
+	static void OutputEntityDescription(FOutputDevice& Ar, const UMassEntitySubsystem& EntitySubsystem, const FMassEntityHandle Entity, const TCHAR* InPrefix = TEXT(""));
 };
 
 #endif // WITH_MASSENTITY_DEBUG
\ No newline at end of file
diff --git a/Engine/Plugins/Runtime/MassEntity/Source/MassEntity/Public/MassEntityDebug.h b/Engine/Plugins/Runtime/MassEntity/Source/MassEntity/Public/MassEntityDebug.h
deleted file mode 100644
diff --git a/Engine/Plugins/Runtime/MassEntity/Source/MassEntity/Public/MassEntityQuery.h b/Engine/Plugins/Runtime/MassEntity/Source/MassEntity/Public/MassEntityQuery.h
--- a/Engine/Plugins/Runtime/MassEntity/Source/MassEntity/Public/MassEntityQuery.h
+++ b/Engine/Plugins/Runtime/MassEntity/Source/MassEntity/Public/MassEntityQuery.h
@@ -8,98 +8,20 @@
 #include "MassEntityTypes.h"
 #include "MassArchetypeTypes.h"
 #include "MassExternalSubsystemTraits.h"
+#include "MassRequirements.h"
 #include "MassEntityQuery.generated.h"
 
 
-class UMassEntitySubsystem;
-class UMassProcessor;
-struct FMassArchetypeData;
-struct FMassExecutionContext;
-struct FMassFragment;
-struct FMassArchetypeHandle;
-struct FMassExecutionRequirements;
-struct FMassRequirementAccessDetector;
-
-UENUM()
-enum class EMassFragmentAccess : uint8
-{
-	// no binding required
-	None, 
-
-	// We want to read the data for the fragment
-	ReadOnly,
-
-	// We want to read and write the data for the fragment
-	ReadWrite,
-
-	MAX
-};
-
-UENUM()
-enum class EMassFragmentPresence : uint8
-{
-	// All of the required fragments must be present
-	All,
-
-	// One of the required fragments must be present
-	Any,
-
-	// None of the required fragments can be present
-	None,
-
-	// If fragment is present we'll use it, but it missing stop processing of a given archetype
-	Optional,
-
-	MAX
-};
-
-
-struct MASSENTITY_API FMassFragmentRequirement
-{
-	const UScriptStruct* StructType = nullptr;
-	EMassFragmentAccess AccessMode = EMassFragmentAccess::None;
-	EMassFragmentPresence Presence = EMassFragmentPresence::Optional;
-
-public:
-	FMassFragmentRequirement(){}
-	FMassFragmentRequirement(const UScriptStruct* InStruct, const EMassFragmentAccess InAccessMode, const EMassFragmentPresence InPresence)
-		: StructType(InStruct)
-		, AccessMode(InAccessMode)
-		, Presence(InPresence)
-	{
-		check(InStruct);
-		checkf((Presence != EMassFragmentPresence::Any && Presence != EMassFragmentPresence::Optional)
-			|| AccessMode == EMassFragmentAccess::ReadOnly || AccessMode == EMassFragmentAccess::ReadWrite, TEXT("Only ReadOnly and ReadWrite modes are suppored for optional requirements"));
-	}
-
-	bool RequiresBinding() const { return (AccessMode != EMassFragmentAccess::None); }
-	bool IsOptional() const { return (Presence == EMassFragmentPresence::Optional || Presence == EMassFragmentPresence::Any); }
-
-	FString DebugGetDescription() const;
-
-	// these functions are used for sorting. See FScriptStructSortOperator
-	int32 GetStructureSize() const
-	{
-		return StructType->GetStructureSize();
-	}
-
-	FName GetFName() const
-	{
-		return StructType->GetFName();
-	}
-};
-
-
 /** 
- *  FMassEntityQuery is a structure that serves two main purposes:
- *  1. Describe properties required of an archetype that's a subject of calculations
- *  2. Trigger calculations on cached set of valid archetypes as described by Requirements
+ *  FMassEntityQuery is a structure that is used to trigger calculations on cached set of valid archetypes as described 
+ *  by requirements. See the parent classes FMassFragmentRequirements and FMassSubsystemRequirements for setting up the 
+ *	required fragments and subsystems.
  * 
  *  A query to be considered valid needs declared at least one EMassFragmentPresence::All, EMassFragmentPresence::Any 
  *  EMassFragmentPresence::Optional fragment requirement.
  */
 USTRUCT()
-struct MASSENTITY_API FMassEntityQuery
+struct MASSENTITY_API FMassEntityQuery : public FMassFragmentRequirements, public FMassSubsystemRequirements
 {
 	GENERATED_BODY()
 
@@ -132,250 +54,10 @@ public:
 	 *	ArchetypeDataVersion properties). */
 	void CacheArchetypes(const UMassEntitySubsystem& InEntitySubsystem);
 
-	FMassEntityQuery& AddRequirement(const UScriptStruct* FragmentType, const EMassFragmentAccess AccessMode, const EMassFragmentPresence Presence = EMassFragmentPresence::All)
-	{
-		checkf(Requirements.FindByPredicate([FragmentType](const FMassFragmentRequirement& Item){ return Item.StructType == FragmentType; }) == nullptr
-			, TEXT("Duplicated requirements are not supported. %s already present"), *GetNameSafe(FragmentType));
-		
-		if (Presence != EMassFragmentPresence::None)
-		{
-			Requirements.Emplace(FragmentType, AccessMode, Presence);
-		}
-
-		switch (Presence)
-		{
-		case EMassFragmentPresence::All:
-			RequiredAllFragments.Add(*FragmentType);
-			break;
-		case EMassFragmentPresence::Any:
-			RequiredAnyFragments.Add(*FragmentType);
-			break;
-		case EMassFragmentPresence::Optional:
-			RequiredOptionalFragments.Add(*FragmentType);
-			break;
-		case EMassFragmentPresence::None:
-			RequiredNoneFragments.Add(*FragmentType);
-			break;
-		}
-		// force recaching the next time this query is used or the following CacheArchetypes call.
-		DirtyCachedData();
-		return *this;
-	}
-
-	/** FMassEntityQuery ref returned for chaining */
-	template<typename T>
-	FMassEntityQuery& AddRequirement(const EMassFragmentAccess AccessMode, const EMassFragmentPresence Presence = EMassFragmentPresence::All)
-	{
-		checkf(Requirements.FindByPredicate([](const FMassFragmentRequirement& Item) { return Item.StructType == T::StaticStruct(); }) == nullptr
-			, TEXT("Duplicated requirements are not supported. %s already present"), *T::StaticStruct()->GetName());
-
-		static_assert(TIsDerivedFrom<T, FMassFragment>::IsDerived, "Given struct doesn't represent a valid fragment type. Make sure to inherit from FMassFragment or one of its child-types.");
-		
-		if (Presence != EMassFragmentPresence::None)
-		{
-			Requirements.Emplace(T::StaticStruct(), AccessMode, Presence);
-		}
-		
-		switch (Presence)
-		{
-		case EMassFragmentPresence::All:
-			RequiredAllFragments.Add<T>();
-			break;
-		case EMassFragmentPresence::Any:
-			RequiredAnyFragments.Add<T>();
-			break;
-		case EMassFragmentPresence::Optional:
-			RequiredOptionalFragments.Add<T>();
-			break;
-		case EMassFragmentPresence::None:
-			RequiredNoneFragments.Add<T>();
-			break;
-		}
-		// force recaching the next time this query is used or the following CacheArchetypes call.
-		DirtyCachedData();
-		return *this;
-	}
-
-	void AddTagRequirement(const UScriptStruct& TagType, const EMassFragmentPresence Presence)
-	{
-		checkfSlow(int(Presence) < int(EMassFragmentPresence::Optional), TEXT("Optional and MAX presence are not valid calues for AddTagRequirement"));
-		switch (Presence)
-		{
-		case EMassFragmentPresence::All:
-			RequiredAllTags.Add(TagType);
-			break;
-		case EMassFragmentPresence::Any:
-			RequiredAnyTags.Add(TagType);
-			break;
-		case EMassFragmentPresence::None:
-			RequiredNoneTags.Add(TagType);
-			break;
-		}
-		DirtyCachedData();
-	}
-
-	template<typename T>
-	FMassEntityQuery& AddTagRequirement(const EMassFragmentPresence Presence)
-	{
-		checkfSlow(int(Presence) < int(EMassFragmentPresence::Optional), TEXT("Optional and MAX presence are not valid calues for AddTagRequirement"));
-		static_assert(TIsDerivedFrom<T, FMassTag>::IsDerived, "Given struct doesn't represent a valid tag type. Make sure to inherit from FMassFragment or one of its child-types.");
-		switch (Presence)
-		{
-			case EMassFragmentPresence::All:
-				RequiredAllTags.Add<T>();
-				break;
-			case EMassFragmentPresence::Any:
-				RequiredAnyTags.Add<T>();
-				break;
-			case EMassFragmentPresence::None:
-				RequiredNoneTags.Add<T>();
-				break;
-		}
-		DirtyCachedData();
-		return *this;
-	}
-
-	// actual implementation in specializations
-	template<EMassFragmentPresence Presence> 
-	FMassEntityQuery& AddTagRequirements(const FMassTagBitSet& TagBitSet)
-	{
-		static_assert(Presence == EMassFragmentPresence::None || Presence == EMassFragmentPresence::All || Presence == EMassFragmentPresence::Any
-			, "The only valid values for AddTagRequirements are All, Any and None");
-		return *this;
-	}
-
-	template<typename T>
-	FMassEntityQuery& AddChunkRequirement(const EMassFragmentAccess AccessMode, const EMassFragmentPresence Presence = EMassFragmentPresence::All)
-	{
-		static_assert(TIsDerivedFrom<T, FMassChunkFragment>::IsDerived, "Given struct doesn't represent a valid chunk fragment type. Make sure to inherit from FMassChunkFragment or one of its child-types.");
-		checkf(ChunkRequirements.FindByPredicate([](const FMassFragmentRequirement& Item) { return Item.StructType == T::StaticStruct(); }) == nullptr
-			, TEXT("Duplicated requirements are not supported. %s already present"), *T::StaticStruct()->GetName());
-		checkfSlow(Presence != EMassFragmentPresence::Any, TEXT("\'Any\' is not a valid Presence value for AddChunkRequirement."));
-
-		switch (Presence)
-		{
-			case EMassFragmentPresence::All:
-				RequiredAllChunkFragments.Add<T>();
-				ChunkRequirements.Emplace(T::StaticStruct(), AccessMode, Presence);
-				break;
-			case EMassFragmentPresence::Optional:
-				RequiredOptionalChunkFragments.Add<T>();
-				ChunkRequirements.Emplace(T::StaticStruct(), AccessMode, Presence);
-				break;
-			case EMassFragmentPresence::None:
-				RequiredNoneChunkFragments.Add<T>();
-				break;
-		}
-		DirtyCachedData();
-		return *this;
-	}
-
-	template<typename T>
-	FMassEntityQuery& AddConstSharedRequirement(const EMassFragmentPresence Presence = EMassFragmentPresence::All)
-	{
-		static_assert(TIsDerivedFrom<T, FMassSharedFragment>::IsDerived, "Given struct doesn't represent a valid shared fragment type. Make sure to inherit from FMassSharedFragment or one of its child-types.");
-		checkf(ConstSharedRequirements.FindByPredicate([](const FMassFragmentRequirement& Item) { return Item.StructType == T::StaticStruct(); }) == nullptr
-			, TEXT("Duplicated requirements are not supported. %s already present"), *T::StaticStruct()->GetName());
-		checkfSlow(Presence != EMassFragmentPresence::Any, TEXT("\'Any\' is not a valid Presence value for AddSharedRequirement."));
-
-		switch (Presence)
-		{
-		case EMassFragmentPresence::All:
-			RequiredAllSharedFragments.Add<T>();
-			ConstSharedRequirements.Emplace(T::StaticStruct(), EMassFragmentAccess::ReadOnly, Presence);
-			break;
-		case EMassFragmentPresence::Optional:
-			RequiredOptionalSharedFragments.Add<T>();
-			ConstSharedRequirements.Emplace(T::StaticStruct(), EMassFragmentAccess::ReadOnly, Presence);
-			break;
-		case EMassFragmentPresence::None:
-			RequiredNoneSharedFragments.Add<T>();
-			break;
-		}
-		DirtyCachedData();
-		return *this;
-	}
-
-	template<typename T>
-	FMassEntityQuery& AddSharedRequirement(const EMassFragmentAccess AccessMode, const EMassFragmentPresence Presence = EMassFragmentPresence::All)
-	{
-		static_assert(TIsDerivedFrom<T, FMassSharedFragment>::IsDerived, "Given struct doesn't represent a valid shared fragment type. Make sure to inherit from FMassSharedFragment or one of its child-types.");
-		checkf(SharedRequirements.FindByPredicate([](const FMassFragmentRequirement& Item) { return Item.StructType == T::StaticStruct(); }) == nullptr
-			, TEXT("Duplicated requirements are not supported. %s already present"), *T::StaticStruct()->GetName());
-		checkfSlow(Presence != EMassFragmentPresence::Any, TEXT("\'Any\' is not a valid Presence value for AddSharedRequirement."));
-
-		switch (Presence)
-		{
-		case EMassFragmentPresence::All:
-			RequiredAllSharedFragments.Add<T>();
-			SharedRequirements.Emplace(T::StaticStruct(), AccessMode, Presence);
-			if (AccessMode == EMassFragmentAccess::ReadWrite)
-			{
-				bRequiresGameThreadExecution |= TMassSharedFragmentTraits<T>::GameThreadOnly;
-			}
-			break;
-		case EMassFragmentPresence::Optional:
-			RequiredOptionalSharedFragments.Add<T>();
-			SharedRequirements.Emplace(T::StaticStruct(), AccessMode, Presence);
-			if (AccessMode == EMassFragmentAccess::ReadWrite)
-			{
-				bRequiresGameThreadExecution |= TMassSharedFragmentTraits<T>::GameThreadOnly;
-			}
-			break;
-		case EMassFragmentPresence::None:
-			RequiredNoneSharedFragments.Add<T>();
-			break;
-		}
-		DirtyCachedData();
-		return *this;
-	}
-
-	template<typename T>
-	FMassEntityQuery& AddSubsystemRequirement(const EMassFragmentAccess AccessMode)
-	{
-		check(AccessMode != EMassFragmentAccess::None && AccessMode != EMassFragmentAccess::MAX);
-
-		switch (AccessMode)
-		{
-		case EMassFragmentAccess::ReadOnly:
-			RequiredConstSubsystems.Add<T>();
-			bRequiresGameThreadExecution |= TMassExternalSubsystemTraits<T>::GameThreadOnly;
-			break;
-		case EMassFragmentAccess::ReadWrite:
-			RequiredMutableSubsystems.Add<T>();
-			bRequiresGameThreadExecution |= TMassExternalSubsystemTraits<T>::GameThreadOnly;
-			break;
-		default:
-			check(false);
-		}
-		
-		return *this;
-	}
-
-	FMassEntityQuery& AddSubsystemRequirement(const TSubclassOf<USubsystem> SubsystemClass, const EMassFragmentAccess AccessMode, const bool bGameThreadOnly = true)
-	{
-		check(AccessMode != EMassFragmentAccess::None && AccessMode != EMassFragmentAccess::MAX);
-
-		switch (AccessMode)
-		{
-		case EMassFragmentAccess::ReadOnly:
-			RequiredConstSubsystems.Add(**SubsystemClass);
-			bRequiresGameThreadExecution |= bGameThreadOnly;
-			break;
-		case EMassFragmentAccess::ReadWrite:
-			RequiredMutableSubsystems.Add(**SubsystemClass);
-			bRequiresGameThreadExecution |= bGameThreadOnly;
-			break;
-		default:
-			check(false);
-		}
-
-		return *this;
-	}
-
 	void Clear()
 	{
-		Requirements.Empty();
+		FMassFragmentRequirements::Reset();
+		FMassSubsystemRequirements::Reset();
 		DirtyCachedData();
 	}
 
@@ -385,34 +67,15 @@ public:
 		ArchetypeDataVersion = 0;
 	}
 	
-	bool DoesArchetypeMatchRequirements(const FMassArchetypeHandle& ArchetypeHandle) const;
+	bool DoesRequireGameThreadExecution() const 
+	{ 
+		return FMassFragmentRequirements::DoesRequireGameThreadExecution() 
+			|| FMassSubsystemRequirements::DoesRequireGameThreadExecution() 
+			|| bRequiresMutatingWorldAccess;
+	}
 
-	bool DoesRequireGameThreadExecution() const { return bRequiresGameThreadExecution || bRequiresMutatingWorldAccess; }
 	void RequireMutatingWorldAccess() { bRequiresMutatingWorldAccess = true; }
 
-	/** The function validates requirements we make for queries. See the FMassEntityQuery struct description for details.
-	 *  Note that this function is non-trivial and end users are not expected to need to use it. 
-	 *  @return whether this query's requirements follow the rules. */
-	bool CheckValidity() const;
-
-	FString DebugGetDescription() const;
-	FString DebugGetArchetypeCompatibilityDescription(const FMassArchetypeHandle& ArchetypeHandle) const;
-
-	TConstArrayView<FMassFragmentRequirement> GetRequirements() const { return Requirements; }
-	const FMassFragmentBitSet& GetRequiredAllFragments() const { return RequiredAllFragments; }
-	const FMassFragmentBitSet& GetRequiredAnyFragments() const { return RequiredAnyFragments; }
-	const FMassFragmentBitSet& GetRequiredOptionalFragments() const { return RequiredOptionalFragments; }
-	const FMassFragmentBitSet& GetRequiredNoneFragments() const { return RequiredNoneFragments; }
-	const FMassTagBitSet& GetRequiredAllTags() const { return RequiredAllTags; }
-	const FMassTagBitSet& GetRequiredAnyTags() const { return RequiredAnyTags; }
-	const FMassTagBitSet& GetRequiredNoneTags() const { return RequiredNoneTags; }
-	const FMassChunkFragmentBitSet& GetRequiredAllChunkFragments() const { return RequiredAllChunkFragments; }
-	const FMassChunkFragmentBitSet& GetRequiredOptionalChunkFragments() const { return RequiredOptionalChunkFragments; }
-	const FMassChunkFragmentBitSet& GetRequiredNoneChunkFragments() const { return RequiredNoneChunkFragments; }
-	const FMassSharedFragmentBitSet& GetRequiredAllSharedFragments() const { return RequiredAllSharedFragments; }
-	const FMassSharedFragmentBitSet& GetRequiredOptionalSharedFragments() const { return RequiredOptionalSharedFragments; }
-	const FMassSharedFragmentBitSet& GetRequiredNoneSharedFragments() const { return RequiredNoneSharedFragments; }
-
 	const TArray<FMassArchetypeHandle>& GetArchetypes() const
 	{ 
 		return ValidArchetypes; 
@@ -462,33 +125,8 @@ public:
 	void ExportRequirements(FMassExecutionRequirements& OutRequirements) const;
 
 protected:
-	void SortRequirements();
 	void ReadCommandlineParams();
 
-protected:
-	friend struct FMassRequirementAccessDetector;
-
-	TArray<FMassFragmentRequirement> Requirements;
-	TArray<FMassFragmentRequirement> ChunkRequirements;
-	TArray<FMassFragmentRequirement> ConstSharedRequirements;
-	TArray<FMassFragmentRequirement> SharedRequirements;
-	FMassTagBitSet RequiredAllTags;
-	FMassTagBitSet RequiredAnyTags;
-	FMassTagBitSet RequiredNoneTags;
-	FMassFragmentBitSet RequiredAllFragments;
-	FMassFragmentBitSet RequiredAnyFragments;
-	FMassFragmentBitSet RequiredOptionalFragments;
-	FMassFragmentBitSet RequiredNoneFragments;
-	FMassChunkFragmentBitSet RequiredAllChunkFragments;
-	FMassChunkFragmentBitSet RequiredOptionalChunkFragments;
-	FMassChunkFragmentBitSet RequiredNoneChunkFragments;
-	FMassSharedFragmentBitSet RequiredAllSharedFragments;
-	FMassSharedFragmentBitSet RequiredOptionalSharedFragments;
-	FMassSharedFragmentBitSet RequiredNoneSharedFragments;
-
-	FMassExternalSubystemBitSet RequiredConstSubsystems;
-	FMassExternalSubystemBitSet RequiredMutableSubsystems;
-
 private:
 	/** 
 	 * This function represents a condition that will be called for every chunk to be processed before the actual 
@@ -521,29 +159,3 @@ private:
 #endif // WITH_MASSENTITY_DEBUG
 };
 
-template<>
-FORCEINLINE FMassEntityQuery& FMassEntityQuery::AddTagRequirements<EMassFragmentPresence::All>(const FMassTagBitSet& TagBitSet)
-{
-	RequiredAllTags += TagBitSet;
-	// force recaching the next time this query is used or the following CacheArchetypes call.
-	DirtyCachedData();
-	return *this;
-}
-
-template<>
-FORCEINLINE FMassEntityQuery& FMassEntityQuery::AddTagRequirements<EMassFragmentPresence::Any>(const FMassTagBitSet& TagBitSet)
-{
-	RequiredAnyTags += TagBitSet;
-	// force recaching the next time this query is used or the following CacheArchetypes call.
-	DirtyCachedData();
-	return *this;
-}
-
-template<>
-FORCEINLINE FMassEntityQuery& FMassEntityQuery::AddTagRequirements<EMassFragmentPresence::None>(const FMassTagBitSet& TagBitSet)
-{
-	RequiredNoneTags += TagBitSet;
-	// force recaching the next time this query is used or the following CacheArchetypes call.
-	DirtyCachedData();
-	return *this;
-}
diff --git a/Engine/Plugins/Runtime/MassEntity/Source/MassEntity/Public/MassEntitySubsystem.h b/Engine/Plugins/Runtime/MassEntity/Source/MassEntity/Public/MassEntitySubsystem.h
--- a/Engine/Plugins/Runtime/MassEntity/Source/MassEntity/Public/MassEntitySubsystem.h
+++ b/Engine/Plugins/Runtime/MassEntity/Source/MassEntity/Public/MassEntitySubsystem.h
@@ -121,7 +121,7 @@ public:
 	FMassArchetypeHandle GetArchetypeForEntityUnsafe(FMassEntityHandle Entity) const;
 
 	/** Method to iterate on all the fragment types of an archetype */
-	static void ForEachArchetypeFragmentType(const FMassArchetypeHandle Archetype, TFunction< void(const UScriptStruct* /*FragmentType*/)> Function);
+	static void ForEachArchetypeFragmentType(const FMassArchetypeHandle& ArchetypeHandle, TFunction< void(const UScriptStruct* /*FragmentType*/)> Function);
 
 	/**
 	 * Go through all archetypes and compact entities
@@ -133,7 +133,7 @@ public:
 	 * Creates fully built entity ready to be used by the subsystem
 	 * @param Archetype you want this entity to be
 	 * @return FMassEntityHandle id of the newly created entity */
-	FMassEntityHandle CreateEntity(const FMassArchetypeHandle Archetype);
+	FMassEntityHandle CreateEntity(const FMassArchetypeHandle& ArchetypeHandle);
 
 	/**
 	 * Creates fully built entity ready to be used by the subsystem
@@ -164,7 +164,7 @@ public:
 	 *  @param Count number of entities to create
 	 *  @param OutEntities the newly created entities are appended to given array, i.e. the pre-existing content of OutEntities won't be affected by the call
 	 *  @return a creation context that will notify all the interested observers about newly created fragments once the context is released */
-	TSharedRef<FEntityCreationContext> BatchCreateEntities(const FMassArchetypeHandle Archetype, const int32 Count, TArray<FMassEntityHandle>& OutEntities);
+	TSharedRef<FEntityCreationContext> BatchCreateEntities(const FMassArchetypeHandle& ArchetypeHandle, const int32 Count, TArray<FMassEntityHandle>& OutEntities);
 
 	/**
 	 * Destroys a fully built entity, use ReleaseReservedEntity if entity was not yet built.
@@ -180,7 +180,7 @@ public:
 	 * Builds an entity for it to be ready to be used by the subsystem
 	 * @param Entity to build which was retrieved with ReserveEntity() method
 	 * @param Archetype you want this entity to be*/
-	void BuildEntity(FMassEntityHandle Entity, const FMassArchetypeHandle Archetype);
+	void BuildEntity(FMassEntityHandle Entity, const FMassArchetypeHandle& ArchetypeHandle);
 
 	/**
 	 * Builds an entity for it to be ready to be used by the subsystem
@@ -346,10 +346,7 @@ public:
 	FMassObserverManager& GetObserverManager() { return ObserverManager; }
 
 #if WITH_MASSENTITY_DEBUG
-	void DebugPrintEntity(int32 Index, FOutputDevice& Ar, const TCHAR* InPrefix = TEXT("")) const;
-	void DebugPrintEntity(FMassEntityHandle Entity, FOutputDevice& Ar, const TCHAR* InPrefix = TEXT("")) const;
 	void DebugPrintArchetypes(FOutputDevice& Ar, const bool bIncludeEmpty = true) const;
-	static void DebugGetStringDesc(const FMassArchetypeHandle& Archetype, FOutputDevice& Ar);
 	void DebugGetArchetypesStringDetails(FOutputDevice& Ar, const bool bIncludeEmpty = true);
 	void DebugGetArchetypeFragmentTypes(const FMassArchetypeHandle& Archetype, TArray<const UScriptStruct*>& InOutFragmentList) const;
 	int32 DebugGetArchetypeEntitiesCount(const FMassArchetypeHandle& Archetype) const;
@@ -377,7 +374,7 @@ protected:
 	FMassArchetypeHandle InternalCreateSimilarArchetype(const FMassArchetypeData& SourceArchetypeRef, FMassArchetypeCompositionDescriptor&& NewComposition);
 
 private:
-	void InternalBuildEntity(FMassEntityHandle Entity, const FMassArchetypeHandle Archetype);
+	void InternalBuildEntity(FMassEntityHandle Entity, const FMassArchetypeHandle& ArchetypeHandle);
 	void InternalReleaseEntity(FMassEntityHandle Entity);
 
 	/** 
diff --git a/Engine/Plugins/Runtime/MassEntity/Source/MassEntity/Public/MassExecutionContext.h b/Engine/Plugins/Runtime/MassEntity/Source/MassEntity/Public/MassExecutionContext.h
--- a/Engine/Plugins/Runtime/MassEntity/Source/MassEntity/Public/MassExecutionContext.h
+++ b/Engine/Plugins/Runtime/MassEntity/Source/MassEntity/Public/MassExecutionContext.h
@@ -14,11 +14,11 @@ private:
 	template< typename ViewType >
 	struct TFragmentView 
 	{
-		FMassFragmentRequirement Requirement;
+		FMassFragmentRequirementDescription Requirement;
 		ViewType FragmentView;
 
 		TFragmentView() {}
-		explicit TFragmentView(const FMassFragmentRequirement& InRequirement) : Requirement(InRequirement) {}
+		explicit TFragmentView(const FMassFragmentRequirementDescription& InRequirement) : Requirement(InRequirement) {}
 
 		bool operator==(const UScriptStruct* FragmentType) const { return Requirement.StructType == FragmentType; }
 	};
@@ -344,10 +344,10 @@ public:
 protected:
 	void SetSubsystemRequirements(const FMassExternalSubystemBitSet& RequiredConstSubsystems, const FMassExternalSubystemBitSet& RequiredMutableSubsystems);
 
-	void SetRequirements(TConstArrayView<FMassFragmentRequirement> InRequirements, 
-		TConstArrayView<FMassFragmentRequirement> InChunkRequirements, 
-		TConstArrayView<FMassFragmentRequirement> InConstSharedRequirements, 
-		TConstArrayView<FMassFragmentRequirement> InSharedRequirements);
+	void SetRequirements(TConstArrayView<FMassFragmentRequirementDescription> InRequirements, 
+		TConstArrayView<FMassFragmentRequirementDescription> InChunkRequirements, 
+		TConstArrayView<FMassFragmentRequirementDescription> InConstSharedRequirements, 
+		TConstArrayView<FMassFragmentRequirementDescription> InSharedRequirements);
 
 	void ClearFragmentViews()
 	{
diff --git a/Engine/Plugins/Runtime/MassEntity/Source/MassEntity/Public/MassProcessingTypes.h b/Engine/Plugins/Runtime/MassEntity/Source/MassEntity/Public/MassProcessingTypes.h
--- a/Engine/Plugins/Runtime/MassEntity/Source/MassEntity/Public/MassProcessingTypes.h
+++ b/Engine/Plugins/Runtime/MassEntity/Source/MassEntity/Public/MassProcessingTypes.h
@@ -112,8 +112,6 @@ struct MASSENTITY_API FMassRuntimePipeline
 	/** goes through Processor looking for a UMassCompositeProcessor instance which GroupName matches the one given as the parameter */
 	UMassCompositeProcessor* FindTopLevelGroupByName(const FName GroupName);
 
-	void DebugOutputDescription(FOutputDevice& Ar) const;
-
 	bool HasProcessorOfExactClass(TSubclassOf<UMassProcessor> InClass) const;
 	bool IsEmpty() const { return Processors.IsEmpty();}
 
diff --git a/Engine/Plugins/Runtime/MassEntity/Source/MassEntity/Public/MassRequirementAccessDetector.h b/Engine/Plugins/Runtime/MassEntity/Source/MassEntity/Public/MassRequirementAccessDetector.h
--- a/Engine/Plugins/Runtime/MassEntity/Source/MassEntity/Public/MassRequirementAccessDetector.h
+++ b/Engine/Plugins/Runtime/MassEntity/Source/MassEntity/Public/MassRequirementAccessDetector.h
@@ -31,9 +31,9 @@ private:
 		}
 	}
 
-	void Aquire(TConstArrayView<FMassFragmentRequirement> Requirements)
+	void Aquire(TConstArrayView<FMassFragmentRequirementDescription> Requirements)
 	{
-		for (const FMassFragmentRequirement& Req : Requirements)
+		for (const FMassFragmentRequirementDescription& Req : Requirements)
 		{
 			if (Req.Presence != EMassFragmentPresence::None)
 			{
@@ -49,9 +49,9 @@ private:
 		}
 	}
 
-	void Release(TConstArrayView<FMassFragmentRequirement> Requirements)
+	void Release(TConstArrayView<FMassFragmentRequirementDescription> Requirements)
 	{
-		for (const FMassFragmentRequirement& Req : Requirements)
+		for (const FMassFragmentRequirementDescription& Req : Requirements)
 		{
 			if (Req.Presence != EMassFragmentPresence::None)
 			{
diff --git a/Engine/Plugins/Runtime/MassEntity/Source/MassEntity/Public/MassRequirements.h b/Engine/Plugins/Runtime/MassEntity/Source/MassEntity/Public/MassRequirements.h
new file mode 100644
--- /dev/null
+++ b/Engine/Plugins/Runtime/MassEntity/Source/MassEntity/Public/MassRequirements.h
@@ -0,0 +1,445 @@
+// Copyright Epic Games, Inc. All Rights Reserved.
+
+#pragma once
+
+#include "MassEntityTypes.h"
+#include "MassRequirements.generated.h"
+
+struct FMassDebugger;
+struct FMassArchetypeHandle;
+struct FMassExecutionRequirements;
+struct FMassRequirementAccessDetector;
+
+UENUM()
+enum class EMassFragmentAccess : uint8
+{
+	/** no binding required */
+	None, 
+
+	/** We want to read the data for the fragment */
+	ReadOnly,
+
+	/** We want to read and write the data for the fragment */
+	ReadWrite,
+
+	MAX
+};
+
+UENUM()
+enum class EMassFragmentPresence : uint8
+{
+	/** All of the required fragments must be present */
+	All,
+
+	/** One of the required fragments must be present */
+	Any,
+
+	/** None of the required fragments can be present */
+	None,
+
+	/** If fragment is present we'll use it, but it missing stop processing of a given archetype */
+	Optional,
+
+	MAX
+};
+
+
+struct MASSENTITY_API FMassFragmentRequirementDescription
+{
+	const UScriptStruct* StructType = nullptr;
+	EMassFragmentAccess AccessMode = EMassFragmentAccess::None;
+	EMassFragmentPresence Presence = EMassFragmentPresence::Optional;
+
+public:
+	FMassFragmentRequirementDescription(){}
+	FMassFragmentRequirementDescription(const UScriptStruct* InStruct, const EMassFragmentAccess InAccessMode, const EMassFragmentPresence InPresence)
+		: StructType(InStruct)
+		, AccessMode(InAccessMode)
+		, Presence(InPresence)
+	{
+		check(InStruct);
+		checkf((Presence != EMassFragmentPresence::Any && Presence != EMassFragmentPresence::Optional)
+			|| AccessMode == EMassFragmentAccess::ReadOnly || AccessMode == EMassFragmentAccess::ReadWrite, TEXT("Only ReadOnly and ReadWrite modes are suppored for optional requirements"));
+	}
+
+	bool RequiresBinding() const { return (AccessMode != EMassFragmentAccess::None); }
+	bool IsOptional() const { return (Presence == EMassFragmentPresence::Optional || Presence == EMassFragmentPresence::Any); }
+
+	/** these functions are used for sorting. See FScriptStructSortOperator */
+	int32 GetStructureSize() const
+	{
+		return StructType->GetStructureSize();
+	}
+
+	FName GetFName() const
+	{
+		return StructType->GetFName();
+	}
+};
+
+/**
+ *  FMassSubsystemRequirements is a structure that declares runtime subsystem access type given calculations require.
+ */
+USTRUCT()
+struct MASSENTITY_API FMassSubsystemRequirements
+{
+	GENERATED_BODY()
+
+	friend FMassDebugger;
+	friend FMassRequirementAccessDetector;
+
+public:
+	template<typename T>
+	FMassSubsystemRequirements& AddSubsystemRequirement(const EMassFragmentAccess AccessMode)
+	{
+		check(AccessMode != EMassFragmentAccess::None && AccessMode != EMassFragmentAccess::MAX);
+
+		switch (AccessMode)
+		{
+		case EMassFragmentAccess::ReadOnly:
+			RequiredConstSubsystems.Add<T>();
+			bRequiresGameThreadExecution |= TMassExternalSubsystemTraits<T>::GameThreadOnly;
+			break;
+		case EMassFragmentAccess::ReadWrite:
+			RequiredMutableSubsystems.Add<T>();
+			bRequiresGameThreadExecution |= TMassExternalSubsystemTraits<T>::GameThreadOnly;
+			break;
+		default:
+			check(false);
+		}
+
+		return *this;
+	}
+
+	FMassSubsystemRequirements& AddSubsystemRequirement(const TSubclassOf<USubsystem> SubsystemClass, const EMassFragmentAccess AccessMode, const bool bGameThreadOnly = true)
+	{
+		check(AccessMode != EMassFragmentAccess::None && AccessMode != EMassFragmentAccess::MAX);
+
+		switch (AccessMode)
+		{
+		case EMassFragmentAccess::ReadOnly:
+			RequiredConstSubsystems.Add(**SubsystemClass);
+			bRequiresGameThreadExecution |= bGameThreadOnly;
+			break;
+		case EMassFragmentAccess::ReadWrite:
+			RequiredMutableSubsystems.Add(**SubsystemClass);
+			bRequiresGameThreadExecution |= bGameThreadOnly;
+			break;
+		default:
+			check(false);
+		}
+
+		return *this;
+	}
+
+	void Reset();
+
+	bool DoesRequireGameThreadExecution() const { return bRequiresGameThreadExecution; }
+	void ExportRequirements(FMassExecutionRequirements& OutRequirements) const;
+
+protected:
+	FMassExternalSubystemBitSet RequiredConstSubsystems;
+	FMassExternalSubystemBitSet RequiredMutableSubsystems;
+
+private:
+	bool bRequiresGameThreadExecution = false;
+};
+
+/** 
+ *  FMassFragmentRequirements is a structure that describes properties required of an archetype that's a subject of calculations.
+ */
+USTRUCT()
+struct MASSENTITY_API FMassFragmentRequirements
+{
+	GENERATED_BODY()
+
+	friend FMassDebugger;
+	friend FMassRequirementAccessDetector;
+
+public:
+	FMassFragmentRequirements(){}
+	FMassFragmentRequirements(std::initializer_list<UScriptStruct*> InitList);
+	FMassFragmentRequirements(TConstArrayView<const UScriptStruct*> InitList);
+
+	FMassFragmentRequirements& AddRequirement(const UScriptStruct* FragmentType, const EMassFragmentAccess AccessMode, const EMassFragmentPresence Presence = EMassFragmentPresence::All)
+	{
+		checkf(FragmentRequirements.FindByPredicate([FragmentType](const FMassFragmentRequirementDescription& Item){ return Item.StructType == FragmentType; }) == nullptr
+			, TEXT("Duplicated requirements are not supported. %s already present"), *GetNameSafe(FragmentType));
+		
+		if (Presence != EMassFragmentPresence::None)
+		{
+			FragmentRequirements.Emplace(FragmentType, AccessMode, Presence);
+		}
+
+		switch (Presence)
+		{
+		case EMassFragmentPresence::All:
+			RequiredAllFragments.Add(*FragmentType);
+			break;
+		case EMassFragmentPresence::Any:
+			RequiredAnyFragments.Add(*FragmentType);
+			break;
+		case EMassFragmentPresence::Optional:
+			RequiredOptionalFragments.Add(*FragmentType);
+			break;
+		case EMassFragmentPresence::None:
+			RequiredNoneFragments.Add(*FragmentType);
+			break;
+		}
+		// force recaching the next time this query is used or the following CacheArchetypes call.
+		IncrementChangeCounter();
+		return *this;
+	}
+
+	/** FMassFragmentRequirements ref returned for chaining */
+	template<typename T>
+	FMassFragmentRequirements& AddRequirement(const EMassFragmentAccess AccessMode, const EMassFragmentPresence Presence = EMassFragmentPresence::All)
+	{
+		checkf(FragmentRequirements.FindByPredicate([](const FMassFragmentRequirementDescription& Item) { return Item.StructType == T::StaticStruct(); }) == nullptr
+			, TEXT("Duplicated requirements are not supported. %s already present"), *T::StaticStruct()->GetName());
+
+		static_assert(TIsDerivedFrom<T, FMassFragment>::IsDerived, "Given struct doesn't represent a valid fragment type. Make sure to inherit from FMassFragment or one of its child-types.");
+		
+		if (Presence != EMassFragmentPresence::None)
+		{
+			FragmentRequirements.Emplace(T::StaticStruct(), AccessMode, Presence);
+		}
+		
+		switch (Presence)
+		{
+		case EMassFragmentPresence::All:
+			RequiredAllFragments.Add<T>();
+			break;
+		case EMassFragmentPresence::Any:
+			RequiredAnyFragments.Add<T>();
+			break;
+		case EMassFragmentPresence::Optional:
+			RequiredOptionalFragments.Add<T>();
+			break;
+		case EMassFragmentPresence::None:
+			RequiredNoneFragments.Add<T>();
+			break;
+		}
+		// force recaching the next time this query is used or the following CacheArchetypes call.
+		IncrementChangeCounter();
+		return *this;
+	}
+
+	void AddTagRequirement(const UScriptStruct& TagType, const EMassFragmentPresence Presence)
+	{
+		checkfSlow(int(Presence) < int(EMassFragmentPresence::Optional), TEXT("Optional and MAX presence are not valid calues for AddTagRequirement"));
+		switch (Presence)
+		{
+		case EMassFragmentPresence::All:
+			RequiredAllTags.Add(TagType);
+			break;
+		case EMassFragmentPresence::Any:
+			RequiredAnyTags.Add(TagType);
+			break;
+		case EMassFragmentPresence::None:
+			RequiredNoneTags.Add(TagType);
+			break;
+		}
+		IncrementChangeCounter();
+	}
+
+	template<typename T>
+	FMassFragmentRequirements& AddTagRequirement(const EMassFragmentPresence Presence)
+	{
+		checkfSlow(int(Presence) < int(EMassFragmentPresence::Optional), TEXT("Optional and MAX presence are not valid calues for AddTagRequirement"));
+		static_assert(TIsDerivedFrom<T, FMassTag>::IsDerived, "Given struct doesn't represent a valid tag type. Make sure to inherit from FMassFragment or one of its child-types.");
+		switch (Presence)
+		{
+			case EMassFragmentPresence::All:
+				RequiredAllTags.Add<T>();
+				break;
+			case EMassFragmentPresence::Any:
+				RequiredAnyTags.Add<T>();
+				break;
+			case EMassFragmentPresence::None:
+				RequiredNoneTags.Add<T>();
+				break;
+		}
+		IncrementChangeCounter();
+		return *this;
+	}
+
+	/** actual implementation in specializations */
+	template<EMassFragmentPresence Presence> 
+	FMassFragmentRequirements& AddTagRequirements(const FMassTagBitSet& TagBitSet)
+	{
+		static_assert(Presence == EMassFragmentPresence::None || Presence == EMassFragmentPresence::All || Presence == EMassFragmentPresence::Any
+			, "The only valid values for AddTagRequirements are All, Any and None");
+		return *this;
+	}
+
+	template<typename T>
+	FMassFragmentRequirements& AddChunkRequirement(const EMassFragmentAccess AccessMode, const EMassFragmentPresence Presence = EMassFragmentPresence::All)
+	{
+		static_assert(TIsDerivedFrom<T, FMassChunkFragment>::IsDerived, "Given struct doesn't represent a valid chunk fragment type. Make sure to inherit from FMassChunkFragment or one of its child-types.");
+		checkf(ChunkFragmentRequirements.FindByPredicate([](const FMassFragmentRequirementDescription& Item) { return Item.StructType == T::StaticStruct(); }) == nullptr
+			, TEXT("Duplicated requirements are not supported. %s already present"), *T::StaticStruct()->GetName());
+		checkfSlow(Presence != EMassFragmentPresence::Any, TEXT("\'Any\' is not a valid Presence value for AddChunkRequirement."));
+
+		switch (Presence)
+		{
+			case EMassFragmentPresence::All:
+				RequiredAllChunkFragments.Add<T>();
+				ChunkFragmentRequirements.Emplace(T::StaticStruct(), AccessMode, Presence);
+				break;
+			case EMassFragmentPresence::Optional:
+				RequiredOptionalChunkFragments.Add<T>();
+				ChunkFragmentRequirements.Emplace(T::StaticStruct(), AccessMode, Presence);
+				break;
+			case EMassFragmentPresence::None:
+				RequiredNoneChunkFragments.Add<T>();
+				break;
+		}
+		IncrementChangeCounter();
+		return *this;
+	}
+
+	template<typename T>
+	FMassFragmentRequirements& AddConstSharedRequirement(const EMassFragmentPresence Presence = EMassFragmentPresence::All)
+	{
+		static_assert(TIsDerivedFrom<T, FMassSharedFragment>::IsDerived, "Given struct doesn't represent a valid shared fragment type. Make sure to inherit from FMassSharedFragment or one of its child-types.");
+		checkf(ConstSharedFragmentRequirements.FindByPredicate([](const FMassFragmentRequirementDescription& Item) { return Item.StructType == T::StaticStruct(); }) == nullptr
+			, TEXT("Duplicated requirements are not supported. %s already present"), *T::StaticStruct()->GetName());
+		checkfSlow(Presence != EMassFragmentPresence::Any, TEXT("\'Any\' is not a valid Presence value for AddSharedRequirement."));
+
+		switch (Presence)
+		{
+		case EMassFragmentPresence::All:
+			RequiredAllSharedFragments.Add<T>();
+			ConstSharedFragmentRequirements.Emplace(T::StaticStruct(), EMassFragmentAccess::ReadOnly, Presence);
+			break;
+		case EMassFragmentPresence::Optional:
+			RequiredOptionalSharedFragments.Add<T>();
+			ConstSharedFragmentRequirements.Emplace(T::StaticStruct(), EMassFragmentAccess::ReadOnly, Presence);
+			break;
+		case EMassFragmentPresence::None:
+			RequiredNoneSharedFragments.Add<T>();
+			break;
+		}
+		IncrementChangeCounter();
+		return *this;
+	}
+
+	template<typename T>
+	FMassFragmentRequirements& AddSharedRequirement(const EMassFragmentAccess AccessMode, const EMassFragmentPresence Presence = EMassFragmentPresence::All)
+	{
+		static_assert(TIsDerivedFrom<T, FMassSharedFragment>::IsDerived, "Given struct doesn't represent a valid shared fragment type. Make sure to inherit from FMassSharedFragment or one of its child-types.");
+		checkf(SharedFragmentRequirements.FindByPredicate([](const FMassFragmentRequirementDescription& Item) { return Item.StructType == T::StaticStruct(); }) == nullptr
+			, TEXT("Duplicated requirements are not supported. %s already present"), *T::StaticStruct()->GetName());
+		checkfSlow(Presence != EMassFragmentPresence::Any, TEXT("\'Any\' is not a valid Presence value for AddSharedRequirement."));
+
+		switch (Presence)
+		{
+		case EMassFragmentPresence::All:
+			RequiredAllSharedFragments.Add<T>();
+			SharedFragmentRequirements.Emplace(T::StaticStruct(), AccessMode, Presence);
+			if (AccessMode == EMassFragmentAccess::ReadWrite)
+			{
+				bRequiresGameThreadExecution |= TMassSharedFragmentTraits<T>::GameThreadOnly;
+			}
+			break;
+		case EMassFragmentPresence::Optional:
+			RequiredOptionalSharedFragments.Add<T>();
+			SharedFragmentRequirements.Emplace(T::StaticStruct(), AccessMode, Presence);
+			if (AccessMode == EMassFragmentAccess::ReadWrite)
+			{
+				bRequiresGameThreadExecution |= TMassSharedFragmentTraits<T>::GameThreadOnly;
+			}
+			break;
+		case EMassFragmentPresence::None:
+			RequiredNoneSharedFragments.Add<T>();
+			break;
+		}
+		IncrementChangeCounter();
+		return *this;
+	}
+
+	void Reset();
+
+	/** The function validates requirements we make for queries. See the FMassFragmentRequirements struct description for details.
+	 *  Note that this function is non-trivial and end users are not expected to need to use it. 
+	 *  @return whether this query's requirements follow the rules. */
+	bool CheckValidity() const;
+
+	TConstArrayView<FMassFragmentRequirementDescription> GetFragmentRequirements() const { return FragmentRequirements; }
+	const FMassFragmentBitSet& GetRequiredAllFragments() const { return RequiredAllFragments; }
+	const FMassFragmentBitSet& GetRequiredAnyFragments() const { return RequiredAnyFragments; }
+	const FMassFragmentBitSet& GetRequiredOptionalFragments() const { return RequiredOptionalFragments; }
+	const FMassFragmentBitSet& GetRequiredNoneFragments() const { return RequiredNoneFragments; }
+	const FMassTagBitSet& GetRequiredAllTags() const { return RequiredAllTags; }
+	const FMassTagBitSet& GetRequiredAnyTags() const { return RequiredAnyTags; }
+	const FMassTagBitSet& GetRequiredNoneTags() const { return RequiredNoneTags; }
+	const FMassChunkFragmentBitSet& GetRequiredAllChunkFragments() const { return RequiredAllChunkFragments; }
+	const FMassChunkFragmentBitSet& GetRequiredOptionalChunkFragments() const { return RequiredOptionalChunkFragments; }
+	const FMassChunkFragmentBitSet& GetRequiredNoneChunkFragments() const { return RequiredNoneChunkFragments; }
+	const FMassSharedFragmentBitSet& GetRequiredAllSharedFragments() const { return RequiredAllSharedFragments; }
+	const FMassSharedFragmentBitSet& GetRequiredOptionalSharedFragments() const { return RequiredOptionalSharedFragments; }
+	const FMassSharedFragmentBitSet& GetRequiredNoneSharedFragments() const { return RequiredNoneSharedFragments; }
+
+	bool DoesArchetypeMatchRequirements(const FMassArchetypeHandle& ArchetypeHandle) const;
+	bool DoesRequireGameThreadExecution() const { return bRequiresGameThreadExecution; }
+	void ExportRequirements(FMassExecutionRequirements& OutRequirements) const;
+
+protected:
+	void SortRequirements();
+
+	FORCEINLINE void IncrementChangeCounter() { ++IncrementalChangesCount; }
+
+protected:
+	friend FMassRequirementAccessDetector;
+
+	TArray<FMassFragmentRequirementDescription> FragmentRequirements;
+	TArray<FMassFragmentRequirementDescription> ChunkFragmentRequirements;
+	TArray<FMassFragmentRequirementDescription> ConstSharedFragmentRequirements;
+	TArray<FMassFragmentRequirementDescription> SharedFragmentRequirements;
+	FMassTagBitSet RequiredAllTags;
+	FMassTagBitSet RequiredAnyTags;
+	FMassTagBitSet RequiredNoneTags;
+	FMassFragmentBitSet RequiredAllFragments;
+	FMassFragmentBitSet RequiredAnyFragments;
+	FMassFragmentBitSet RequiredOptionalFragments;
+	FMassFragmentBitSet RequiredNoneFragments;
+	FMassChunkFragmentBitSet RequiredAllChunkFragments;
+	FMassChunkFragmentBitSet RequiredOptionalChunkFragments;
+	FMassChunkFragmentBitSet RequiredNoneChunkFragments;
+	FMassSharedFragmentBitSet RequiredAllSharedFragments;
+	FMassSharedFragmentBitSet RequiredOptionalSharedFragments;
+	FMassSharedFragmentBitSet RequiredNoneSharedFragments;
+
+	uint32 IncrementalChangesCount = 0;
+
+private:
+	bool bRequiresGameThreadExecution = false;
+};
+
+template<>
+FORCEINLINE FMassFragmentRequirements& FMassFragmentRequirements::AddTagRequirements<EMassFragmentPresence::All>(const FMassTagBitSet& TagBitSet)
+{
+	RequiredAllTags += TagBitSet;
+	// force recaching the next time this query is used or the following CacheArchetypes call.
+	IncrementChangeCounter();
+	return *this;
+}
+
+template<>
+FORCEINLINE FMassFragmentRequirements& FMassFragmentRequirements::AddTagRequirements<EMassFragmentPresence::Any>(const FMassTagBitSet& TagBitSet)
+{
+	RequiredAnyTags += TagBitSet;
+	// force recaching the next time this query is used or the following CacheArchetypes call.
+	IncrementChangeCounter();
+	return *this;
+}
+
+template<>
+FORCEINLINE FMassFragmentRequirements& FMassFragmentRequirements::AddTagRequirements<EMassFragmentPresence::None>(const FMassTagBitSet& TagBitSet)
+{
+	RequiredNoneTags += TagBitSet;
+	// force recaching the next time this query is used or the following CacheArchetypes call.
+	IncrementChangeCounter();
+	return *this;
+}
diff --git a/Engine/Plugins/Runtime/MassEntity/Source/MassEntityTestSuite/Private/MassQueryTest.cpp b/Engine/Plugins/Runtime/MassEntity/Source/MassEntityTestSuite/Private/MassQueryTest.cpp
--- a/Engine/Plugins/Runtime/MassEntity/Source/MassEntityTestSuite/Private/MassQueryTest.cpp
+++ b/Engine/Plugins/Runtime/MassEntity/Source/MassEntityTestSuite/Private/MassQueryTest.cpp
@@ -34,7 +34,7 @@ struct FQueryTest_ProcessorRequirements : FEntityTestBase
 		CA_ASSUME(EntitySubsystem);
 
 		UMassTestProcessor_Floats* Processor = NewObject<UMassTestProcessor_Floats>(EntitySubsystem);
-		TConstArrayView<FMassFragmentRequirement> Requirements = Processor->TestGetQuery().GetRequirements();
+		TConstArrayView<FMassFragmentRequirementDescription> Requirements = Processor->TestGetQuery().GetFragmentRequirements();
 		
 		AITEST_TRUE("Query should have extracted some requirements from the given Processor", Requirements.Num() > 0);
 		AITEST_TRUE("There should be exactly one requirement", Requirements.Num() == 1);
@@ -53,7 +53,7 @@ struct FQueryTest_ExplicitRequirements : FEntityTestBase
 		CA_ASSUME(EntitySubsystem);
 				
 		FMassEntityQuery Query({ FTestFragment_Float::StaticStruct()});
-		TConstArrayView<FMassFragmentRequirement> Requirements = Query.GetRequirements();
+		TConstArrayView<FMassFragmentRequirementDescription> Requirements = Query.GetFragmentRequirements();
 
 		AITEST_TRUE("Query should have extracted some requirements from the given Processor", Requirements.Num() > 0);
 		AITEST_TRUE("There should be exactly one requirement", Requirements.Num() == 1);
@@ -214,8 +214,9 @@ struct FQueryTest_ExecuteSparse : FEntityTestBase
 		int TotalProcessed = 0;
 
 		FMassExecutionContext ExecContext;
-		FMassEntityQuery().AddRequirement<FTestFragment_Float>(EMassFragmentAccess::ReadWrite)
-			.ForEachEntityChunk(FMassArchetypeEntityCollection(FloatsArchetype, EntitiesToProcess, FMassArchetypeEntityCollection::NoDuplicates)
+		FMassEntityQuery TestQuery;
+		TestQuery.AddRequirement<FTestFragment_Float>(EMassFragmentAccess::ReadWrite);
+		TestQuery.ForEachEntityChunk(FMassArchetypeEntityCollection(FloatsArchetype, EntitiesToProcess, FMassArchetypeEntityCollection::NoDuplicates)
 								, *EntitySubsystem, ExecContext, [&TotalProcessed](FMassExecutionContext& Context)
 			{
 				TotalProcessed += Context.GetNumEntities();
diff --git a/Engine/Plugins/Runtime/MassGameplay/Source/MassSpawner/Private/MassEntityTemplate.cpp b/Engine/Plugins/Runtime/MassGameplay/Source/MassSpawner/Private/MassEntityTemplate.cpp
--- a/Engine/Plugins/Runtime/MassGameplay/Source/MassSpawner/Private/MassEntityTemplate.cpp
+++ b/Engine/Plugins/Runtime/MassGameplay/Source/MassSpawner/Private/MassEntityTemplate.cpp
@@ -2,6 +2,7 @@
 
 #include "MassEntityTemplate.h"
 #include "VisualLogger/VisualLoggerTypes.h"
+#include "MassDebugger.h"
 
 DEFINE_ENUM_TO_STRING(EMassEntityTemplateIDType, "/Script/MassSpawner");
 
@@ -47,7 +48,7 @@ FString FMassEntityTemplate::DebugGetArchetypeDescription(UMassEntitySubsystem&
 {
 	FStringOutputDevice OutDescription;
 #if WITH_MASSGAMEPLAY_DEBUG
-	EntitySubsystem.DebugGetStringDesc(Archetype, OutDescription);
+	FMassDebugger::OutputArchetypeDescription(OutDescription, Archetype);
 #endif // WITH_MASSGAMEPLAY_DEBUG
 	return MoveTemp(OutDescription);
 }
