Renamed FMassExternalSubystemBitSet to FMassExternalSubsystemBitSet

diff --git a/Engine/Plugins/Runtime/MassEntity/Source/MassEntity/Private/MassEntityTypes.cpp b/Engine/Plugins/Runtime/MassEntity/Source/MassEntity/Private/MassEntityTypes.cpp
--- a/Engine/Plugins/Runtime/MassEntity/Source/MassEntity/Private/MassEntityTypes.cpp
+++ b/Engine/Plugins/Runtime/MassEntity/Source/MassEntity/Private/MassEntityTypes.cpp
@@ -7,7 +7,7 @@ DEFINE_TYPEBITSET(FMassFragmentBitSet);
 DEFINE_TYPEBITSET(FMassTagBitSet);
 DEFINE_TYPEBITSET(FMassChunkFragmentBitSet);
 DEFINE_TYPEBITSET(FMassSharedFragmentBitSet);
-DEFINE_TYPEBITSET(FMassExternalSubystemBitSet);
+DEFINE_TYPEBITSET(FMassExternalSubsystemBitSet);
 
 uint32 FMassArchetypeSharedFragmentValues::CalculateHash() const
 {
diff --git a/Engine/Plugins/Runtime/MassEntity/Source/MassEntity/Private/MassExecutionContext.cpp b/Engine/Plugins/Runtime/MassEntity/Source/MassEntity/Private/MassExecutionContext.cpp
--- a/Engine/Plugins/Runtime/MassEntity/Source/MassEntity/Private/MassExecutionContext.cpp
+++ b/Engine/Plugins/Runtime/MassEntity/Source/MassEntity/Private/MassExecutionContext.cpp
@@ -21,7 +21,7 @@ void FMassExecutionContext::ClearExecutionData()
 	MutableSubsystemsBitSet.Reset();
 }
 
-void FMassExecutionContext::SetSubsystemRequirements(const FMassExternalSubystemBitSet& RequiredConstSubsystems, const FMassExternalSubystemBitSet& RequiredMutableSubsystems)
+void FMassExecutionContext::SetSubsystemRequirements(const FMassExternalSubsystemBitSet& RequiredConstSubsystems, const FMassExternalSubsystemBitSet& RequiredMutableSubsystems)
 {
 	ConstSubsystemsBitSet = RequiredConstSubsystems;
 	MutableSubsystemsBitSet = RequiredMutableSubsystems;
diff --git a/Engine/Plugins/Runtime/MassEntity/Source/MassEntity/Private/MassProcessorDependencySolver.cpp b/Engine/Plugins/Runtime/MassEntity/Source/MassEntity/Private/MassProcessorDependencySolver.cpp
--- a/Engine/Plugins/Runtime/MassEntity/Source/MassEntity/Private/MassProcessorDependencySolver.cpp
+++ b/Engine/Plugins/Runtime/MassEntity/Source/MassEntity/Private/MassProcessorDependencySolver.cpp
@@ -85,7 +85,7 @@ FProcessorDependencySolver::FResourceUsage::FResourceUsage()
 		FragmentsAccess[i].Access.AddZeroed(FMassFragmentBitSet::GetMaxNum());
 		ChunkFragmentsAccess[i].Access.AddZeroed(FMassChunkFragmentBitSet::GetMaxNum());
 		SharedFragmentsAccess[i].Access.AddZeroed(FMassSharedFragmentBitSet::GetMaxNum());
-		RequiredSubsystemsAccess[i].Access.AddZeroed(FMassExternalSubystemBitSet::GetMaxNum());
+		RequiredSubsystemsAccess[i].Access.AddZeroed(FMassExternalSubsystemBitSet::GetMaxNum());
 	}
 }
 
@@ -154,7 +154,7 @@ bool FProcessorDependencySolver::FResourceUsage::CanAccessRequirements(const FMa
 	bool bCanAccess = CanAccess<FMassFragmentBitSet>(Requirements.Fragments, TestedRequirements.Fragments)
 		&& CanAccess<FMassChunkFragmentBitSet>(Requirements.ChunkFragments, TestedRequirements.ChunkFragments)
 		&& CanAccess<FMassSharedFragmentBitSet>(Requirements.SharedFragments, TestedRequirements.SharedFragments)
-		&& CanAccess<FMassExternalSubystemBitSet>(Requirements.RequiredSubsystems, TestedRequirements.RequiredSubsystems);
+		&& CanAccess<FMassExternalSubsystemBitSet>(Requirements.RequiredSubsystems, TestedRequirements.RequiredSubsystems);
 
 	return bCanAccess;
 }
@@ -164,7 +164,7 @@ void FProcessorDependencySolver::FResourceUsage::SubmitNode(const int32 NodeInde
 	HandleElementType<FMassFragmentBitSet>(FragmentsAccess, InOutNode.Requirements.Fragments, InOutNode, NodeIndex);
 	HandleElementType<FMassChunkFragmentBitSet>(ChunkFragmentsAccess, InOutNode.Requirements.ChunkFragments, InOutNode, NodeIndex);
 	HandleElementType<FMassSharedFragmentBitSet>(SharedFragmentsAccess, InOutNode.Requirements.SharedFragments, InOutNode, NodeIndex);
-	HandleElementType<FMassExternalSubystemBitSet>(RequiredSubsystemsAccess, InOutNode.Requirements.RequiredSubsystems, InOutNode, NodeIndex);
+	HandleElementType<FMassExternalSubsystemBitSet>(RequiredSubsystemsAccess, InOutNode.Requirements.RequiredSubsystems, InOutNode, NodeIndex);
 
 	Requirements.Append(InOutNode.Requirements);
 }
diff --git a/Engine/Plugins/Runtime/MassEntity/Source/MassEntity/Private/MassRequirementAccessDetector.cpp b/Engine/Plugins/Runtime/MassEntity/Source/MassEntity/Private/MassRequirementAccessDetector.cpp
--- a/Engine/Plugins/Runtime/MassEntity/Source/MassEntity/Private/MassRequirementAccessDetector.cpp
+++ b/Engine/Plugins/Runtime/MassEntity/Source/MassEntity/Private/MassRequirementAccessDetector.cpp
@@ -19,7 +19,7 @@ void FMassRequirementAccessDetector::Initialize()
 	AddDetectors(FMassFragmentBitSet::FStructTrackerWrapper::StructTracker);
 	AddDetectors(FMassChunkFragmentBitSet::FStructTrackerWrapper::StructTracker);
 	AddDetectors(FMassSharedFragmentBitSet::FStructTrackerWrapper::StructTracker);
-	AddDetectors(FMassExternalSubystemBitSet::FStructTrackerWrapper::StructTracker);
+	AddDetectors(FMassExternalSubsystemBitSet::FStructTrackerWrapper::StructTracker);
 }
 
 void FMassRequirementAccessDetector::AddDetectors(const FStructTracker& StructTracker)
diff --git a/Engine/Plugins/Runtime/MassEntity/Source/MassEntity/Public/MassDebugger.h b/Engine/Plugins/Runtime/MassEntity/Source/MassEntity/Public/MassDebugger.h
--- a/Engine/Plugins/Runtime/MassEntity/Source/MassEntity/Public/MassDebugger.h
+++ b/Engine/Plugins/Runtime/MassEntity/Source/MassEntity/Public/MassDebugger.h
@@ -35,8 +35,8 @@ struct MASSENTITY_API FQueryRequirementsView
 	const FMassTagBitSet& RequiredAllTags;
 	const FMassTagBitSet& RequiredAnyTags;
 	const FMassTagBitSet& RequiredNoneTags;
-	const FMassExternalSubystemBitSet& RequiredConstSubsystems;
-	const FMassExternalSubystemBitSet& RequiredMutableSubsystems;
+	const FMassExternalSubsystemBitSet& RequiredConstSubsystems;
+	const FMassExternalSubsystemBitSet& RequiredMutableSubsystems;
 };
 
 FString DebugGetFragmentAccessString(EMassFragmentAccess Access);
diff --git a/Engine/Plugins/Runtime/MassEntity/Source/MassEntity/Public/MassEntityTypes.h b/Engine/Plugins/Runtime/MassEntity/Source/MassEntity/Public/MassEntityTypes.h
--- a/Engine/Plugins/Runtime/MassEntity/Source/MassEntity/Public/MassEntityTypes.h
+++ b/Engine/Plugins/Runtime/MassEntity/Source/MassEntity/Public/MassEntityTypes.h
@@ -108,7 +108,7 @@ DECLARE_STRUCTTYPEBITSET_EXPORTED(MASSENTITY_API, FMassFragmentBitSet, FMassFrag
 DECLARE_STRUCTTYPEBITSET_EXPORTED(MASSENTITY_API, FMassTagBitSet, FMassTag);
 DECLARE_STRUCTTYPEBITSET_EXPORTED(MASSENTITY_API, FMassChunkFragmentBitSet, FMassChunkFragment);
 DECLARE_STRUCTTYPEBITSET_EXPORTED(MASSENTITY_API, FMassSharedFragmentBitSet, FMassSharedFragment);
-DECLARE_CLASSTYPEBITSET_EXPORTED(MASSENTITY_API, FMassExternalSubystemBitSet, USubsystem);
+DECLARE_CLASSTYPEBITSET_EXPORTED(MASSENTITY_API, FMassExternalSubsystemBitSet, USubsystem);
 
 /** The type summarily describing a composition of an entity or an archetype. It contains information on both the
  *  fragments as well as tags */
diff --git a/Engine/Plugins/Runtime/MassEntity/Source/MassEntity/Public/MassExecutionContext.h b/Engine/Plugins/Runtime/MassEntity/Source/MassEntity/Public/MassExecutionContext.h
--- a/Engine/Plugins/Runtime/MassEntity/Source/MassEntity/Public/MassExecutionContext.h
+++ b/Engine/Plugins/Runtime/MassEntity/Source/MassEntity/Public/MassExecutionContext.h
@@ -34,8 +34,8 @@ private:
 	using FSharedFragmentView = TFragmentView<FStructView>;
 	TArray<FSharedFragmentView, TInlineAllocator<4>> SharedFragmentViews;
 
-	FMassExternalSubystemBitSet ConstSubsystemsBitSet;
-	FMassExternalSubystemBitSet MutableSubsystemsBitSet;
+	FMassExternalSubsystemBitSet ConstSubsystemsBitSet;
+	FMassExternalSubsystemBitSet MutableSubsystemsBitSet;
 	TArray<TObjectPtr<USubsystem>> Subsystems;
 	
 	// mz@todo make this shared ptr thread-safe and never auto-flush in MT environment. 
@@ -75,7 +75,7 @@ public:
 		: DeltaTimeSeconds(InDeltaTimeSeconds)
 		, bFlushDeferredCommands(bInFlushDeferredCommands)
 	{
-		Subsystems.AddZeroed(FMassExternalSubystemBitSet::GetMaxNum());
+		Subsystems.AddZeroed(FMassExternalSubsystemBitSet::GetMaxNum());
 	}
 
 	FMassExecutionContext()
@@ -244,7 +244,7 @@ public:
 	T* GetMutableSubsystem(const UWorld* World)
 	{
 		// @todo consider getting this directly from entity subsystem - it should cache all the used system
-		const uint32 SystemIndex = FMassExternalSubystemBitSet::GetTypeIndex<T>();
+		const uint32 SystemIndex = FMassExternalSubsystemBitSet::GetTypeIndex<T>();
 		if (ensure(MutableSubsystemsBitSet.IsBitSet(SystemIndex)))
 		{
 			return GetSubsystemInternal<T>(World, SystemIndex);
@@ -264,7 +264,7 @@ public:
 	template<typename T>
 	const T* GetSubsystem(const UWorld* World)
 	{
-		const uint32 SystemIndex = FMassExternalSubystemBitSet::GetTypeIndex<T>();
+		const uint32 SystemIndex = FMassExternalSubsystemBitSet::GetTypeIndex<T>();
 		if (ensure(ConstSubsystemsBitSet.IsBitSet(SystemIndex) || MutableSubsystemsBitSet.IsBitSet(SystemIndex)))
 		{
 			return GetSubsystemInternal<T>(World, SystemIndex);
@@ -284,7 +284,7 @@ public:
 	T* GetMutableSubsystem(const UWorld* World, const TSubclassOf<UWorldSubsystem> SubsystemClass)
 	{
 		// @todo consider getting this directly from entity subsystem - it should cache all the used system
-		const uint32 SystemIndex = FMassExternalSubystemBitSet::GetTypeIndex(**SubsystemClass);
+		const uint32 SystemIndex = FMassExternalSubsystemBitSet::GetTypeIndex(**SubsystemClass);
 		if (ensure(MutableSubsystemsBitSet.IsBitSet(SystemIndex)))
 		{
 			return GetSubsystemInternal<T>(World, SystemIndex, SubsystemClass);
@@ -304,7 +304,7 @@ public:
 	template<typename T>
 	const T* GetSubsystem(const UWorld* World, const TSubclassOf<UWorldSubsystem> SubsystemClass)
 	{
-		const uint32 SystemIndex = FMassExternalSubystemBitSet::GetTypeIndex(**SubsystemClass);
+		const uint32 SystemIndex = FMassExternalSubsystemBitSet::GetTypeIndex(**SubsystemClass);
 		if (ensure(ConstSubsystemsBitSet.IsBitSet(SystemIndex) || MutableSubsystemsBitSet.IsBitSet(SystemIndex)))
 		{
 			return GetSubsystemInternal<T>(World, SystemIndex, SubsystemClass);
@@ -342,7 +342,7 @@ public:
 	}
 
 protected:
-	void SetSubsystemRequirements(const FMassExternalSubystemBitSet& RequiredConstSubsystems, const FMassExternalSubystemBitSet& RequiredMutableSubsystems);
+	void SetSubsystemRequirements(const FMassExternalSubsystemBitSet& RequiredConstSubsystems, const FMassExternalSubsystemBitSet& RequiredMutableSubsystems);
 
 	void SetRequirements(TConstArrayView<FMassFragmentRequirementDescription> InRequirements, 
 		TConstArrayView<FMassFragmentRequirementDescription> InChunkRequirements, 
diff --git a/Engine/Plugins/Runtime/MassEntity/Source/MassEntity/Public/MassProcessorDependencySolver.h b/Engine/Plugins/Runtime/MassEntity/Source/MassEntity/Public/MassProcessorDependencySolver.h
--- a/Engine/Plugins/Runtime/MassEntity/Source/MassEntity/Public/MassProcessorDependencySolver.h
+++ b/Engine/Plugins/Runtime/MassEntity/Source/MassEntity/Public/MassProcessorDependencySolver.h
@@ -57,7 +57,7 @@ struct MASSENTITY_API FMassExecutionRequirements
 	TMassExecutionAccess<FMassFragmentBitSet> Fragments;
 	TMassExecutionAccess<FMassChunkFragmentBitSet> ChunkFragments;
 	TMassExecutionAccess<FMassSharedFragmentBitSet> SharedFragments;
-	TMassExecutionAccess<FMassExternalSubystemBitSet> RequiredSubsystems;
+	TMassExecutionAccess<FMassExternalSubsystemBitSet> RequiredSubsystems;
 	FMassTagBitSet RequiredAllTags;
 	FMassTagBitSet RequiredAnyTags;
 	FMassTagBitSet RequiredNoneTags;
diff --git a/Engine/Plugins/Runtime/MassEntity/Source/MassEntity/Public/MassRequirements.h b/Engine/Plugins/Runtime/MassEntity/Source/MassEntity/Public/MassRequirements.h
--- a/Engine/Plugins/Runtime/MassEntity/Source/MassEntity/Public/MassRequirements.h
+++ b/Engine/Plugins/Runtime/MassEntity/Source/MassEntity/Public/MassRequirements.h
@@ -138,8 +138,8 @@ public:
 	void ExportRequirements(FMassExecutionRequirements& OutRequirements) const;
 
 protected:
-	FMassExternalSubystemBitSet RequiredConstSubsystems;
-	FMassExternalSubystemBitSet RequiredMutableSubsystems;
+	FMassExternalSubsystemBitSet RequiredConstSubsystems;
+	FMassExternalSubsystemBitSet RequiredMutableSubsystems;
 
 private:
 	bool bRequiresGameThreadExecution = false;
