Got rid of DECLARE_ENUM_TO_STRING-declared fuinctions form Mass enums

diff --git a/Engine/Plugins/Runtime/MassEntity/Source/MassEntity/Private/MassDebugger.cpp b/Engine/Plugins/Runtime/MassEntity/Source/MassEntity/Private/MassDebugger.cpp
--- a/Engine/Plugins/Runtime/MassEntity/Source/MassEntity/Private/MassDebugger.cpp
+++ b/Engine/Plugins/Runtime/MassEntity/Source/MassEntity/Private/MassDebugger.cpp
@@ -14,10 +14,6 @@
 #include "Engine/Engine.h"
 
 
-DEFINE_ENUM_TO_STRING(EMassFragmentAccess, "/Script/MassEntity");
-DEFINE_ENUM_TO_STRING(EMassFragmentPresence, "/Script/MassEntity");
-DEFINE_ENUM_TO_STRING(EMassProcessingPhase, "/Script/MassEntity");
-
 namespace UE::Mass::Debug
 {
 
diff --git a/Engine/Plugins/Runtime/MassEntity/Source/MassEntity/Private/MassProcessingPhase.cpp b/Engine/Plugins/Runtime/MassEntity/Source/MassEntity/Private/MassProcessingPhase.cpp
--- a/Engine/Plugins/Runtime/MassEntity/Source/MassEntity/Private/MassProcessingPhase.cpp
+++ b/Engine/Plugins/Runtime/MassEntity/Source/MassEntity/Private/MassProcessingPhase.cpp
@@ -232,7 +232,7 @@ void UMassProcessingPhaseManager::EnableTickFunctions(const UWorld& World)
 			FStringOutputDevice Ar;
 			Phase.PhaseProcessor->DebugOutputDescription(Ar);
 			UE_VLOG_UELOG(this, LogMass, Log, TEXT("Enabling phase %s tick:\n%s")
-				, *EnumToString(Phase.Phase), *Ar);
+				, *UEnum::GetValueAsString(Phase.Phase), *Ar);
 		}
 #endif // WITH_MASSENTITY_DEBUG		
 	}
@@ -287,7 +287,7 @@ void UMassProcessingPhaseManager::SetPhaseProcessor(const EMassProcessingPhase P
 		FStringOutputDevice Ar;
 		PhaseProcessor->DebugOutputDescription(Ar);
 		UE_VLOG(this, LogMass, Log, TEXT("Setting new group processor for phase %s:\n%s")
-			, *EnumToString(Phase), *Ar);
+			, *UEnum::GetValueAsString(Phase), *Ar);
 #endif // WITH_MASSENTITY_DEBUG
 	}
 	else
diff --git a/Engine/Plugins/Runtime/MassEntity/Source/MassEntity/Public/MassDebugger.h b/Engine/Plugins/Runtime/MassEntity/Source/MassEntity/Public/MassDebugger.h
--- a/Engine/Plugins/Runtime/MassEntity/Source/MassEntity/Public/MassDebugger.h
+++ b/Engine/Plugins/Runtime/MassEntity/Source/MassEntity/Public/MassDebugger.h
@@ -6,7 +6,6 @@
 
 #if WITH_MASSENTITY_DEBUG
 #include "Containers/ContainersFwd.h"
-#include "VisualLogger/VisualLoggerTypes.h"
 #include "MassEntityQuery.h"
 #include "MassProcessor.h"
 
@@ -20,10 +19,6 @@ struct FMassFragmentRequirementDescription;
 enum class EMassFragmentAccess : uint8;
 enum class EMassFragmentPresence : uint8;
 
-MASSENTITY_API DECLARE_ENUM_TO_STRING(EMassFragmentAccess);
-MASSENTITY_API DECLARE_ENUM_TO_STRING(EMassFragmentPresence);
-MASSENTITY_API DECLARE_ENUM_TO_STRING(EMassProcessingPhase);
-
 namespace UE::Mass::Debug
 {
 struct MASSENTITY_API FQueryRequirementsView
