Made MassProcessor's subsystem requirements no longer requiring manual calls to ConfigureContextForProcessorUse.
This also let us not call processor's Execute if subsystem requirements are not met.

diff --git a/Engine/Plugins/AI/MassAI/Source/MassAIBehavior/Private/MassStateTreeProcessors.cpp b/Engine/Plugins/AI/MassAI/Source/MassAIBehavior/Private/MassStateTreeProcessors.cpp
--- a/Engine/Plugins/AI/MassAI/Source/MassAIBehavior/Private/MassStateTreeProcessors.cpp
+++ b/Engine/Plugins/AI/MassAI/Source/MassAIBehavior/Private/MassStateTreeProcessors.cpp
@@ -255,7 +255,6 @@ void UMassStateTreeActivationProcessor::ConfigureQueries()
 
 void UMassStateTreeActivationProcessor::Execute(UMassEntitySubsystem& EntitySubsystem, FMassExecutionContext& Context)
 {
-	ConfigureContextForProcessorUse(Context);
 	UMassSignalSubsystem& SignalSubsystem = Context.GetMutableSubsystemChecked<UMassSignalSubsystem>(EntitySubsystem.GetWorld());
 
 	const UMassBehaviorSettings* BehaviorSettings = GetDefault<UMassBehaviorSettings>();
@@ -383,7 +382,6 @@ void UMassStateTreeProcessor::ConfigureQueries()
 void UMassStateTreeProcessor::SignalEntities(UMassEntitySubsystem& EntitySubsystem, FMassExecutionContext& Context, FMassSignalNameLookup& EntitySignals)
 {
 	UWorld* World = EntitySubsystem.GetWorld();
-	ConfigureContextForProcessorUse(Context);
 	UMassSignalSubsystem& SignalSubsystem = Context.GetMutableSubsystemChecked<UMassSignalSubsystem>(World);
 	
 	QUICK_SCOPE_CYCLE_COUNTER(StateTreeProcessor_Run);
diff --git a/Engine/Plugins/AI/MassAI/Source/MassAIBehavior/Private/MassZoneGraphAnnotationProcessors.cpp b/Engine/Plugins/AI/MassAI/Source/MassAIBehavior/Private/MassZoneGraphAnnotationProcessors.cpp
--- a/Engine/Plugins/AI/MassAI/Source/MassAIBehavior/Private/MassZoneGraphAnnotationProcessors.cpp
+++ b/Engine/Plugins/AI/MassAI/Source/MassAIBehavior/Private/MassZoneGraphAnnotationProcessors.cpp
@@ -119,7 +119,6 @@ void UMassZoneGraphAnnotationTagUpdateProcessor::Execute(UMassEntitySubsystem& E
 
 	if (TransientEntitiesToSignal.Num())
 	{
-		ConfigureContextForProcessorUse(Context);
 		UMassSignalSubsystem& SignalSubsystem = Context.GetMutableSubsystemChecked<UMassSignalSubsystem>(World);
 		SignalSubsystem.SignalEntities(UE::Mass::Signals::AnnotationTagsChanged, TransientEntitiesToSignal);
 	}
diff --git a/Engine/Plugins/AI/MassCrowd/Source/MassCrowd/Private/MassCrowdServerRepresentationLODProcessor.cpp b/Engine/Plugins/AI/MassCrowd/Source/MassCrowd/Private/MassCrowdServerRepresentationLODProcessor.cpp
--- a/Engine/Plugins/AI/MassCrowd/Source/MassCrowd/Private/MassCrowdServerRepresentationLODProcessor.cpp
+++ b/Engine/Plugins/AI/MassCrowd/Source/MassCrowd/Private/MassCrowdServerRepresentationLODProcessor.cpp
@@ -50,7 +50,6 @@ void UMassCrowdServerRepresentationLODProcessor::Execute(UMassEntitySubsystem& E
 {
 	TRACE_CPUPROFILER_EVENT_SCOPE(TEXT("CrowdServerRepresentationLOD"))
 
-	ConfigureContextForProcessorUse(Context);
 	const UMassLODSubsystem& LODSubsystem = Context.GetSubsystemChecked<UMassLODSubsystem>(EntitySubsystem.GetWorld());
 	const TArray<FViewerInfo>& Viewers = LODSubsystem.GetViewers();
 	LODCalculator.PrepareExecution(Viewers);
diff --git a/Engine/Plugins/Experimental/StructUtils/Source/StructUtils/Public/StructTypeBitSet.h b/Engine/Plugins/Experimental/StructUtils/Source/StructUtils/Public/StructTypeBitSet.h
--- a/Engine/Plugins/Experimental/StructUtils/Source/StructUtils/Public/StructTypeBitSet.h
+++ b/Engine/Plugins/Experimental/StructUtils/Source/StructUtils/Public/StructTypeBitSet.h
@@ -375,6 +375,11 @@ public:
 		static const TStructTypeBitSet TypeBitSet(GetTypeIndex<T>());
 		return TypeBitSet;
 	}
+
+	static const TUStructType* GetTypeAtIndex(const int32 Index)
+	{
+		return Cast<const TUStructType>(TStructTrackerWrapper::StructTracker.GetStructType(Index));
+	}
 	
 	template<typename T>
 	FORCEINLINE void Add()
diff --git a/Engine/Plugins/Runtime/MassEntity/Source/MassEntity/Private/MassEntityQuery.cpp b/Engine/Plugins/Runtime/MassEntity/Source/MassEntity/Private/MassEntityQuery.cpp
--- a/Engine/Plugins/Runtime/MassEntity/Source/MassEntity/Private/MassEntityQuery.cpp
+++ b/Engine/Plugins/Runtime/MassEntity/Source/MassEntity/Private/MassEntityQuery.cpp
@@ -132,7 +132,28 @@ void FMassEntityQuery::ForEachEntityChunk(UMassEntitySubsystem& EntitySubsystem,
 	EntitySubsystem.GetRequirementAccessDetector().RequireAccess(*this);
 #endif
 
-	ExecutionContext.SetSubsystemRequirements(RequiredConstSubsystems, RequiredMutableSubsystems);
+	struct FScopedSubsystemRequirementsRestore
+	{
+		FScopedSubsystemRequirementsRestore(FMassExecutionContext& ExecutionContext)
+			: CachedExecutionContext(ExecutionContext)
+		{
+			ConstSubsystemsBitSet = ExecutionContext.ConstSubsystemsBitSet;
+			MutableSubsystemsBitSet = ExecutionContext.MutableSubsystemsBitSet;
+		}
+
+		~FScopedSubsystemRequirementsRestore()
+		{
+			CachedExecutionContext.ConstSubsystemsBitSet = ConstSubsystemsBitSet;
+			CachedExecutionContext.MutableSubsystemsBitSet = MutableSubsystemsBitSet;
+		}
+
+		FMassExecutionContext& CachedExecutionContext;
+		FMassExternalSubsystemBitSet ConstSubsystemsBitSet;
+		FMassExternalSubsystemBitSet MutableSubsystemsBitSet;
+	};
+	FScopedSubsystemRequirementsRestore SubsystemRestore(ExecutionContext);
+
+	ExecutionContext.SetSubsystemRequirements(*this);
 
 	// if there's a chunk collection set by the external code - use that
 	if (ExecutionContext.GetEntityCollection().IsSet())
@@ -149,7 +170,7 @@ void FMassEntityQuery::ForEachEntityChunk(UMassEntitySubsystem& EntitySubsystem,
 #endif // WITH_MASSENTITY_DEBUG
 			return;
 		}
-		ExecutionContext.SetRequirements(FragmentRequirements, ChunkFragmentRequirements, ConstSharedFragmentRequirements, SharedFragmentRequirements);
+		ExecutionContext.SetFragmentRequirements(*this);
 		
 		FMassArchetypeData& ArchetypeData = FMassArchetypeHelper::ArchetypeDataFromHandleChecked(ArchetypeHandle);
 		ArchetypeData.ExecuteFunction(ExecutionContext, ExecuteFunction
@@ -163,7 +184,7 @@ void FMassEntityQuery::ForEachEntityChunk(UMassEntitySubsystem& EntitySubsystem,
 	{
 		CacheArchetypes(EntitySubsystem);
 		// it's important to set requirements after caching archetypes due to that call potentially sorting the requirements and the order is relevant here.
-		ExecutionContext.SetRequirements(FragmentRequirements, ChunkFragmentRequirements, ConstSharedFragmentRequirements, SharedFragmentRequirements);
+		ExecutionContext.SetFragmentRequirements(*this);
 
 		for (int i = 0; i < ValidArchetypes.Num(); ++i)
 		{
@@ -217,7 +238,7 @@ void FMassEntityQuery::ParallelForEachEntityChunk(UMassEntitySubsystem& EntitySu
 			return;
 		}
 
-		ExecutionContext.SetRequirements(FragmentRequirements, ChunkFragmentRequirements, ConstSharedFragmentRequirements, SharedFragmentRequirements);
+		ExecutionContext.SetFragmentRequirements(*this);
 		FMassArchetypeData& ArchetypeRef = FMassArchetypeHelper::ArchetypeDataFromHandleChecked(ArchetypeHandle);
 		const FMassArchetypeEntityCollection AsEntityCollection(ArchetypeHandle);
 		for (const FMassArchetypeEntityCollection::FArchetypeEntityRange& EntityRange : AsEntityCollection.GetRanges())
@@ -228,7 +249,7 @@ void FMassEntityQuery::ParallelForEachEntityChunk(UMassEntitySubsystem& EntitySu
 	else
 	{
 		CacheArchetypes(EntitySubsystem);
-		ExecutionContext.SetRequirements(FragmentRequirements, ChunkFragmentRequirements, ConstSharedFragmentRequirements, SharedFragmentRequirements);
+		ExecutionContext.SetFragmentRequirements(*this);
 		for (int ArchetypeIndex = 0; ArchetypeIndex < ValidArchetypes.Num(); ++ArchetypeIndex)
 		{
 			FMassArchetypeHandle& ArchetypeHandle = ValidArchetypes[ArchetypeIndex];
@@ -253,12 +274,6 @@ void FMassEntityQuery::ParallelForEachEntityChunk(UMassEntitySubsystem& EntitySu
 	ExecutionContext.FlushDeferred(EntitySubsystem);
 }
 
-void FMassEntityQuery::ApplyQueryRequirementsToContext(FMassExecutionContext& ExecutionContext)
-{
-	ExecutionContext.SetSubsystemRequirements(RequiredConstSubsystems, RequiredMutableSubsystems);
-	ExecutionContext.SetRequirements(FragmentRequirements, ChunkFragmentRequirements, ConstSharedFragmentRequirements, SharedFragmentRequirements);
-}
-
 int32 FMassEntityQuery::GetNumMatchingEntities(UMassEntitySubsystem& InEntitySubsystem)
 {
 	CacheArchetypes(InEntitySubsystem);
diff --git a/Engine/Plugins/Runtime/MassEntity/Source/MassEntity/Private/MassExecutionContext.cpp b/Engine/Plugins/Runtime/MassEntity/Source/MassEntity/Private/MassExecutionContext.cpp
--- a/Engine/Plugins/Runtime/MassEntity/Source/MassEntity/Private/MassExecutionContext.cpp
+++ b/Engine/Plugins/Runtime/MassEntity/Source/MassEntity/Private/MassExecutionContext.cpp
@@ -4,6 +4,8 @@
 #include "MassArchetypeData.h"
 
 
+PRAGMA_DISABLE_OPTIMIZATION
+
 void FMassExecutionContext::FlushDeferred(UMassEntitySubsystem& EntitySystem) const
 {
 	if (bFlushDeferredCommands && DeferredCommandBuffer)
@@ -15,16 +17,76 @@ void FMassExecutionContext::FlushDeferred(UMassEntitySubsystem& EntitySystem) co
 void FMassExecutionContext::ClearExecutionData()
 {
 	FragmentViews.Reset();
+	ChunkFragmentViews.Reset();
+	ConstSharedFragmentViews.Reset();
+	SharedFragmentViews.Reset();
 	CurrentArchetypesTagBitSet.Reset();
 	ChunkSerialModificationNumber = -1;
-	ConstSubsystemsBitSet.Reset();
-	MutableSubsystemsBitSet.Reset();
 }
 
-void FMassExecutionContext::SetSubsystemRequirements(const FMassExternalSubsystemBitSet& RequiredConstSubsystems, const FMassExternalSubsystemBitSet& RequiredMutableSubsystems)
+bool FMassExecutionContext::CacheSubsystem(const UWorld* World, const uint32 SystemIndex)
 {
-	ConstSubsystemsBitSet = RequiredConstSubsystems;
-	MutableSubsystemsBitSet = RequiredMutableSubsystems;
+	if (UNLIKELY(Subsystems.IsValidIndex(SystemIndex) == false))
+	{
+		Subsystems.AddZeroed(Subsystems.Num() - SystemIndex + 1);
+	}
+
+	if (Subsystems[SystemIndex].Get())
+	{
+		return true;
+	}
+
+	const UClass* SubsystemClass = FMassExternalSubsystemBitSet::GetTypeAtIndex(SystemIndex);
+	if (SubsystemClass)
+	{
+		const TSubclassOf<UWorldSubsystem> WorldSubsystemClass(const_cast<UClass*>(SubsystemClass));
+		if (*WorldSubsystemClass)
+		{
+			USubsystem* SystemInstance = FMassExternalSubsystemTraits::GetInstance<UWorldSubsystem>(World, WorldSubsystemClass);
+			Subsystems[SystemIndex] = SystemInstance;
+			return SystemInstance != nullptr;
+		}
+		else
+		{
+			// not a UWorldSubsystem. There's no way to generically fetch an instance of one so we assume it's there
+			return true;
+		}
+	}
+	return false;
+}
+
+bool FMassExecutionContext::CacheSubsystemRequirements(const UWorld* World, const FMassSubsystemRequirements& SubsystemRequirements)
+{
+	if (SubsystemRequirements.IsEmpty())
+	{
+		return true;
+	}
+
+	for (FMassExternalSubsystemBitSet::FIndexIterator It = SubsystemRequirements.GetRequiredConstSubsystems().GetIndexIterator(); It; ++It)
+	{
+		if (CacheSubsystem(World, *It) == false)
+		{
+			return false;
+		}
+	}
+
+	for (FMassExternalSubsystemBitSet::FIndexIterator It = SubsystemRequirements.GetRequiredMutableSubsystems().GetIndexIterator(); It; ++It)
+	{
+		if (CacheSubsystem(World, *It) == false)
+		{
+			return false;
+		}
+	}
+
+	ConstSubsystemsBitSet = SubsystemRequirements.GetRequiredConstSubsystems();
+	MutableSubsystemsBitSet = SubsystemRequirements.GetRequiredMutableSubsystems();
+	return true;
+}
+
+void FMassExecutionContext::SetSubsystemRequirements(const FMassSubsystemRequirements& SubsystemRequirements)
+{
+	ConstSubsystemsBitSet = SubsystemRequirements.GetRequiredConstSubsystems();
+	MutableSubsystemsBitSet = SubsystemRequirements.GetRequiredMutableSubsystems();
 }
 
 void FMassExecutionContext::SetEntityCollection(const FMassArchetypeEntityCollection& InEntityCollection)
@@ -39,13 +101,10 @@ void FMassExecutionContext::SetEntityCollection(FMassArchetypeEntityCollection&&
 	EntityCollection = MoveTemp(InEntityCollection);
 }
 
-void FMassExecutionContext::SetRequirements(TConstArrayView<FMassFragmentRequirementDescription> InRequirements, 
-	TConstArrayView<FMassFragmentRequirementDescription> InChunkRequirements, 
-	TConstArrayView<FMassFragmentRequirementDescription> InConstSharedRequirements, 
-	TConstArrayView<FMassFragmentRequirementDescription> InSharedRequirements)
-{ 
+void FMassExecutionContext::SetFragmentRequirements(const FMassFragmentRequirements& FragmentRequirements)
+{
 	FragmentViews.Reset();
-	for (const FMassFragmentRequirementDescription& Requirement : InRequirements)
+	for (const FMassFragmentRequirementDescription& Requirement : FragmentRequirements.GetFragmentRequirements())
 	{
 		if (Requirement.RequiresBinding())
 		{
@@ -54,7 +113,7 @@ void FMassExecutionContext::SetRequirements(TConstArrayView<FMassFragmentRequire
 	}
 
 	ChunkFragmentViews.Reset();
-	for (const FMassFragmentRequirementDescription& Requirement : InChunkRequirements)
+	for (const FMassFragmentRequirementDescription& Requirement : FragmentRequirements.GetChunkFragmentRequirements())
 	{
 		if (Requirement.RequiresBinding())
 		{
@@ -63,7 +122,7 @@ void FMassExecutionContext::SetRequirements(TConstArrayView<FMassFragmentRequire
 	}
 
 	ConstSharedFragmentViews.Reset();
-	for (const FMassFragmentRequirementDescription& Requirement : InConstSharedRequirements)
+	for (const FMassFragmentRequirementDescription& Requirement : FragmentRequirements.GetConstSharedFragmentRequirements())
 	{
 		if (Requirement.RequiresBinding())
 		{
@@ -72,7 +131,7 @@ void FMassExecutionContext::SetRequirements(TConstArrayView<FMassFragmentRequire
 	}
 
 	SharedFragmentViews.Reset();
-	for (const FMassFragmentRequirementDescription& Requirement : InSharedRequirements)
+	for (const FMassFragmentRequirementDescription& Requirement : FragmentRequirements.GetSharedFragmentRequirements())
 	{
 		if (Requirement.RequiresBinding())
 		{
@@ -81,3 +140,5 @@ void FMassExecutionContext::SetRequirements(TConstArrayView<FMassFragmentRequire
 	}
 }
 
+
+PRAGMA_ENABLE_OPTIMIZATION
\ No newline at end of file
diff --git a/Engine/Plugins/Runtime/MassEntity/Source/MassEntity/Private/MassProcessor.cpp b/Engine/Plugins/Runtime/MassEntity/Source/MassEntity/Private/MassProcessor.cpp
--- a/Engine/Plugins/Runtime/MassEntity/Source/MassEntity/Private/MassProcessor.cpp
+++ b/Engine/Plugins/Runtime/MassEntity/Source/MassEntity/Private/MassProcessor.cpp
@@ -16,13 +16,19 @@ namespace UE::Mass::Debug
 	FAutoConsoleVariableRef CVarLogProcessingGraph(TEXT("mass.LogProcessingGraph"), bLogProcessingGraph
 		, TEXT("When enabled will log task graph tasks created while dispatching processors to other threads, along with their dependencies"), ECVF_Cheat);
 }
-#endif // WITH_MASSENTITY_DEBUG
 
-// change to && 1 to enable more detailed processing tasks logging
-#if WITH_MASSENTITY_DEBUG && 0
-#define PROCESSOR_LOG(Fmt, ...) UE_LOG(LogMass, Verbose, Fmt, ##__VA_ARGS__)
-#else // WITH_MASSENTITY_DEBUG
-#define PROCESSOR_LOG(...) 
+#define PROCESSOR_LOG(Verbosity, Fmt, ...) UE_VLOG_UELOG(this, LogMass, Verbosity, Fmt, ##__VA_ARGS__)
+
+// change to 1 to enable more detailed processing tasks logging
+#if 0
+#define PROCESSOR_TASK_LOG(Fmt, ...) UE_VLOG_UELOG(this, LogMass, Verbose, Fmt, ##__VA_ARGS__)
+#else
+#define PROCESSOR_TASK_LOG(...) 
+#endif // 0
+
+#else 
+#define PROCESSOR_LOG(...)
+#define PROCESSOR_TASK_LOG(...) 
 #endif // WITH_MASSENTITY_DEBUG
 
 class FMassProcessorTask
@@ -52,7 +58,7 @@ public:
 		checkf(Processor, TEXT("Expecting a valid processor to execute"));
 		checkf(EntitySubsystem, TEXT("Expecting a valid entity subsystem to execute processor"));
 
-		PROCESSOR_LOG(TEXT("+--+ Task %s started on %u"), *Processor->GetProcessorName(), FPlatformTLS::GetCurrentThreadId());
+		PROCESSOR_TASK_LOG(TEXT("+--+ Task %s started on %u"), *Processor->GetProcessorName(), FPlatformTLS::GetCurrentThreadId());
 
 		UMassEntitySubsystem::FScopedProcessing ProcessingScope = EntitySubsystem->NewProcessingScope();
 
@@ -69,7 +75,7 @@ public:
 		{
 			Processor->CallExecute(*EntitySubsystem, ExecutionContext);
 		}
-		PROCESSOR_LOG(TEXT("+--+ Task %s finished"), *Processor->GetProcessorName());
+		PROCESSOR_TASK_LOG(TEXT("+--+ Task %s finished"), *Processor->GetProcessorName());
 	}
 
 private:
@@ -107,7 +113,6 @@ UMassProcessor::UMassProcessor(const FObjectInitializer& ObjectInitializer)
 UMassProcessor::UMassProcessor()
 	: ExecutionFlags((int32)(EProcessorExecutionFlags::Server | EProcessorExecutionFlags::Standalone))
 {
-	RegisterQuery(ProcessorRequirements);
 }
 
 void UMassProcessor::SetShouldAutoRegisterWithGlobalList(const bool bAutoRegister)
@@ -152,12 +157,18 @@ void UMassProcessor::CallExecute(UMassEntitySubsystem& EntitySubsystem, FMassExe
 #if WITH_MASSENTITY_DEBUG
 	Context.DebugSetExecutionDesc(FString::Printf(TEXT("%s (%s)"), *GetProcessorName(), *ToString(EntitySubsystem.GetWorld()->GetNetMode())));
 #endif
-	Execute(EntitySubsystem, Context);
-}
-
-void UMassProcessor::ConfigureContextForProcessorUse(FMassExecutionContext& Context)
-{
-	ProcessorRequirements.ApplyQueryRequirementsToContext(Context);
+	// CacheSubsystemRequirements will return true only if all requirements declared with ProcessorRequirements are met
+	// meaning if it fails there's no point in calling Execute.
+	// Note that we're not testing individual queries in OwnedQueries - processors can function just fine with some 
+	// of their queries not having anything to do.
+	if (Context.CacheSubsystemRequirements(GetWorld(), ProcessorRequirements))
+	{
+		Execute(EntitySubsystem, Context);
+	}
+	else
+	{
+		PROCESSOR_LOG(VeryVerbose, TEXT("%s Skipping Execute due to subsystem requirements not being met"), *GetProcessorName());
+	}
 }
 
 void UMassProcessor::ExportRequirements(FMassExecutionRequirements& OutRequirements) const
@@ -279,7 +290,7 @@ FGraphEventRef UMassCompositeProcessor::DispatchProcessorTasks(UMassEntitySubsys
 			}
 
 			check(ProcessingNode.Processor);
-			PROCESSOR_LOG(TEXT("Task %u %s%s%s"), Events[i]->GetTraceId(), *ProcessingNode.Processor->GetProcessorName()
+			PROCESSOR_TASK_LOG(TEXT("Task %u %s%s%s"), Events[i]->GetTraceId(), *ProcessingNode.Processor->GetProcessorName()
 				, DependenciesDesc.Len() > 0 ? TEXT(" depends on ") : TEXT(""), *DependenciesDesc);
 		}
 	}
diff --git a/Engine/Plugins/Runtime/MassEntity/Source/MassEntity/Private/MassRequirements.cpp b/Engine/Plugins/Runtime/MassEntity/Source/MassEntity/Private/MassRequirements.cpp
--- a/Engine/Plugins/Runtime/MassEntity/Source/MassEntity/Private/MassRequirements.cpp
+++ b/Engine/Plugins/Runtime/MassEntity/Source/MassEntity/Private/MassRequirements.cpp
@@ -83,6 +83,12 @@ bool FMassFragmentRequirements::CheckValidity() const
 	return RequiredAllFragments.IsEmpty() == false || RequiredAnyFragments.IsEmpty() == false || RequiredOptionalFragments.IsEmpty() == false;
 }
 
+bool FMassFragmentRequirements::IsEmpty() const
+{
+	return FragmentRequirements.IsEmpty() && ChunkFragmentRequirements.IsEmpty() && ConstSharedFragmentRequirements.IsEmpty() 
+		&& SharedFragmentRequirements.IsEmpty() && RequiredAllTags.IsEmpty() && RequiredAnyTags.IsEmpty() && RequiredNoneTags.IsEmpty();
+}
+
 bool FMassFragmentRequirements::DoesArchetypeMatchRequirements(const FMassArchetypeHandle& ArchetypeHandle) const
 {
 	check(ArchetypeHandle.IsValid());
diff --git a/Engine/Plugins/Runtime/MassEntity/Source/MassEntity/Public/MassEntityQuery.h b/Engine/Plugins/Runtime/MassEntity/Source/MassEntity/Public/MassEntityQuery.h
--- a/Engine/Plugins/Runtime/MassEntity/Source/MassEntity/Public/MassEntityQuery.h
+++ b/Engine/Plugins/Runtime/MassEntity/Source/MassEntity/Public/MassEntityQuery.h
@@ -46,9 +46,6 @@ public:
 	 */
 	void ParallelForEachEntityChunk(UMassEntitySubsystem& EntitySubsystem, FMassExecutionContext& ExecutionContext, const FMassExecuteFunction& ExecuteFunction);
 
-	/** Configures ExecutionContext with access/presence requirements declared by this query */
-	void ApplyQueryRequirementsToContext(FMassExecutionContext& ExecutionContext);
-
 	/** Will gather all archetypes from InEntitySubsystem matching this->Requirements.
 	 *  Note that no work will be done if the cached data is up to date (as tracked by EntitySubsystemHash and 
 	 *	ArchetypeDataVersion properties). */
diff --git a/Engine/Plugins/Runtime/MassEntity/Source/MassEntity/Public/MassExecutionContext.h b/Engine/Plugins/Runtime/MassEntity/Source/MassEntity/Public/MassExecutionContext.h
--- a/Engine/Plugins/Runtime/MassEntity/Source/MassEntity/Public/MassExecutionContext.h
+++ b/Engine/Plugins/Runtime/MassEntity/Source/MassEntity/Public/MassExecutionContext.h
@@ -341,13 +341,16 @@ public:
 		CurrentArchetypesTagBitSet = BitSet;
 	}
 
+	/** 
+	 * Processes SubsystemRequirements to fetch and process all the indicated subsystems.
+	 * @return `true` if all required subsystems have been found, `false` otherwise
+	 */
+	bool CacheSubsystemRequirements(const UWorld* World, const FMassSubsystemRequirements& SubsystemRequirements);
+
 protected:
-	void SetSubsystemRequirements(const FMassExternalSubsystemBitSet& RequiredConstSubsystems, const FMassExternalSubsystemBitSet& RequiredMutableSubsystems);
+	void SetSubsystemRequirements(const FMassSubsystemRequirements& SubsystemRequirements);
 
-	void SetRequirements(TConstArrayView<FMassFragmentRequirementDescription> InRequirements, 
-		TConstArrayView<FMassFragmentRequirementDescription> InChunkRequirements, 
-		TConstArrayView<FMassFragmentRequirementDescription> InConstSharedRequirements, 
-		TConstArrayView<FMassFragmentRequirementDescription> InSharedRequirements);
+	void SetFragmentRequirements(const FMassFragmentRequirements& FragmentRequirements);
 
 	void ClearFragmentViews()
 	{
@@ -402,4 +405,6 @@ protected:
 		}
 		return SystemInstance;
 	}
+
+	bool CacheSubsystem(const UWorld* World, const uint32 SystemIndex);
 };
diff --git a/Engine/Plugins/Runtime/MassEntity/Source/MassEntity/Public/MassExternalSubsystemTraits.h b/Engine/Plugins/Runtime/MassEntity/Source/MassEntity/Public/MassExternalSubsystemTraits.h
--- a/Engine/Plugins/Runtime/MassEntity/Source/MassEntity/Public/MassExternalSubsystemTraits.h
+++ b/Engine/Plugins/Runtime/MassEntity/Source/MassEntity/Public/MassExternalSubsystemTraits.h
@@ -43,14 +43,14 @@ namespace FMassExternalSubsystemTraits
 	This is a getter function that given a UWorld* fetches an instance.
 	*/
 	template<typename T, typename = typename TEnableIf<TIsDerivedFrom<T, UWorldSubsystem>::IsDerived>::Type>
-	T* GetInstance(const UWorld* World)
+	FORCEINLINE T* GetInstance(const UWorld* World)
 	{ 
 		// note that the default implementation works only for UWorldSubsystems
 		return UWorld::GetSubsystem<T>(World);
 	}
 
 	template<typename T>
-	T* GetInstance(const UWorld* World, const TSubclassOf<UWorldSubsystem> SubsystemClass)
+	FORCEINLINE T* GetInstance(const UWorld* World, const TSubclassOf<UWorldSubsystem> SubsystemClass)
 	{
 		return World ? Cast<T>(World->GetSubsystemBase(SubsystemClass)) : (T*)nullptr;
 	}
diff --git a/Engine/Plugins/Runtime/MassEntity/Source/MassEntity/Public/MassProcessor.h b/Engine/Plugins/Runtime/MassEntity/Source/MassEntity/Public/MassProcessor.h
--- a/Engine/Plugins/Runtime/MassEntity/Source/MassEntity/Public/MassProcessor.h
+++ b/Engine/Plugins/Runtime/MassEntity/Source/MassEntity/Public/MassProcessor.h
@@ -56,10 +56,6 @@ public:
 	/** Whether this processor should execute according the CurrentExecutionFlags parameters */
 	bool ShouldExecute(const EProcessorExecutionFlags CurrentExecutionFlags) const { return (GetExecutionFlags() & CurrentExecutionFlags) != EProcessorExecutionFlags::None; }
 	void CallExecute(UMassEntitySubsystem& EntitySubsystem, FMassExecutionContext& Context);
-	
-	/** Configures Context to accept access requests performed outside of FMassEntityQuery::ForEachEntityChunk calls.
-	 *  The access is declared via ProcessorQuery member variable */
-	void ConfigureContextForProcessorUse(FMassExecutionContext& Context);
 
 	bool AllowDuplicates() const { return bAllowDuplicates; }
 
@@ -79,6 +75,8 @@ public:
 	 *	this function to supply additional requirements */
 	virtual void ExportRequirements(FMassExecutionRequirements& OutRequirements) const;
 
+	const FMassSubsystemRequirements& GetProcessorRequirements() const { return ProcessorRequirements; }
+
 	/** Adds Query to RegisteredQueries list. Query is required to be a member variable of this processor. Not meeting
 	 *  this requirement will cause check failure and the query won't be registered. */
 	void RegisterQuery(FMassEntityQuery& Query);
@@ -141,7 +139,7 @@ protected:
 	friend FMassDebugger;
 
 	/** A query representing elements this processor is accessing in Execute function outside of query execution */
-	FMassEntityQuery ProcessorRequirements;
+	FMassSubsystemRequirements ProcessorRequirements;
 
 private:
 	/** Stores processor's queries registered via RegisterQuery. 
diff --git a/Engine/Plugins/Runtime/MassEntity/Source/MassEntity/Public/MassRequirements.h b/Engine/Plugins/Runtime/MassEntity/Source/MassEntity/Public/MassRequirements.h
--- a/Engine/Plugins/Runtime/MassEntity/Source/MassEntity/Public/MassRequirements.h
+++ b/Engine/Plugins/Runtime/MassEntity/Source/MassEntity/Public/MassRequirements.h
@@ -134,6 +134,10 @@ public:
 
 	void Reset();
 
+	const FMassExternalSubsystemBitSet& GetRequiredConstSubsystems() const { return RequiredConstSubsystems; }
+	const FMassExternalSubsystemBitSet& GetRequiredMutableSubsystems() const { return RequiredMutableSubsystems; }
+	bool IsEmpty() const { return RequiredConstSubsystems.IsEmpty() && RequiredMutableSubsystems.IsEmpty(); }
+
 	bool DoesRequireGameThreadExecution() const { return bRequiresGameThreadExecution; }
 	void ExportRequirements(FMassExecutionRequirements& OutRequirements) const;
 
@@ -367,6 +371,9 @@ public:
 	bool CheckValidity() const;
 
 	TConstArrayView<FMassFragmentRequirementDescription> GetFragmentRequirements() const { return FragmentRequirements; }
+	TConstArrayView<FMassFragmentRequirementDescription> GetChunkFragmentRequirements() const { return ChunkFragmentRequirements; }
+	TConstArrayView<FMassFragmentRequirementDescription> GetConstSharedFragmentRequirements() const { return ConstSharedFragmentRequirements; }
+	TConstArrayView<FMassFragmentRequirementDescription> GetSharedFragmentRequirements() const { return SharedFragmentRequirements; }
 	const FMassFragmentBitSet& GetRequiredAllFragments() const { return RequiredAllFragments; }
 	const FMassFragmentBitSet& GetRequiredAnyFragments() const { return RequiredAnyFragments; }
 	const FMassFragmentBitSet& GetRequiredOptionalFragments() const { return RequiredOptionalFragments; }
@@ -381,6 +388,8 @@ public:
 	const FMassSharedFragmentBitSet& GetRequiredOptionalSharedFragments() const { return RequiredOptionalSharedFragments; }
 	const FMassSharedFragmentBitSet& GetRequiredNoneSharedFragments() const { return RequiredNoneSharedFragments; }
 
+	bool IsEmpty() const;
+
 	bool DoesArchetypeMatchRequirements(const FMassArchetypeHandle& ArchetypeHandle) const;
 	bool DoesRequireGameThreadExecution() const { return bRequiresGameThreadExecution; }
 	void ExportRequirements(FMassExecutionRequirements& OutRequirements) const;
diff --git a/Engine/Plugins/Runtime/MassGameplay/Source/MassLOD/Private/MassLODCollectorProcessor.cpp b/Engine/Plugins/Runtime/MassGameplay/Source/MassLOD/Private/MassLODCollectorProcessor.cpp
--- a/Engine/Plugins/Runtime/MassGameplay/Source/MassLOD/Private/MassLODCollectorProcessor.cpp
+++ b/Engine/Plugins/Runtime/MassGameplay/Source/MassLOD/Private/MassLODCollectorProcessor.cpp
@@ -80,8 +80,6 @@ void UMassLODCollectorProcessor::ExecuteInternal(UMassEntitySubsystem& EntitySub
 
 void UMassLODCollectorProcessor::Execute(UMassEntitySubsystem& EntitySubsystem, FMassExecutionContext& Context)
 {
-	ConfigureContextForProcessorUse(Context);
-
 	const UMassLODSubsystem& LODSubsystem = Context.GetSubsystemChecked<UMassLODSubsystem>(EntitySubsystem.GetWorld());
 	const TArray<FViewerInfo>& Viewers = LODSubsystem.GetViewers();
 	Collector.PrepareExecution(Viewers);
diff --git a/Engine/Plugins/Runtime/MassGameplay/Source/MassLOD/Private/MassSimulationLOD.cpp b/Engine/Plugins/Runtime/MassGameplay/Source/MassLOD/Private/MassSimulationLOD.cpp
--- a/Engine/Plugins/Runtime/MassGameplay/Source/MassLOD/Private/MassSimulationLOD.cpp
+++ b/Engine/Plugins/Runtime/MassGameplay/Source/MassLOD/Private/MassSimulationLOD.cpp
@@ -119,8 +119,6 @@ void UMassSimulationLODProcessor::Execute(UMassEntitySubsystem& EntitySubsystem,
 
 	{
 		TRACE_CPUPROFILER_EVENT_SCOPE(TEXT("PrepareExecution"));
-		
-		ConfigureContextForProcessorUse(Context);
 
 		const UMassLODSubsystem& LODSubsystem = Context.GetSubsystemChecked<UMassLODSubsystem>(EntitySubsystem.GetWorld());
 		const TArray<FViewerInfo>& Viewers = LODSubsystem.GetViewers();
diff --git a/Engine/Plugins/Runtime/MassGameplay/Source/MassReplication/Private/MassReplicationProcessor.cpp b/Engine/Plugins/Runtime/MassGameplay/Source/MassReplication/Private/MassReplicationProcessor.cpp
--- a/Engine/Plugins/Runtime/MassGameplay/Source/MassReplication/Private/MassReplicationProcessor.cpp
+++ b/Engine/Plugins/Runtime/MassGameplay/Source/MassReplication/Private/MassReplicationProcessor.cpp
@@ -164,7 +164,6 @@ void UMassReplicationProcessor::Execute(UMassEntitySubsystem& EntitySubsystem, F
 		PrepareExecution(EntitySubsystem);
 	}
 
-	ConfigureContextForProcessorUse(Context);
 	const UMassLODSubsystem& LODSubsystem = Context.GetSubsystemChecked<UMassLODSubsystem>(EntitySubsystem.GetWorld());
 	const TArray<FViewerInfo>& AllViewersInfo = LODSubsystem.GetViewers();
 	const TArray<FMassClientHandle>& ClientHandles = ReplicationSubsystem->GetClientReplicationHandles();
diff --git a/Engine/Plugins/Runtime/MassGameplay/Source/MassRepresentation/Private/MassVisualizationLODProcessor.cpp b/Engine/Plugins/Runtime/MassGameplay/Source/MassRepresentation/Private/MassVisualizationLODProcessor.cpp
--- a/Engine/Plugins/Runtime/MassGameplay/Source/MassRepresentation/Private/MassVisualizationLODProcessor.cpp
+++ b/Engine/Plugins/Runtime/MassGameplay/Source/MassRepresentation/Private/MassVisualizationLODProcessor.cpp
@@ -65,7 +65,6 @@ void UMassVisualizationLODProcessor::Execute(UMassEntitySubsystem& EntitySubsyst
 
 	{
 		TRACE_CPUPROFILER_EVENT_SCOPE(PrepareExecution)
-		ConfigureContextForProcessorUse(Context);
 		const UMassLODSubsystem& LODSubsystem = Context.GetSubsystemChecked<UMassLODSubsystem>(EntitySubsystem.GetWorld());
 		const TArray<FViewerInfo>& Viewers = LODSubsystem.GetViewers();
 		EntitySubsystem.ForEachSharedFragment<FMassVisualizationLODSharedFragment>([this, &Viewers](FMassVisualizationLODSharedFragment& LODSharedFragment)
diff --git a/Engine/Plugins/Runtime/MassGameplay/Source/MassSmartObjects/Private/MassSmartObjectProcessor.cpp b/Engine/Plugins/Runtime/MassGameplay/Source/MassSmartObjects/Private/MassSmartObjectProcessor.cpp
--- a/Engine/Plugins/Runtime/MassGameplay/Source/MassSmartObjects/Private/MassSmartObjectProcessor.cpp
+++ b/Engine/Plugins/Runtime/MassGameplay/Source/MassSmartObjects/Private/MassSmartObjectProcessor.cpp
@@ -56,7 +56,6 @@ void UMassSmartObjectCandidatesFinderProcessor::Execute(UMassEntitySubsystem& En
 {
 	UWorld* World = EntitySubsystem.GetWorld();
 
-	ConfigureContextForProcessorUse(Context);
 	UMassSignalSubsystem& SignalSubsystem = Context.GetMutableSubsystemChecked<UMassSignalSubsystem>(World);
 	const UZoneGraphAnnotationSubsystem& AnnotationSubsystem = Context.GetSubsystemChecked<UZoneGraphAnnotationSubsystem>(World);
 	
@@ -378,7 +377,6 @@ void UMassSmartObjectTimedBehaviorProcessor::Execute(UMassEntitySubsystem& Entit
 
 	if (ToRelease.Num())
 	{
-		ConfigureContextForProcessorUse(Context);
 		UMassSignalSubsystem& SignalSubsystem = Context.GetMutableSubsystemChecked<UMassSignalSubsystem>(World);
 		SignalSubsystem.SignalEntities(UE::Mass::Signals::SmartObjectInteractionDone, ToRelease);
 	}
