Updated MassDebuggerModel witch changes to UMassProcessor.ProcessorRequirements

diff --git a/Engine/Plugins/Runtime/MassEntity/Source/MassEntityDebugger/Private/MassDebuggerModel.cpp b/Engine/Plugins/Runtime/MassEntity/Source/MassEntityDebugger/Private/MassDebuggerModel.cpp
--- a/Engine/Plugins/Runtime/MassEntity/Source/MassEntityDebugger/Private/MassDebuggerModel.cpp
+++ b/Engine/Plugins/Runtime/MassEntity/Source/MassEntityDebugger/Private/MassDebuggerModel.cpp
@@ -78,6 +78,14 @@ FMassDebuggerQueryData::FMassDebuggerQueryData(const FMassEntityQuery& Query, co
 #endif // WITH_MASSENTITY_DEBUG
 }
 
+FMassDebuggerQueryData::FMassDebuggerQueryData(const FMassSubsystemRequirements& SubsystemRequirements, const FText& InLabel)
+	: Label(InLabel)
+{
+#if WITH_MASSENTITY_DEBUG
+	SubsystemRequirements.ExportRequirements(ExecutionRequirements);
+#endif // WITH_MASSENTITY_DEBUG
+}
+
 int32 FMassDebuggerQueryData::GetTotalBitsUsedCount() 
 {
 	return ExecutionRequirements.GetTotalBitsUsedCount();
@@ -97,17 +105,13 @@ FMassDebuggerProcessorData::FMassDebuggerProcessorData(const UMassProcessor& InP
 #if WITH_MASSENTITY_DEBUG
 	TConstArrayView<FMassEntityQuery*> ProcessorQueries = FMassDebugger::GetProcessorQueries(InProcessor);
 
-	if (ensureMsgf(ProcessorQueries.Num() >= 1, TEXT("We expect at least one query to be there - every processor should have their ProcessorRequirements registered in OwnedQueries")))
-	{
-		ProcessorRequirements = MakeShareable(new FMassDebuggerQueryData(*ProcessorQueries[0], LOCTEXT("MassProcessorRequirementsLabel", "Processor Requirements")));
+	ProcessorRequirements = MakeShareable(new FMassDebuggerQueryData(InProcessor.GetProcessorRequirements(), LOCTEXT("MassProcessorRequirementsLabel", "Processor Requirements")));
 
-		Queries.Reserve(ProcessorQueries.Num() - 1);
-		for (int i = 1; i < ProcessorQueries.Num(); ++i)
-		{
-			const FMassEntityQuery* Query = ProcessorQueries[i];
-			check(Query);
-			Queries.Add(MakeShareable(new FMassDebuggerQueryData(*Query, LOCTEXT("MassEntityQueryLabel", "Query"))));
-		}
+	Queries.Reserve(ProcessorQueries.Num());
+	for (const FMassEntityQuery* Query : ProcessorQueries)
+	{
+		check(Query);
+		Queries.Add(MakeShareable(new FMassDebuggerQueryData(*Query, LOCTEXT("MassEntityQueryLabel", "Query"))));
 	}
 #endif // WITH_MASSENTITY_DEBUG
 }
@@ -117,21 +121,18 @@ FMassDebuggerProcessorData::FMassDebuggerProcessorData(const UMassEntitySubsyste
 	SetProcessor(InProcessor);
 #if WITH_MASSENTITY_DEBUG
 	TConstArrayView<FMassEntityQuery*> ProcessorQueries = FMassDebugger::GetUpToDateProcessorQueries(EntitySubsystem, InProcessor);
-	if (ensureMsgf(ProcessorQueries.Num() >= 1, TEXT("We expect at least one query to be there - every processor should have their ProcessorRequirements registered in OwnedQueries")))
+
+	ProcessorRequirements = MakeShareable(new FMassDebuggerQueryData(InProcessor.GetProcessorRequirements(), LOCTEXT("MassProcessorRequirementsLabel", "Processor Requirements")));
+
+	Queries.Reserve(ProcessorQueries.Num());
+	for (const FMassEntityQuery* Query : ProcessorQueries)
 	{
-		ProcessorRequirements = MakeShareable(new FMassDebuggerQueryData(*ProcessorQueries[0], LOCTEXT("MassProcessorRequirementsLabel", "Processor Requirements")));
+		check(Query);
+		Queries.Add(MakeShareable(new FMassDebuggerQueryData(*Query, LOCTEXT("MassEntityQueryLabel", "Query"))));
 
-		Queries.Reserve(ProcessorQueries.Num() - 1);
-		for (int i = 1; i < ProcessorQueries.Num(); ++i)
+		for (const FMassArchetypeHandle& ArchetypeHandle : Query->GetArchetypes())
 		{
-			const FMassEntityQuery* Query = ProcessorQueries[i];
-			check(Query);
-			Queries.Add(MakeShareable(new FMassDebuggerQueryData(*Query, LOCTEXT("MassEntityQueryLabel", "Query"))));
-
-			for (const FMassArchetypeHandle& ArchetypeHandle : Query->GetArchetypes())
-			{
-				ValidArchetypes.Add(InTransientArchetypesMap.FindChecked(ArchetypeHandle));
-			}
+			ValidArchetypes.Add(InTransientArchetypesMap.FindChecked(ArchetypeHandle));
 		}
 	}
 #endif // WITH_MASSENTITY_DEBUG
diff --git a/Engine/Plugins/Runtime/MassEntity/Source/MassEntityDebugger/Private/MassDebuggerModel.h b/Engine/Plugins/Runtime/MassEntity/Source/MassEntityDebugger/Private/MassDebuggerModel.h
--- a/Engine/Plugins/Runtime/MassEntity/Source/MassEntityDebugger/Private/MassDebuggerModel.h
+++ b/Engine/Plugins/Runtime/MassEntity/Source/MassEntityDebugger/Private/MassDebuggerModel.h
@@ -44,6 +44,7 @@ enum class EMassDebuggerProcessingGraphNodeSelection : uint8
 struct FMassDebuggerQueryData
 {
 	FMassDebuggerQueryData(const FMassEntityQuery& Query, const FText& InLabel);
+	FMassDebuggerQueryData(const FMassSubsystemRequirements& SubsystemRequirements, const FText& InLabel);
 
 	FMassExecutionRequirements ExecutionRequirements;
 	FText Label;
