Switched FMassExecutionContext.Subsystems from TObjectPtr to raw pointers

diff --git a/Engine/Plugins/Runtime/MassEntity/Source/MassEntity/Private/MassExecutionContext.cpp b/Engine/Plugins/Runtime/MassEntity/Source/MassEntity/Private/MassExecutionContext.cpp
--- a/Engine/Plugins/Runtime/MassEntity/Source/MassEntity/Private/MassExecutionContext.cpp
+++ b/Engine/Plugins/Runtime/MassEntity/Source/MassEntity/Private/MassExecutionContext.cpp
@@ -4,8 +4,6 @@
 #include "MassArchetypeData.h"
 
 
-PRAGMA_DISABLE_OPTIMIZATION
-
 void FMassExecutionContext::FlushDeferred(UMassEntitySubsystem& EntitySystem) const
 {
 	if (bFlushDeferredCommands && DeferredCommandBuffer)
@@ -31,7 +29,7 @@ bool FMassExecutionContext::CacheSubsystem(const UWorld* World, const uint32 Sys
 		Subsystems.AddZeroed(Subsystems.Num() - SystemIndex + 1);
 	}
 
-	if (Subsystems[SystemIndex].Get())
+	if (Subsystems[SystemIndex])
 	{
 		return true;
 	}
@@ -139,6 +137,3 @@ void FMassExecutionContext::SetFragmentRequirements(const FMassFragmentRequireme
 		}
 	}
 }
-
-
-PRAGMA_ENABLE_OPTIMIZATION
\ No newline at end of file
diff --git a/Engine/Plugins/Runtime/MassEntity/Source/MassEntity/Public/MassExecutionContext.h b/Engine/Plugins/Runtime/MassEntity/Source/MassEntity/Public/MassExecutionContext.h
--- a/Engine/Plugins/Runtime/MassEntity/Source/MassEntity/Public/MassExecutionContext.h
+++ b/Engine/Plugins/Runtime/MassEntity/Source/MassEntity/Public/MassExecutionContext.h
@@ -36,7 +36,7 @@ private:
 
 	FMassExternalSubsystemBitSet ConstSubsystemsBitSet;
 	FMassExternalSubsystemBitSet MutableSubsystemsBitSet;
-	TArray<TObjectPtr<USubsystem>> Subsystems;
+	TArray<USubsystem*> Subsystems;
 	
 	// mz@todo make this shared ptr thread-safe and never auto-flush in MT environment. 
 	TSharedPtr<FMassCommandBuffer> DeferredCommandBuffer;
@@ -380,7 +380,7 @@ protected:
 			Subsystems.AddZeroed(Subsystems.Num() - SystemIndex + 1);
 		}
 
-		T* SystemInstance = (T*)Subsystems[SystemIndex].Get();
+		T* SystemInstance = (T*)Subsystems[SystemIndex];
 		if (SystemInstance == nullptr)
 		{
 			SystemInstance = FMassExternalSubsystemTraits::GetInstance<typename TRemoveConst<T>::Type>(World);
@@ -397,7 +397,7 @@ protected:
 			Subsystems.AddZeroed(Subsystems.Num() - SystemIndex + 1);
 		}
 
-		T* SystemInstance = (T*)Subsystems[SystemIndex].Get();
+		T* SystemInstance = (T*)Subsystems[SystemIndex];
 		if (SystemInstance == nullptr)
 		{
 			SystemInstance = FMassExternalSubsystemTraits::GetInstance<typename TRemoveConst<T>::Type>(World, SubsystemClass);
