Cut FMassEntityQuery::ParallelForEachEntityChunk as defunct and no loner being used

diff --git a/Engine/Plugins/Runtime/MassEntity/Source/MassEntity/Private/MassEntityQuery.cpp b/Engine/Plugins/Runtime/MassEntity/Source/MassEntity/Private/MassEntityQuery.cpp
--- a/Engine/Plugins/Runtime/MassEntity/Source/MassEntity/Private/MassEntityQuery.cpp
+++ b/Engine/Plugins/Runtime/MassEntity/Source/MassEntity/Private/MassEntityQuery.cpp
@@ -210,70 +210,6 @@ void FMassEntityQuery::ForEachEntityChunk(UMassEntitySubsystem& EntitySubsystem,
 	ExecutionContext.FlushDeferred(EntitySubsystem);
 }
 
-void FMassEntityQuery::ParallelForEachEntityChunk(UMassEntitySubsystem& EntitySubsystem, FMassExecutionContext& ExecutionContext, const FMassExecuteFunction& ExecuteFunction)
-{
-	if (bAllowParallelExecution == false)
-	{
-		ForEachEntityChunk(EntitySubsystem, ExecutionContext, ExecuteFunction);
-		return;
-	}
-
-	struct FChunkJob
-	{
-		FMassArchetypeData& Archetype;
-		const int32 ArchetypeIndex;
-		const FMassArchetypeEntityCollection::FArchetypeEntityRange EntityRange;
-	};
-	TArray<FChunkJob> Jobs;
-
-	// if there's a chunk collection set by the external code - use that
-	if (ExecutionContext.GetEntityCollection().IsSet())
-	{
-		// verify the archetype matches requirements
-		FMassArchetypeHandle ArchetypeHandle = ExecutionContext.GetEntityCollection().GetArchetype();
-		if (DoesArchetypeMatchRequirements(ArchetypeHandle) == false)
-		{
-			UE_VLOG_UELOG(&EntitySubsystem, LogMass, Log, TEXT("Attempted to execute FMassEntityQuery with an incompatible Archetype: %s")
-				, *FMassDebugger::GetArchetypeRequirementCompatibilityDescription(*this, ExecutionContext.GetEntityCollection().GetArchetype()));
-			return;
-		}
-
-		ExecutionContext.SetFragmentRequirements(*this);
-		FMassArchetypeData& ArchetypeRef = FMassArchetypeHelper::ArchetypeDataFromHandleChecked(ArchetypeHandle);
-		const FMassArchetypeEntityCollection AsEntityCollection(ArchetypeHandle);
-		for (const FMassArchetypeEntityCollection::FArchetypeEntityRange& EntityRange : AsEntityCollection.GetRanges())
-		{
-			Jobs.Add({ ArchetypeRef, INDEX_NONE, EntityRange });
-		}
-	}
-	else
-	{
-		CacheArchetypes(EntitySubsystem);
-		ExecutionContext.SetFragmentRequirements(*this);
-		for (int ArchetypeIndex = 0; ArchetypeIndex < ValidArchetypes.Num(); ++ArchetypeIndex)
-		{
-			FMassArchetypeHandle& ArchetypeHandle = ValidArchetypes[ArchetypeIndex];
-			FMassArchetypeData& ArchetypeRef = FMassArchetypeHelper::ArchetypeDataFromHandleChecked(ArchetypeHandle);
-			const FMassArchetypeEntityCollection AsEntityCollection(ArchetypeHandle);
-			for (const FMassArchetypeEntityCollection::FArchetypeEntityRange& EntityRange : AsEntityCollection.GetRanges())
-			{
-				Jobs.Add({ArchetypeRef, ArchetypeIndex, EntityRange});
-			}
-		}
-	}
-	// ExecutionContext passed by copy on purpose
-	ParallelFor(Jobs.Num(), [this, ExecutionContext, &ExecuteFunction, Jobs](const int32 JobIndex)
-	{
-		Jobs[JobIndex].Archetype.ExecutionFunctionForChunk(ExecutionContext, ExecuteFunction
-			, Jobs[JobIndex].ArchetypeIndex != INDEX_NONE ? ArchetypeFragmentMapping[Jobs[JobIndex].ArchetypeIndex] : FMassQueryRequirementIndicesMapping()
-			, Jobs[JobIndex].EntityRange
-			, ChunkCondition);
-	});
-
-	ExecutionContext.ClearExecutionData();
-	ExecutionContext.FlushDeferred(EntitySubsystem);
-}
-
 int32 FMassEntityQuery::GetNumMatchingEntities(UMassEntitySubsystem& InEntitySubsystem)
 {
 	CacheArchetypes(InEntitySubsystem);
diff --git a/Engine/Plugins/Runtime/MassEntity/Source/MassEntity/Public/MassEntityQuery.h b/Engine/Plugins/Runtime/MassEntity/Source/MassEntity/Public/MassEntityQuery.h
--- a/Engine/Plugins/Runtime/MassEntity/Source/MassEntity/Public/MassEntityQuery.h
+++ b/Engine/Plugins/Runtime/MassEntity/Source/MassEntity/Public/MassEntityQuery.h
@@ -41,11 +41,6 @@ public:
 	/** Will first verify that the archetype given with Collection matches the query's requirements, and if so will run the other, more generic ForEachEntityChunk implementation */
 	void ForEachEntityChunk(const FMassArchetypeEntityCollection& Collection, UMassEntitySubsystem& EntitySubsystem, FMassExecutionContext& ExecutionContext, const FMassExecuteFunction& ExecuteFunction);
 
-	/**
-	 * Attempts to process every chunk of every affected archetype in parallel.
-	 */
-	void ParallelForEachEntityChunk(UMassEntitySubsystem& EntitySubsystem, FMassExecutionContext& ExecutionContext, const FMassExecuteFunction& ExecuteFunction);
-
 	/** Will gather all archetypes from InEntitySubsystem matching this->Requirements.
 	 *  Note that no work will be done if the cached data is up to date (as tracked by EntitySubsystemHash and 
 	 *	ArchetypeDataVersion properties). */
