Switched MassEntity from supporting interaction with USystems over to UWorldSubsystems.
The main reason for the change is to avoid non-UWorldSubsystems special cases we don't utilize nor test at this time. The change is reversible once we have a non-UWorldSubsystem use case

diff --git a/Engine/Plugins/Runtime/MassEntity/Source/MassEntity/Private/MassExecutionContext.cpp b/Engine/Plugins/Runtime/MassEntity/Source/MassEntity/Private/MassExecutionContext.cpp
--- a/Engine/Plugins/Runtime/MassEntity/Source/MassEntity/Private/MassExecutionContext.cpp
+++ b/Engine/Plugins/Runtime/MassEntity/Source/MassEntity/Private/MassExecutionContext.cpp
@@ -35,22 +35,14 @@ bool FMassExecutionContext::CacheSubsystem(const UWorld* World, const uint32 Sys
 	}
 
 	const UClass* SubsystemClass = FMassExternalSubsystemBitSet::GetTypeAtIndex(SystemIndex);
-	if (SubsystemClass)
-	{
-		const TSubclassOf<UWorldSubsystem> WorldSubsystemClass(const_cast<UClass*>(SubsystemClass));
-		if (*WorldSubsystemClass)
-		{
-			USubsystem* SystemInstance = FMassExternalSubsystemTraits::GetInstance<UWorldSubsystem>(World, WorldSubsystemClass);
-			Subsystems[SystemIndex] = SystemInstance;
-			return SystemInstance != nullptr;
-		}
-		else
-		{
-			// not a UWorldSubsystem. There's no way to generically fetch an instance of one so we assume it's there
-			return true;
-		}
-	}
-	return false;
+	checkSlow(SubsystemClass);
+	
+	const TSubclassOf<UWorldSubsystem> WorldSubsystemClass(const_cast<UClass*>(SubsystemClass));
+	checkSlow(*WorldSubsystemClass);
+		
+	UWorldSubsystem* SystemInstance = FMassExternalSubsystemTraits::GetInstance<UWorldSubsystem>(World, WorldSubsystemClass);
+	Subsystems[SystemIndex] = SystemInstance;
+	return SystemInstance != nullptr;
 }
 
 bool FMassExecutionContext::CacheSubsystemRequirements(const UWorld* World, const FMassSubsystemRequirements& SubsystemRequirements)
diff --git a/Engine/Plugins/Runtime/MassEntity/Source/MassEntity/Public/MassEntityTypes.h b/Engine/Plugins/Runtime/MassEntity/Source/MassEntity/Public/MassEntityTypes.h
--- a/Engine/Plugins/Runtime/MassEntity/Source/MassEntity/Public/MassEntityTypes.h
+++ b/Engine/Plugins/Runtime/MassEntity/Source/MassEntity/Public/MassEntityTypes.h
@@ -108,7 +108,7 @@ DECLARE_STRUCTTYPEBITSET_EXPORTED(MASSENTITY_API, FMassFragmentBitSet, FMassFrag
 DECLARE_STRUCTTYPEBITSET_EXPORTED(MASSENTITY_API, FMassTagBitSet, FMassTag);
 DECLARE_STRUCTTYPEBITSET_EXPORTED(MASSENTITY_API, FMassChunkFragmentBitSet, FMassChunkFragment);
 DECLARE_STRUCTTYPEBITSET_EXPORTED(MASSENTITY_API, FMassSharedFragmentBitSet, FMassSharedFragment);
-DECLARE_CLASSTYPEBITSET_EXPORTED(MASSENTITY_API, FMassExternalSubsystemBitSet, USubsystem);
+DECLARE_CLASSTYPEBITSET_EXPORTED(MASSENTITY_API, FMassExternalSubsystemBitSet, UWorldSubsystem);
 
 /** The type summarily describing a composition of an entity or an archetype. It contains information on both the
  *  fragments as well as tags */
diff --git a/Engine/Plugins/Runtime/MassEntity/Source/MassEntity/Public/MassExecutionContext.h b/Engine/Plugins/Runtime/MassEntity/Source/MassEntity/Public/MassExecutionContext.h
--- a/Engine/Plugins/Runtime/MassEntity/Source/MassEntity/Public/MassExecutionContext.h
+++ b/Engine/Plugins/Runtime/MassEntity/Source/MassEntity/Public/MassExecutionContext.h
@@ -36,7 +36,7 @@ private:
 
 	FMassExternalSubsystemBitSet ConstSubsystemsBitSet;
 	FMassExternalSubsystemBitSet MutableSubsystemsBitSet;
-	TArray<USubsystem*> Subsystems;
+	TArray<UWorldSubsystem*> Subsystems;
 	
 	// mz@todo make this shared ptr thread-safe and never auto-flush in MT environment. 
 	TSharedPtr<FMassCommandBuffer> DeferredCommandBuffer;
diff --git a/Engine/Plugins/Runtime/MassEntity/Source/MassEntity/Public/MassExternalSubsystemTraits.h b/Engine/Plugins/Runtime/MassEntity/Source/MassEntity/Public/MassExternalSubsystemTraits.h
--- a/Engine/Plugins/Runtime/MassEntity/Source/MassEntity/Public/MassExternalSubsystemTraits.h
+++ b/Engine/Plugins/Runtime/MassEntity/Source/MassEntity/Public/MassExternalSubsystemTraits.h
@@ -27,7 +27,7 @@ struct TMassExternalSubsystemTraits final
 	{
 		// Note that we're not supplying a default value for this property to be able to statically catch code that tries
 		// to access given subsystem without including the appropriate headers. See the comment above if you want to use
-		// a USubsystem that has not had its traits defined before. 
+		// a UWorldSubsystem that has not had its traits defined before. 
 		// GameThreadOnly = true,
 		ThreadSafeRead = false,
 		ThreadSafeWrite = false,
diff --git a/Engine/Plugins/Runtime/MassEntity/Source/MassEntity/Public/MassRequirements.h b/Engine/Plugins/Runtime/MassEntity/Source/MassEntity/Public/MassRequirements.h
--- a/Engine/Plugins/Runtime/MassEntity/Source/MassEntity/Public/MassRequirements.h
+++ b/Engine/Plugins/Runtime/MassEntity/Source/MassEntity/Public/MassRequirements.h
@@ -111,7 +111,7 @@ public:
 		return *this;
 	}
 
-	FMassSubsystemRequirements& AddSubsystemRequirement(const TSubclassOf<USubsystem> SubsystemClass, const EMassFragmentAccess AccessMode, const bool bGameThreadOnly = true)
+	FMassSubsystemRequirements& AddSubsystemRequirement(const TSubclassOf<UWorldSubsystem> SubsystemClass, const EMassFragmentAccess AccessMode, const bool bGameThreadOnly = true)
 	{
 		check(AccessMode != EMassFragmentAccess::None && AccessMode != EMassFragmentAccess::MAX);
 
