Compilation Issue 196162 - fixed missing FMassDebugger funcitons used for vislogging in cooked editor

diff --git a/Engine/Plugins/Runtime/MassEntity/Source/MassEntity/Public/MassDebugger.h b/Engine/Plugins/Runtime/MassEntity/Source/MassEntity/Public/MassDebugger.h
--- a/Engine/Plugins/Runtime/MassEntity/Source/MassEntity/Public/MassDebugger.h
+++ b/Engine/Plugins/Runtime/MassEntity/Source/MassEntity/Public/MassDebugger.h
@@ -4,21 +4,21 @@
 
 #include "MassProcessingTypes.h"
 
-#if WITH_MASSENTITY_DEBUG
-#include "Containers/ContainersFwd.h"
-#include "MassEntityQuery.h"
-#include "MassProcessor.h"
-
-
 class FOutputDevice;
 class UMassProcessor;
 struct FMassEntityQuery;
 class UMassEntitySubsystem;
 struct FMassArchetypeHandle;
+struct FMassFragmentRequirements;
 struct FMassFragmentRequirementDescription;
 enum class EMassFragmentAccess : uint8;
 enum class EMassFragmentPresence : uint8;
 
+#if WITH_MASSENTITY_DEBUG
+#include "Containers/ContainersFwd.h"
+#include "MassEntityQuery.h"
+#include "MassProcessor.h"
+
 namespace UE::Mass::Debug
 {
 struct MASSENTITY_API FQueryRequirementsView
@@ -65,4 +65,13 @@ struct MASSENTITY_API FMassDebugger
 	static void OutputEntityDescription(FOutputDevice& Ar, const UMassEntitySubsystem& EntitySubsystem, const FMassEntityHandle Entity, const TCHAR* InPrefix = TEXT(""));
 };
 
+#else
+
+struct MASSENTITY_API FMassDebugger
+{
+	static FString GetSingleRequirementDescription(const FMassFragmentRequirementDescription& Requirement) { return TEXT("[no debug information]"); }
+	static FString GetRequirementsDescription(const FMassFragmentRequirements& Requirements) { return TEXT("[no debug information]"); }
+	static FString GetArchetypeRequirementCompatibilityDescription(const FMassFragmentRequirements& Requirements, const FMassArchetypeHandle& ArchetypeHandle) { return TEXT("[no debug information]"); }
+};
+
 #endif // WITH_MASSENTITY_DEBUG
\ No newline at end of file
