Issue 196676, reordered headers and forward declarations in MassDebugger.h to avoid the warning

diff --git a/Engine/Plugins/Runtime/MassEntity/Source/MassEntity/Public/MassDebugger.h b/Engine/Plugins/Runtime/MassEntity/Source/MassEntity/Public/MassDebugger.h
--- a/Engine/Plugins/Runtime/MassEntity/Source/MassEntity/Public/MassDebugger.h
+++ b/Engine/Plugins/Runtime/MassEntity/Source/MassEntity/Public/MassDebugger.h
@@ -3,6 +3,11 @@
 #pragma once
 
 #include "MassProcessingTypes.h"
+#if WITH_MASSENTITY_DEBUG
+#include "Containers/ContainersFwd.h"
+#include "MassEntityQuery.h"
+#include "MassProcessor.h"
+#endif // WITH_MASSENTITY_DEBUG
 
 class FOutputDevice;
 class UMassProcessor;
@@ -15,10 +20,6 @@ enum class EMassFragmentAccess : uint8;
 enum class EMassFragmentPresence : uint8;
 
 #if WITH_MASSENTITY_DEBUG
-#include "Containers/ContainersFwd.h"
-#include "MassEntityQuery.h"
-#include "MassProcessor.h"
-
 namespace UE::Mass::Debug
 {
 struct MASSENTITY_API FQueryRequirementsView
