Fixed phase processor's display name in MassDebugger (via properly picking the composite processors' group name)

diff --git a/Engine/Plugins/Runtime/MassEntity/Source/MassEntity/Private/MassEntitySettings.cpp b/Engine/Plugins/Runtime/MassEntity/Source/MassEntity/Private/MassEntitySettings.cpp
--- a/Engine/Plugins/Runtime/MassEntity/Source/MassEntity/Private/MassEntitySettings.cpp
+++ b/Engine/Plugins/Runtime/MassEntity/Source/MassEntity/Private/MassEntitySettings.cpp
@@ -17,7 +17,7 @@ UMassEntitySettings::UMassEntitySettings(const FObjectInitializer& ObjectInitial
 {
 	for (int i = 0; i < (int)EMassProcessingPhase::MAX; ++i)
 	{
-		ProcessingPhasesConfig[i].PhaseName = *UEnum::GetValueAsString(EMassProcessingPhase(i));
+		ProcessingPhasesConfig[i].PhaseName = *UEnum::GetDisplayValueAsText(EMassProcessingPhase(i)).ToString();
 	}
 
 	FCoreDelegates::OnPostEngineInit.AddUObject(this, &UMassEntitySettings::BuildProcessorListAndPhases);
diff --git a/Engine/Plugins/Runtime/MassEntity/Source/MassEntity/Private/MassProcessingPhase.cpp b/Engine/Plugins/Runtime/MassEntity/Source/MassEntity/Private/MassProcessingPhase.cpp
--- a/Engine/Plugins/Runtime/MassEntity/Source/MassEntity/Private/MassProcessingPhase.cpp
+++ b/Engine/Plugins/Runtime/MassEntity/Source/MassEntity/Private/MassProcessingPhase.cpp
@@ -278,7 +278,7 @@ void UMassProcessingPhaseManager::SetPhaseProcessor(const EMassProcessingPhase P
 	{
 		REDIRECT_OBJECT_TO_VLOG(PhaseProcessor, this);
 		PhaseProcessor->SetProcessingPhase(Phase);
-		PhaseProcessor->SetGroupName(FName(FString::Printf(TEXT("%s-%s"), *PhaseProcessor->GetName(), *UEnum::GetValueAsString(Phase))));
+		PhaseProcessor->SetGroupName(FName(FString::Printf(TEXT("%s Group"), *UEnum::GetDisplayValueAsText(Phase).ToString())));
 		
 		check(GetOuter());
 		PhaseProcessor->Initialize(*GetOuter());
@@ -304,7 +304,7 @@ void UMassProcessingPhaseManager::CreatePhases()
 		ProcessingPhases[i].Phase = EMassProcessingPhase(i);
 		ProcessingPhases[i].TickGroup = UE::Mass::Private::PhaseToTickingGroup[i];
 		UMassCompositeProcessor* PhaseProcessor = NewObject<UMassCompositeProcessor>(this, UMassCompositeProcessor::StaticClass()
-			, *FString::Printf(TEXT("ProcessingPhase_%s"), *UEnum::GetValueAsString(EMassProcessingPhase(i))));
+			, *FString::Printf(TEXT("ProcessingPhase_%s"), *UEnum::GetDisplayValueAsText(EMassProcessingPhase(i)).ToString()));
 		SetPhaseProcessor(EMassProcessingPhase(i), PhaseProcessor);
 	}
 }
