MassDebugger: Changed view division from horiz box to splitter, color tweaks.

diff --git a/Engine/Plugins/Runtime/MassEntity/Source/MassEntityDebugger/Private/SMassProcessingGraphView.cpp b/Engine/Plugins/Runtime/MassEntity/Source/MassEntityDebugger/Private/SMassProcessingGraphView.cpp
--- a/Engine/Plugins/Runtime/MassEntity/Source/MassEntityDebugger/Private/SMassProcessingGraphView.cpp
+++ b/Engine/Plugins/Runtime/MassEntity/Source/MassEntityDebugger/Private/SMassProcessingGraphView.cpp
@@ -7,6 +7,7 @@
 #include "Styling/SlateBrush.h"
 #include "MassDebuggerStyle.h"
 #include "Widgets/Views/STreeView.h"
+#include "Styling/StyleColors.h"
 
 #define LOCTEXT_NAMESPACE "SMassDebugger"
 
@@ -38,7 +39,7 @@ public:
 		STableRow<FMassDebuggerProcessingGraphNodeTreeItemPtr>::ConstructInternal(STableRow::FArguments()
 			.Padding(5.0f)
 			, InOwnerTableView.ToSharedRef());
-	
+
 		ChildSlot
 			[
 				SNew(SHorizontalBox)
@@ -63,16 +64,18 @@ public:
 						.Text(Item->Node.GetLabel())
 						.ColorAndOpacity_Lambda([this]()
 							{
+								const FLinearColor Foreground = FStyleColors::Foreground.GetSpecifiedColor();
+							
 								switch (Item->Node.GraphNodeSelection)
 								{
-									case EMassDebuggerProcessingGraphNodeSelection::None: // fall through on purpose
+									case EMassDebuggerProcessingGraphNodeSelection::None:
 										return FLinearColor::White;
 										break;
 									case EMassDebuggerProcessingGraphNodeSelection::WaitFor:
-										return FLinearColor::Green;
+										return FMath::Lerp(Foreground, FLinearColor::Green, 0.75f);
 										break;
 									case EMassDebuggerProcessingGraphNodeSelection::Block:
-										return FLinearColor::Red;
+										return FMath::Lerp(Foreground, FLinearColor::Red, 0.75f);
 										break;
 									default:
 										return FLinearColor::White;
diff --git a/Engine/Plugins/Runtime/MassEntity/Source/MassEntityDebugger/Private/SMassProcessingView.cpp b/Engine/Plugins/Runtime/MassEntity/Source/MassEntityDebugger/Private/SMassProcessingView.cpp
--- a/Engine/Plugins/Runtime/MassEntity/Source/MassEntityDebugger/Private/SMassProcessingView.cpp
+++ b/Engine/Plugins/Runtime/MassEntity/Source/MassEntityDebugger/Private/SMassProcessingView.cpp
@@ -48,9 +48,12 @@ void SMassProcessingView::Construct(const FArguments& InArgs, TSharedRef<FMassDe
 
 	ChildSlot
 	[
-		SNew(SHorizontalBox)
-		+ SHorizontalBox::Slot()
-		.AutoWidth()
+		SNew(SSplitter)
+		.Orientation(Orient_Horizontal)
+	 
+		+ SSplitter::Slot()
+		.Value(.3f)
+		.MinSize(250.0f)
 		[
 			SAssignNew(GraphsListWidget, SListView<TSharedPtr<FMassDebuggerProcessingGraph>>)
 			.ListItemsSource(&DebuggerModel->CachedProcessingGraphs)
@@ -64,7 +67,9 @@ void SMassProcessingView::Construct(const FArguments& InArgs, TSharedRef<FMassDe
 					ProcessingGraphWidget->Display(SelectedItem);
 				})
 		]
-		+ SHorizontalBox::Slot()
+		
+		+ SSplitter::Slot()
+		.Value(.7f)
 		[
 			SAssignNew(ProcessingGraphWidget, SMassProcessingGraphView, InDebuggerModel)
 		]
diff --git a/Engine/Plugins/Runtime/MassEntity/Source/MassEntityDebugger/Private/SMassProcessorsView.cpp b/Engine/Plugins/Runtime/MassEntity/Source/MassEntityDebugger/Private/SMassProcessorsView.cpp
--- a/Engine/Plugins/Runtime/MassEntity/Source/MassEntityDebugger/Private/SMassProcessorsView.cpp
+++ b/Engine/Plugins/Runtime/MassEntity/Source/MassEntityDebugger/Private/SMassProcessorsView.cpp
@@ -90,9 +90,12 @@ void SMassProcessorsView::Construct(const FArguments& InArgs, TSharedRef<FMassDe
 
 	ChildSlot
 		[
-			SNew(SHorizontalBox)
-			+ SHorizontalBox::Slot()
-			.FillWidth(0.3)
+			SNew(SSplitter)
+			.Orientation(Orient_Horizontal)
+			 
+			+ SSplitter::Slot()
+			.Value(.3f)
+			.MinSize(250.0f)
 			[
 				SAssignNew(ProcessorsListWidget, SListView<TSharedPtr<FMassDebuggerProcessorData> >)
 					.ItemHeight(20)
@@ -104,8 +107,9 @@ void SMassProcessorsView::Construct(const FArguments& InArgs, TSharedRef<FMassDe
 							return SNew(SMassProcessorListTableRow, OwnerTable, Item);
 						})
 			]
-			+ SHorizontalBox::Slot()
-			.FillWidth(0.7)
+			
+			+ SSplitter::Slot()
+			.Value(.7f)
 			[
 				SNew(SScrollBox)
 				.Orientation(Orient_Vertical)
