Mass: Add debug name to archetypes.
- Allow to specify debug archetype name when creating archetypes
- Store unique debug names per archetype (archetypes of same compositon are shared)

diff --git a/Engine/Plugins/Runtime/MassEntity/Source/MassEntity/Private/MassArchetypeData.cpp b/Engine/Plugins/Runtime/MassEntity/Source/MassEntity/Private/MassArchetypeData.cpp
--- a/Engine/Plugins/Runtime/MassEntity/Source/MassEntity/Private/MassArchetypeData.cpp
+++ b/Engine/Plugins/Runtime/MassEntity/Source/MassEntity/Private/MassArchetypeData.cpp
@@ -741,6 +741,12 @@ FString FMassArchetypeData::DebugGetDescription() const
 #if WITH_MASSENTITY_DEBUG
 	FStringOutputDevice OutDescription;
 
+	if (!DebugNames.IsEmpty())
+	{
+		OutDescription += TEXT("Name: ");
+		OutDescription += GetCombinedDebugNamesAsString();
+		OutDescription += TEXT("\n");
+	}
 	OutDescription += TEXT("Chunk fragments: ");
 	CompositionDescriptor.ChunkFragments.DebugGetStringDesc(OutDescription);
 	OutDescription += TEXT("\nTags: ");
@@ -754,9 +760,25 @@ FString FMassArchetypeData::DebugGetDescription() const
 #endif
 }
 
+FString FMassArchetypeData::GetCombinedDebugNamesAsString() const
+{
+	TStringBuilder<256> StringBuilder;
+	for (int i = 0; i < DebugNames.Num(); i++)
+	{
+		if (i > 0)
+		{
+			StringBuilder.Append(TEXT(", "));;
+		}
+		StringBuilder.Append(DebugNames[i].ToString());
+	}
+	return StringBuilder.ToString();
+}
+
 #if WITH_MASSENTITY_DEBUG
 void FMassArchetypeData::DebugPrintArchetype(FOutputDevice& Ar)
 {
+	Ar.Logf(ELogVerbosity::Log, TEXT("Name: %s"), *GetCombinedDebugNamesAsString());
+
 	FStringOutputDevice TagsDecription;
 	CompositionDescriptor.Tags.DebugGetStringDesc(TagsDecription);
 	Ar.Logf(ELogVerbosity::Log, TEXT("Tags: %s"), *TagsDecription);
diff --git a/Engine/Plugins/Runtime/MassEntity/Source/MassEntity/Private/MassArchetypeData.h b/Engine/Plugins/Runtime/MassEntity/Source/MassEntity/Private/MassArchetypeData.h
--- a/Engine/Plugins/Runtime/MassEntity/Source/MassEntity/Private/MassArchetypeData.h
+++ b/Engine/Plugins/Runtime/MassEntity/Source/MassEntity/Private/MassArchetypeData.h
@@ -167,6 +167,9 @@ private:
 	int32 TotalBytesPerEntity;
 	int32 EntityListOffsetWithinChunk;
 
+	// Arrays of names the archetype is referred as.
+	TArray<FName> DebugNames;
+	
 	friend FMassEntityQuery;
 	friend FMassArchetypeEntityCollection;
 	friend FMassDebugger;
@@ -260,6 +263,18 @@ public:
 	// Converts the list of fragments into a user-readable debug string
 	FString DebugGetDescription() const;
 
+	/** Adds new debug name associated with the archetype. */
+	void AddUniqueDebugName(const FName& Name) { DebugNames.AddUnique(Name); }
+	
+	/** @return array of debug names associated with this archetype. */
+	const TConstArrayView<FName> GetDebugNames() const { return DebugNames; }
+	
+	/** Copies debug names from another archetype data. */
+	void CopyDebugNamesFrom(const FMassArchetypeData& Other) { DebugNames = Other.DebugNames; }
+	
+	/** @return string of all debug names combined */
+	FString GetCombinedDebugNamesAsString() const;
+
 #if WITH_MASSENTITY_DEBUG
 	/**
 	 * Prints out debug information about the archetype
diff --git a/Engine/Plugins/Runtime/MassEntity/Source/MassEntity/Private/MassDebugger.cpp b/Engine/Plugins/Runtime/MassEntity/Source/MassEntity/Private/MassDebugger.cpp
--- a/Engine/Plugins/Runtime/MassEntity/Source/MassEntity/Private/MassDebugger.cpp
+++ b/Engine/Plugins/Runtime/MassEntity/Source/MassEntity/Private/MassDebugger.cpp
@@ -268,6 +268,12 @@ void FMassDebugger::GetArchetypeEntityStats(const FMassArchetypeHandle& Archetyp
 	OutChunksCount = ArchetypeData.GetChunkCount();
 }
 
+const TConstArrayView<FName> FMassDebugger::GetArchetypeDebugNames(const FMassArchetypeHandle& ArchetypeHandle)
+{
+	const FMassArchetypeData& ArchetypeData = FMassArchetypeHelper::ArchetypeDataFromHandleChecked(ArchetypeHandle);
+	return ArchetypeData.GetDebugNames();
+}
+
 TConstArrayView<struct UMassCompositeProcessor::FDependencyNode> FMassDebugger::GetProcessingGraph(const UMassCompositeProcessor& GraphOwner)
 {
 	return GraphOwner.ProcessingFlatGraph;
diff --git a/Engine/Plugins/Runtime/MassEntity/Source/MassEntity/Private/MassEntitySubsystem.cpp b/Engine/Plugins/Runtime/MassEntity/Source/MassEntity/Private/MassEntitySubsystem.cpp
--- a/Engine/Plugins/Runtime/MassEntity/Source/MassEntity/Private/MassEntitySubsystem.cpp
+++ b/Engine/Plugins/Runtime/MassEntity/Source/MassEntity/Private/MassEntitySubsystem.cpp
@@ -126,7 +126,7 @@ void UMassEntitySubsystem::Deinitialize()
 	DeferredCommandBuffer->CleanUp();
 }
 
-FMassArchetypeHandle UMassEntitySubsystem::CreateArchetype(TConstArrayView<const UScriptStruct*> FragmentsAndTagsList)
+FMassArchetypeHandle UMassEntitySubsystem::CreateArchetype(TConstArrayView<const UScriptStruct*> FragmentsAndTagsList, const FName ArchetypeDebugName)
 {
 	FMassChunkFragmentBitSet ChunkFragments;
 	FMassTagBitSet Tags;
@@ -155,19 +155,19 @@ FMassArchetypeHandle UMassEntitySubsystem::CreateArchetype(TConstArrayView<const
 	}
 
 	const FMassArchetypeCompositionDescriptor Composition(FMassFragmentBitSet(FragmentList), Tags, ChunkFragments, FMassSharedFragmentBitSet());
-	return CreateArchetype(Composition, FMassArchetypeSharedFragmentValues());
+	return CreateArchetype(Composition, FMassArchetypeSharedFragmentValues(), ArchetypeDebugName);
 }
 
-FMassArchetypeHandle UMassEntitySubsystem::CreateArchetype(const TSharedPtr<FMassArchetypeData>& SourceArchetype, const FMassFragmentBitSet& AddedFragments)
+FMassArchetypeHandle UMassEntitySubsystem::CreateArchetype(const TSharedPtr<FMassArchetypeData>& SourceArchetype, const FMassFragmentBitSet& AddedFragments, const FName ArchetypeDebugName)
 {
 	check(SourceArchetype.IsValid());
 	checkf(AddedFragments.IsEmpty() == false, TEXT("%s Adding an empty fragment list to an archetype is not supported."), ANSI_TO_TCHAR(__FUNCTION__));
 
 	const FMassArchetypeCompositionDescriptor Composition(AddedFragments + SourceArchetype->GetFragmentBitSet(), SourceArchetype->GetTagBitSet(), SourceArchetype->GetChunkFragmentBitSet(), SourceArchetype->GetSharedFragmentBitSet());
-	return CreateArchetype(Composition, SourceArchetype->GetSharedFragmentValues());
+	return CreateArchetype(Composition, SourceArchetype->GetSharedFragmentValues(), ArchetypeDebugName);
 }
 
-FMassArchetypeHandle UMassEntitySubsystem::CreateArchetype(const FMassArchetypeCompositionDescriptor& Composition, const FMassArchetypeSharedFragmentValues& SharedFragmentValues)
+FMassArchetypeHandle UMassEntitySubsystem::CreateArchetype(const FMassArchetypeCompositionDescriptor& Composition, const FMassArchetypeSharedFragmentValues& SharedFragmentValues, const FName ArchetypeDebugName)
 {
 	const uint32 TypeHash = HashCombine(Composition.CalculateHash(), GetTypeHash(SharedFragmentValues));
 
@@ -178,6 +178,11 @@ FMassArchetypeHandle UMassEntitySubsystem::CreateArchetype(const FMassArchetypeC
 	{
 		if (Ptr->IsEquivalent(Composition, SharedFragmentValues))
 		{
+			// Keep track of all names for this archetype.
+			if (!ArchetypeDebugName.IsNone())
+			{
+				Ptr->AddUniqueDebugName(ArchetypeDebugName);
+			}
 			ArchetypeDataPtr = Ptr;
 			break;
 		}
@@ -188,7 +193,10 @@ FMassArchetypeHandle UMassEntitySubsystem::CreateArchetype(const FMassArchetypeC
 		// Create a new archetype
 		FMassArchetypeData* NewArchetype = new FMassArchetypeData();
 		NewArchetype->Initialize(Composition, SharedFragmentValues);
-		
+		if (!ArchetypeDebugName.IsNone())
+		{
+			NewArchetype->AddUniqueDebugName(ArchetypeDebugName);
+		}
 		ArchetypeDataPtr = HashRow.Add_GetRef(MakeShareable(NewArchetype));
 
 		for (const FMassArchetypeFragmentConfig& FragmentConfig : NewArchetype->GetFragmentConfigs())
@@ -240,6 +248,7 @@ FMassArchetypeHandle UMassEntitySubsystem::InternalCreateSimilarArchetype(const
 		// Create a new archetype
 		FMassArchetypeData* NewArchetype = new FMassArchetypeData();
 		NewArchetype->InitializeWithSimilar(SourceArchetypeRef, MoveTemp(NewComposition));
+		NewArchetype->CopyDebugNamesFrom(SourceArchetypeRef);
 
 		ArchetypeDataPtr = HashRow.Add_GetRef(MakeShareable(NewArchetype));
 
@@ -309,11 +318,12 @@ FMassEntityHandle UMassEntitySubsystem::CreateEntity(const FMassArchetypeHandle&
 	return Entity;
 }
 
-FMassEntityHandle UMassEntitySubsystem::CreateEntity(TConstArrayView<FInstancedStruct> FragmentInstanceList)
+FMassEntityHandle UMassEntitySubsystem::CreateEntity(TConstArrayView<FInstancedStruct> FragmentInstanceList, const FName ArchetypeDebugName)
 {
 	check(FragmentInstanceList.Num() > 0);
 
-	const FMassArchetypeHandle& ArchetypeHandle = CreateArchetype(FMassArchetypeCompositionDescriptor(FragmentInstanceList, FMassTagBitSet(), FMassChunkFragmentBitSet(), FMassSharedFragmentBitSet()), FMassArchetypeSharedFragmentValues());
+	const FMassArchetypeHandle& ArchetypeHandle = CreateArchetype(FMassArchetypeCompositionDescriptor(FragmentInstanceList,
+		FMassTagBitSet(), FMassChunkFragmentBitSet(), FMassSharedFragmentBitSet()), FMassArchetypeSharedFragmentValues(), ArchetypeDebugName);
 	check(ArchetypeHandle.IsValid());
 
 	const FMassEntityHandle Entity = ReserveEntity();
@@ -377,7 +387,7 @@ void UMassEntitySubsystem::BuildEntity(FMassEntityHandle Entity, TConstArrayView
 	EntityData.CurrentArchetype->SetFragmentsData(Entity, FragmentInstanceList);
 }
 
-void UMassEntitySubsystem::BatchBuildEntities(const FMassArchetypeEntityCollectionWithPayload& EncodedEntitiesWithPayload, const FMassFragmentBitSet& FragmentsAffected, const FMassArchetypeSharedFragmentValues& SharedFragmentValues)
+void UMassEntitySubsystem::BatchBuildEntities(const FMassArchetypeEntityCollectionWithPayload& EncodedEntitiesWithPayload, const FMassFragmentBitSet& FragmentsAffected, const FMassArchetypeSharedFragmentValues& SharedFragmentValues, const FName ArchetypeDebugName)
 {
 	check(SharedFragmentValues.IsSorted());
 
@@ -391,17 +401,17 @@ void UMassEntitySubsystem::BatchBuildEntities(const FMassArchetypeEntityCollecti
 		Composition.SharedFragments.Add(*SharedFragment.GetScriptStruct());
 	}
 
-	BatchBuildEntities(EncodedEntitiesWithPayload, MoveTemp(Composition), SharedFragmentValues);
+	BatchBuildEntities(EncodedEntitiesWithPayload, MoveTemp(Composition), SharedFragmentValues, ArchetypeDebugName);
 }
 
-void UMassEntitySubsystem::BatchBuildEntities(const FMassArchetypeEntityCollectionWithPayload& EncodedEntitiesWithPayload, FMassArchetypeCompositionDescriptor&& Composition, const FMassArchetypeSharedFragmentValues& SharedFragmentValues)
+void UMassEntitySubsystem::BatchBuildEntities(const FMassArchetypeEntityCollectionWithPayload& EncodedEntitiesWithPayload, FMassArchetypeCompositionDescriptor&& Composition, const FMassArchetypeSharedFragmentValues& SharedFragmentValues, const FName ArchetypeDebugName)
 {
 	TRACE_CPUPROFILER_EVENT_SCOPE(Mass_BatchBuildEntities);
 
 	FMassArchetypeEntityCollection::FEntityRangeArray TargetArchetypeEntityRanges;
 
 	// "built" entities case, this is verified during FMassArchetypeEntityCollectionWithPayload construction
-	FMassArchetypeHandle TargetArchetypeHandle = CreateArchetype(Composition, SharedFragmentValues);
+	FMassArchetypeHandle TargetArchetypeHandle = CreateArchetype(Composition, SharedFragmentValues, ArchetypeDebugName);
 	check(TargetArchetypeHandle.IsValid());
 
 	// there are some extra steps in creating EncodedEntities from the original given entity handles and then back
@@ -605,7 +615,7 @@ void UMassEntitySubsystem::AddCompositionToEntity_GetDelta(FMassEntityHandle Ent
 	{
 		FMassArchetypeCompositionDescriptor NewDescriptor = OldArchetype->GetCompositionDescriptor();
 		NewDescriptor.Fragments += InDescriptor.Fragments;
-		NewDescriptor.Tags+= InDescriptor.Tags;
+		NewDescriptor.Tags += InDescriptor.Tags;
 
 		const FMassArchetypeHandle NewArchetypeHandle = CreateArchetype(NewDescriptor, OldArchetype->GetSharedFragmentValues());
 
@@ -613,6 +623,7 @@ void UMassEntitySubsystem::AddCompositionToEntity_GetDelta(FMassEntityHandle Ent
 		{
 			// Move the entity over
 			FMassArchetypeData& NewArchetype = FMassArchetypeHelper::ArchetypeDataFromHandleChecked(NewArchetypeHandle);
+			NewArchetype.CopyDebugNamesFrom(*OldArchetype);
 			EntityData.CurrentArchetype->MoveEntityToAnotherArchetype(Entity, NewArchetype);
 			EntityData.CurrentArchetype = NewArchetypeHandle.DataPtr;
 
@@ -649,6 +660,7 @@ void UMassEntitySubsystem::RemoveCompositionFromEntity(FMassEntityHandle Entity,
 			{
 				// Move the entity over
 				FMassArchetypeData& NewArchetype = FMassArchetypeHelper::ArchetypeDataFromHandleChecked(NewArchetypeHandle);
+				NewArchetype.CopyDebugNamesFrom(*OldArchetype);
 				EntityData.CurrentArchetype->MoveEntityToAnotherArchetype(Entity, NewArchetype);
 				EntityData.CurrentArchetype = NewArchetypeHandle.DataPtr;
 			}
@@ -699,6 +711,7 @@ void UMassEntitySubsystem::InternalAddFragmentListToEntity(FMassEntityHandle Ent
 	check(Entities.IsValidIndex(Entity.Index));
 	FEntityData& EntityData = Entities[Entity.Index];
 	check(EntityData.CurrentArchetype.IsValid());
+	const FMassArchetypeData* OldArchetype = EntityData.CurrentArchetype.Get();
 
 	// fetch or create the new archetype
 	const FMassArchetypeHandle NewArchetypeHandle = CreateArchetype(EntityData.CurrentArchetype, InFragments);
@@ -707,6 +720,7 @@ void UMassEntitySubsystem::InternalAddFragmentListToEntity(FMassEntityHandle Ent
 	{
 		// Move the entity over
 		FMassArchetypeData& NewArchetype = FMassArchetypeHelper::ArchetypeDataFromHandleChecked(NewArchetypeHandle);
+		NewArchetype.CopyDebugNamesFrom(*OldArchetype);
 		EntityData.CurrentArchetype->MoveEntityToAnotherArchetype(Entity, NewArchetype);
 		EntityData.CurrentArchetype = NewArchetypeHandle.DataPtr;
 	}
@@ -750,6 +764,7 @@ void UMassEntitySubsystem::RemoveFragmentListFromEntity(FMassEntityHandle Entity
 
 		// Move the entity over
 		FMassArchetypeData& NewArchetype = FMassArchetypeHelper::ArchetypeDataFromHandleChecked(NewArchetypeHandle);
+		NewArchetype.CopyDebugNamesFrom(*OldArchetype);
 		OldArchetype->MoveEntityToAnotherArchetype(Entity, NewArchetype);
 		EntityData.CurrentArchetype = NewArchetypeHandle.DataPtr;
 	}
diff --git a/Engine/Plugins/Runtime/MassEntity/Source/MassEntity/Public/MassDebugger.h b/Engine/Plugins/Runtime/MassEntity/Source/MassEntity/Public/MassDebugger.h
--- a/Engine/Plugins/Runtime/MassEntity/Source/MassEntity/Public/MassDebugger.h
+++ b/Engine/Plugins/Runtime/MassEntity/Source/MassEntity/Public/MassDebugger.h
@@ -54,6 +54,7 @@ struct MASSENTITY_API FMassDebugger
 	static const FMassArchetypeSharedFragmentValues& GetArchetypeSharedFragmentValues(const FMassArchetypeHandle& ArchetypeHandle);
 
 	static void GetArchetypeEntityStats(const FMassArchetypeHandle& ArchetypeHandle, int32& OutEntitiesCount, int32& OutNumEntitiesPerChunk, int32& OutChunksCount);
+	static const TConstArrayView<FName> GetArchetypeDebugNames(const FMassArchetypeHandle& ArchetypeHandle);
 
 	static TConstArrayView<UMassCompositeProcessor::FDependencyNode> GetProcessingGraph(const UMassCompositeProcessor& GraphOwner);
 	
diff --git a/Engine/Plugins/Runtime/MassEntity/Source/MassEntity/Public/MassEntitySubsystem.h b/Engine/Plugins/Runtime/MassEntity/Source/MassEntity/Public/MassEntitySubsystem.h
--- a/Engine/Plugins/Runtime/MassEntity/Source/MassEntity/Public/MassEntitySubsystem.h
+++ b/Engine/Plugins/Runtime/MassEntity/Source/MassEntity/Public/MassEntitySubsystem.h
@@ -89,28 +89,30 @@ public:
 	 * A special, relaxed but slower version of CreateArchetype functions that allows FragmentAngTagsList to contain 
 	 * both fragments and tags. 
 	 */
-	FMassArchetypeHandle CreateArchetype(TConstArrayView<const UScriptStruct*> FragmentsAndTagsList);
+	FMassArchetypeHandle CreateArchetype(TConstArrayView<const UScriptStruct*> FragmentsAndTagsList, const FName ArchetypeDebugName = FName());
 
 	/**
 	 * CreateArchetype from a composition descriptor and initial values
 	 *
 	 * @param Composition of fragment, tag and chunk fragment types
 	 * @param SharedFragmentValues are the actual pointer to shared fragments
+	 * @param ArchetypeDebugName Name to identify the archetype while debugging
 	 * @return a handle of a new archetype 
 	 */
-	FMassArchetypeHandle CreateArchetype(const FMassArchetypeCompositionDescriptor& Composition, const FMassArchetypeSharedFragmentValues& SharedFragmentValues);
+	FMassArchetypeHandle CreateArchetype(const FMassArchetypeCompositionDescriptor& Composition, const FMassArchetypeSharedFragmentValues& SharedFragmentValues, const FName ArchetypeDebugName = FName());
 
 	/** 
 	 *  Creates an archetype like SourceArchetype + InFragments. 
 	 *  @param SourceArchetype the archetype used to initially populate the list of fragments of the archetype being created. 
 	 *  @param InFragments list of unique fragments to add to fragments fetched from SourceArchetype. Note that 
 	 *   adding an empty list is not supported and doing so will result in failing a `check`
+	 *  @param ArchetypeDebugName Name to identify the archetype while debugging
 	 *  @return a handle of a new archetype
 	 *  @note it's caller's responsibility to ensure that NewFragmentList is not empty and contains only fragment
 	 *   types that SourceArchetype doesn't already have. If the caller cannot guarantee it use of AddFragment functions
 	 *   family is recommended.
 	 */
-	FMassArchetypeHandle CreateArchetype(const TSharedPtr<FMassArchetypeData>& SourceArchetype, const FMassFragmentBitSet& InFragments);
+	FMassArchetypeHandle CreateArchetype(const TSharedPtr<FMassArchetypeData>& SourceArchetype, const FMassFragmentBitSet& InFragments, const FName ArchetypeDebugName = FName());
 
 	/** Fetches the archetype for a given Entity. If Entity is not valid it will still return a handle, just with an invalid archetype */
 	FMassArchetypeHandle GetArchetypeForEntity(FMassEntityHandle Entity) const;
@@ -138,8 +140,9 @@ public:
 	/**
 	 * Creates fully built entity ready to be used by the subsystem
 	 * @param FragmentInstanceList is the fragments to create the entity from and initialize values
+	 * @param ArchetypeDebugName Name to identify the archetype while debugging
 	 * @return FMassEntityHandle id of the newly created entity */
-	FMassEntityHandle CreateEntity(TConstArrayView<FInstancedStruct> FragmentInstanceList);
+	FMassEntityHandle CreateEntity(TConstArrayView<FInstancedStruct> FragmentInstanceList, const FName ArchetypeDebugName = FName());
 
 	/**
 	 * A dedicated structure for ensuring the "on entities creation" observers get notified only once all other 
@@ -217,8 +220,8 @@ public:
 	void RemoveTagFromEntity(FMassEntityHandle Entity, const UScriptStruct* TagType);
 	void SwapTagsForEntity(FMassEntityHandle Entity, const UScriptStruct* FromFragmentType, const UScriptStruct* ToFragmentType);
 
-	void BatchBuildEntities(const FMassArchetypeEntityCollectionWithPayload& EncodedEntitiesWithPayload, const FMassFragmentBitSet& FragmentsAffected, const FMassArchetypeSharedFragmentValues& SharedFragmentValues);
-	void BatchBuildEntities(const FMassArchetypeEntityCollectionWithPayload& EncodedEntitiesWithPayload, FMassArchetypeCompositionDescriptor&& Composition, const FMassArchetypeSharedFragmentValues& SharedFragmentValues);
+	void BatchBuildEntities(const FMassArchetypeEntityCollectionWithPayload& EncodedEntitiesWithPayload, const FMassFragmentBitSet& FragmentsAffected, const FMassArchetypeSharedFragmentValues& SharedFragmentValues, const FName ArchetypeDebugName = FName());
+	void BatchBuildEntities(const FMassArchetypeEntityCollectionWithPayload& EncodedEntitiesWithPayload, FMassArchetypeCompositionDescriptor&& Composition, const FMassArchetypeSharedFragmentValues& SharedFragmentValues, const FName ArchetypeDebugName = FName());
 	void BatchChangeTagsForEntities(TConstArrayView<FMassArchetypeEntityCollection> EntityCollections, const FMassTagBitSet& TagsToAdd, const FMassTagBitSet& TagsToRemove);
 	void BatchChangeFragmentCompositionForEntities(TConstArrayView<FMassArchetypeEntityCollection> EntityCollections, const FMassFragmentBitSet& FragmentsToAdd, const FMassFragmentBitSet& FragmentsToRemove);
 	void BatchAddFragmentInstancesForEntities(TConstArrayView<FMassArchetypeEntityCollectionWithPayload> EntityCollections, const FMassFragmentBitSet& FragmentsAffected);
diff --git a/Engine/Plugins/Runtime/MassGameplay/Source/MassSpawner/Private/MassEntityConfigAsset.cpp b/Engine/Plugins/Runtime/MassGameplay/Source/MassSpawner/Private/MassEntityConfigAsset.cpp
--- a/Engine/Plugins/Runtime/MassGameplay/Source/MassSpawner/Private/MassEntityConfigAsset.cpp
+++ b/Engine/Plugins/Runtime/MassGameplay/Source/MassSpawner/Private/MassEntityConfigAsset.cpp
@@ -81,6 +81,7 @@ const FMassEntityTemplate& FMassEntityConfig::GetOrCreateEntityTemplate(const UW
 	FMassEntityTemplateBuildContext BuildContext(Template);
 
 	BuildContext.BuildFromTraits(CombinedTraits, World);
+	Template.SetTemplateName(ConfigOwner.GetName());
 
 	if (ensureMsgf(!Template.IsEmpty(), TEXT("Need at least one fragment to create an Archetype")))
 	{
diff --git a/Engine/Plugins/Runtime/MassGameplay/Source/MassSpawner/Private/MassEntityTemplateRegistry.cpp b/Engine/Plugins/Runtime/MassGameplay/Source/MassSpawner/Private/MassEntityTemplateRegistry.cpp
--- a/Engine/Plugins/Runtime/MassGameplay/Source/MassSpawner/Private/MassEntityTemplateRegistry.cpp
+++ b/Engine/Plugins/Runtime/MassGameplay/Source/MassSpawner/Private/MassEntityTemplateRegistry.cpp
@@ -156,7 +156,7 @@ void UMassEntityTemplateRegistry::InitializeEntityTemplate(FMassEntityTemplate&
 	// Sort anything there is to sort for later comparison purposes
 	OutTemplate.Sort();
 
-	const FMassArchetypeHandle ArchetypeHandle = EntitySys->CreateArchetype(OutTemplate.GetCompositionDescriptor(), OutTemplate.GetSharedFragmentValues());
+	const FMassArchetypeHandle ArchetypeHandle = EntitySys->CreateArchetype(OutTemplate.GetCompositionDescriptor(), OutTemplate.GetSharedFragmentValues(), FName(OutTemplate.GetTemplateName()));
 	OutTemplate.SetArchetype(ArchetypeHandle);
 }
 
diff --git a/Engine/Plugins/Runtime/MassGameplay/Source/MassSpawner/Public/MassEntityTemplate.h b/Engine/Plugins/Runtime/MassGameplay/Source/MassSpawner/Public/MassEntityTemplate.h
--- a/Engine/Plugins/Runtime/MassGameplay/Source/MassSpawner/Public/MassEntityTemplate.h
+++ b/Engine/Plugins/Runtime/MassGameplay/Source/MassSpawner/Public/MassEntityTemplate.h
@@ -85,6 +85,9 @@ struct MASSSPAWNER_API FMassEntityTemplate
 	void SetTemplateID(FMassEntityTemplateID InTemplateID) { TemplateID = InTemplateID; }
 	FMassEntityTemplateID GetTemplateID() const { return TemplateID; }
 
+	void SetTemplateName(const FString& Name) { TemplateName = Name; }
+	const FString& GetTemplateName() const { return TemplateName; }
+	
 	const FMassArchetypeCompositionDescriptor& GetCompositionDescriptor() const { return Composition; }
 	const FMassArchetypeSharedFragmentValues& GetSharedFragmentValues() const { return SharedFragmentValues; }
 	TConstArrayView<FInstancedStruct> GetInitialFragmentValues() const { return InitialFragmentValues; }
@@ -252,4 +255,6 @@ private:
 	TArray<FObjectFragmentInitializerFunction> ObjectInitializers;
 
 	FMassEntityTemplateID TemplateID;
+
+	FString TemplateName;
 };
