Removed MassSchematic and MassSimulationLocalCoordinator

diff --git a/Engine/Plugins/Runtime/MassEntity/Source/MassEntity/Private/MassProcessingTypes.cpp b/Engine/Plugins/Runtime/MassEntity/Source/MassEntity/Private/MassProcessingTypes.cpp
--- a/Engine/Plugins/Runtime/MassEntity/Source/MassEntity/Private/MassProcessingTypes.cpp
+++ b/Engine/Plugins/Runtime/MassEntity/Source/MassEntity/Private/MassProcessingTypes.cpp
@@ -2,7 +2,6 @@
 
 #include "MassProcessingTypes.h"
 #include "MassProcessor.h"
-#include "MassSchematic.h"
 #include "Misc/OutputDevice.h"
 #include "MassEntityUtils.h"
 #include "VisualLogger/VisualLogger.h"
@@ -58,27 +57,6 @@ void FMassRuntimePipeline::SetProcessors(TArray<UMassProcessor*>&& InProcessors)
 	Processors = InProcessors;
 }
 
-void FMassRuntimePipeline::InitializeFromSchematics(TConstArrayView<TSoftObjectPtr<UMassSchematic>> Schematics, UObject& InOwner)
-{
-	Reset();
-	
-	// @todo we'll sometimes end up with duplicated MassProcessors in the resulting array. We need to come up with a consistent policy for handling that 
-	for (const TSoftObjectPtr<UMassSchematic>& Schematic : Schematics)
-	{
-		UMassSchematic* SchematicInstance = Schematic.LoadSynchronous();
-		if (SchematicInstance)
-		{
-			AppendOrOverrideRuntimeProcessorCopies(SchematicInstance->GetProcessors(), InOwner);
-		}
-		else
-		{
-			UE_LOG(LogMass, Error, TEXT("Unable to resolve MassSchematic %s while creating FMassRuntimePipeline"), *Schematic.GetLongPackageName());
-		}
-	}
-
-	Initialize(InOwner);
-}
-
 void FMassRuntimePipeline::CreateFromArray(TConstArrayView<const UMassProcessor*> InProcessors, UObject& InOwner)
 {
 	Reset();
diff --git a/Engine/Plugins/Runtime/MassEntity/Source/MassEntity/Private/MassSchematic.cpp b/Engine/Plugins/Runtime/MassEntity/Source/MassEntity/Private/MassSchematic.cpp
deleted file mode 100644
diff --git a/Engine/Plugins/Runtime/MassEntity/Source/MassEntity/Public/MassProcessingTypes.h b/Engine/Plugins/Runtime/MassEntity/Source/MassEntity/Public/MassProcessingTypes.h
--- a/Engine/Plugins/Runtime/MassEntity/Source/MassEntity/Public/MassProcessingTypes.h
+++ b/Engine/Plugins/Runtime/MassEntity/Source/MassEntity/Public/MassProcessingTypes.h
@@ -14,7 +14,6 @@
 
 class UMassEntitySubsystem;
 class UMassProcessor;
-class UMassSchematic;
 class UMassCompositeProcessor;
 struct FMassCommandBuffer;
 
@@ -64,7 +63,7 @@ struct MASSENTITY_API FMassProcessingContext
 };
 
 /** 
- *  A helper type that converts a set of UMassSchematics into a runtime-usable array of MassProcessor copies
+ *  Runtime-usable array of MassProcessor copies
  */
 USTRUCT()
 struct MASSENTITY_API FMassRuntimePipeline
@@ -80,9 +79,6 @@ struct MASSENTITY_API FMassRuntimePipeline
 	/** Creates runtime copies of the given UMassProcessors collection. */
 	void SetProcessors(TArray<UMassProcessor*>&& InProcessors);
 
-	/** Creates runtime copies of UMassProcessors declared in Schematics using InOwner as new UMassProcessors' outer. */
-	void InitializeFromSchematics(TConstArrayView<TSoftObjectPtr<UMassSchematic>> Schematics, UObject& InOwner);
-
 	/** Creates runtime copies of UMassProcessors given in InProcessors input parameter, using InOwner as new UMassProcessors' outer. */
 	void CreateFromArray(TConstArrayView<const UMassProcessor*> InProcessors, UObject& InOwner);
 
diff --git a/Engine/Plugins/Runtime/MassEntity/Source/MassEntity/Public/MassSchematic.h b/Engine/Plugins/Runtime/MassEntity/Source/MassEntity/Public/MassSchematic.h
deleted file mode 100644
diff --git a/Engine/Plugins/Runtime/MassEntity/Source/MassEntityEditor/Private/AssetTypeActions_MassSchematic.cpp b/Engine/Plugins/Runtime/MassEntity/Source/MassEntityEditor/Private/AssetTypeActions_MassSchematic.cpp
deleted file mode 100644
diff --git a/Engine/Plugins/Runtime/MassEntity/Source/MassEntityEditor/Private/AssetTypeActions_MassSchematic.h b/Engine/Plugins/Runtime/MassEntity/Source/MassEntityEditor/Private/AssetTypeActions_MassSchematic.h
deleted file mode 100644
diff --git a/Engine/Plugins/Runtime/MassEntity/Source/MassEntityEditor/Private/MassEntityEditor.cpp b/Engine/Plugins/Runtime/MassEntity/Source/MassEntityEditor/Private/MassEntityEditor.cpp
deleted file mode 100644
diff --git a/Engine/Plugins/Runtime/MassEntity/Source/MassEntityEditor/Private/MassEntityEditor.h b/Engine/Plugins/Runtime/MassEntity/Source/MassEntityEditor/Private/MassEntityEditor.h
deleted file mode 100644
diff --git a/Engine/Plugins/Runtime/MassEntity/Source/MassEntityEditor/Private/MassEntityEditorModule.cpp b/Engine/Plugins/Runtime/MassEntity/Source/MassEntityEditor/Private/MassEntityEditorModule.cpp
--- a/Engine/Plugins/Runtime/MassEntity/Source/MassEntityEditor/Private/MassEntityEditorModule.cpp
+++ b/Engine/Plugins/Runtime/MassEntity/Source/MassEntityEditor/Private/MassEntityEditorModule.cpp
@@ -4,12 +4,10 @@
 #include "Modules/ModuleManager.h"
 #include "Developer/AssetTools/Public/IAssetTools.h"
 #include "Developer/AssetTools/Public/AssetToolsModule.h"
-#include "AssetTypeActions_MassSchematic.h"
 #include "MassEditorStyle.h"
-#include "MassEntityEditor.h"
 #include "AIGraphTypes.h" // Class cache
-#include "MassSchematic.h"
 #include "MassProcessor.h"
+#include "AssetTypeActions_Base.h"
 
 
 #define LOCTEXT_NAMESPACE "MassEntityEditor"
@@ -23,12 +21,10 @@ void FMassEntityEditorModule::StartupModule()
 
 	FMassEntityEditorStyle::Initialize();
 
-	// Register asset types
-	IAssetTools& AssetTools = FModuleManager::LoadModuleChecked<FAssetToolsModule>("AssetTools").Get();
-
-	TSharedPtr<FAssetTypeActions_MassSchematic> MassAssetTypeAction = MakeShareable(new FAssetTypeActions_MassSchematic);
-	ItemDataAssetTypeActions.Add(MassAssetTypeAction);
-	AssetTools.RegisterAssetTypeActions(MassAssetTypeAction.ToSharedRef());
+	// Do register asset types here
+	//	IAssetTools& AssetTools = FModuleManager::LoadModuleChecked<FAssetToolsModule>("AssetTools").Get();
+	//	ItemDataAssetTypeActions.Add(MakeShareable(new FAssetTypeActions_MassThingy));
+	//	AssetTools.RegisterAssetTypeActions(ItemDataAssetTypeActions.Last().ToSharedRef());
 
 	// Register the details customizers
 	FPropertyEditorModule& PropertyModule = FModuleManager::LoadModuleChecked<FPropertyEditorModule>("PropertyEditor");
@@ -66,20 +62,4 @@ void FMassEntityEditorModule::ShutdownModule()
 		PropertyModule.NotifyCustomizationModuleChanged();
 	}
 }
-
-TSharedRef<IMassEntityEditor> FMassEntityEditorModule::CreateMassEntityEditor(const EToolkitMode::Type Mode, const TSharedPtr< class IToolkitHost >& InitToolkitHost, UMassSchematic* MassSchematic)
-{
-	if (!ProcessorClassCache.IsValid())
-	{
-		ProcessorClassCache = MakeShareable(new FGraphNodeClassHelper(UMassProcessor::StaticClass()));
-		ProcessorClassCache->UpdateAvailableBlueprintClasses();
-	}
-
-	TSharedRef<FMassEntityEditor> NewEditor(new FMassEntityEditor());
-	if (ensure(MassSchematic))
-	{
-		NewEditor->InitEditor(Mode, InitToolkitHost, *MassSchematic);
-	}
-	return NewEditor;
-}
 #undef LOCTEXT_NAMESPACE
diff --git a/Engine/Plugins/Runtime/MassEntity/Source/MassEntityEditor/Private/MassSchematicFactory.cpp b/Engine/Plugins/Runtime/MassEntity/Source/MassEntityEditor/Private/MassSchematicFactory.cpp
deleted file mode 100644
diff --git a/Engine/Plugins/Runtime/MassEntity/Source/MassEntityEditor/Public/MassEntityEditorModule.h b/Engine/Plugins/Runtime/MassEntity/Source/MassEntityEditor/Public/MassEntityEditorModule.h
--- a/Engine/Plugins/Runtime/MassEntity/Source/MassEntityEditor/Public/MassEntityEditorModule.h
+++ b/Engine/Plugins/Runtime/MassEntity/Source/MassEntityEditor/Public/MassEntityEditorModule.h
@@ -12,7 +12,6 @@
 class IMassEntityEditor;
 class FAssetTypeActions_Base;
 struct FGraphPanelNodeFactory;
-class UMassSchematic;
 
 /**
 * The public interface to this module
@@ -25,9 +24,6 @@ public:
 	virtual void StartupModule() override;
 	virtual void ShutdownModule() override;
 
-	/** Creates an instance of Mass editor. Only virtual so that it can be called across the DLL boundary. */
-	virtual TSharedRef<IMassEntityEditor> CreateMassEntityEditor(const EToolkitMode::Type Mode, const TSharedPtr< class IToolkitHost >& InitToolkitHost, UMassSchematic* MassSchematic);
-
 	virtual TSharedPtr<FExtensibilityManager> GetMenuExtensibilityManager() override { return MenuExtensibilityManager; }
 	virtual TSharedPtr<FExtensibilityManager> GetToolBarExtensibilityManager() override { return ToolBarExtensibilityManager; }
 
diff --git a/Engine/Plugins/Runtime/MassEntity/Source/MassEntityEditor/Public/MassSchematicFactory.h b/Engine/Plugins/Runtime/MassEntity/Source/MassEntityEditor/Public/MassSchematicFactory.h
deleted file mode 100644
diff --git a/Engine/Plugins/Runtime/MassGameplay/Source/MassGameplayEditor/Public/MassGameplayEditorModule.h b/Engine/Plugins/Runtime/MassGameplay/Source/MassGameplayEditor/Public/MassGameplayEditorModule.h
--- a/Engine/Plugins/Runtime/MassGameplay/Source/MassGameplayEditor/Public/MassGameplayEditorModule.h
+++ b/Engine/Plugins/Runtime/MassGameplay/Source/MassGameplayEditor/Public/MassGameplayEditorModule.h
@@ -9,7 +9,6 @@
 class IMassGameplayEditor;
 class FAssetTypeActions_Base;
 struct FGraphPanelNodeFactory;
-class UMassSchematic;
 
 /**
 * The public interface to this module
@@ -17,8 +16,6 @@ class UMassSchematic;
 class MASSGAMEPLAYEDITOR_API FMassGameplayEditorModule : public IModuleInterface
 {
 public:
-	DECLARE_MULTICAST_DELEGATE_TwoParams(FOnAssetPropertiesChanged, class UMassSchematic* /*MassSchematic*/, const FPropertyChangedEvent& /*PropertyChangedEvent*/);
-
 	virtual void StartupModule() override;
 	virtual void ShutdownModule() override;
 
diff --git a/Engine/Plugins/Runtime/MassGameplay/Source/MassSimulation/Private/MassSimulationLocalCoordinator.cpp b/Engine/Plugins/Runtime/MassGameplay/Source/MassSimulation/Private/MassSimulationLocalCoordinator.cpp
deleted file mode 100644
diff --git a/Engine/Plugins/Runtime/MassGameplay/Source/MassSimulation/Private/MassSimulationSettings.cpp b/Engine/Plugins/Runtime/MassGameplay/Source/MassSimulation/Private/MassSimulationSettings.cpp
deleted file mode 100644
diff --git a/Engine/Plugins/Runtime/MassGameplay/Source/MassSimulation/Private/MassSimulationSubsystem.cpp b/Engine/Plugins/Runtime/MassGameplay/Source/MassSimulation/Private/MassSimulationSubsystem.cpp
--- a/Engine/Plugins/Runtime/MassGameplay/Source/MassSimulation/Private/MassSimulationSubsystem.cpp
+++ b/Engine/Plugins/Runtime/MassGameplay/Source/MassSimulation/Private/MassSimulationSubsystem.cpp
@@ -1,7 +1,6 @@
 // Copyright Epic Games, Inc. All Rights Reserved.
 
 #include "MassSimulationSubsystem.h"
-#include "MassSimulationLocalCoordinator.h"
 #include "MassEntitySubsystem.h"
 #include "MassExecutor.h"
 #include "MassSimulationSettings.h"
@@ -27,20 +26,6 @@ UMassSimulationSubsystem::UMassSimulationSubsystem(const FObjectInitializer& Obj
 	PhaseManager = CreateDefaultSubobject<UMassProcessingPhaseManager>(TEXT("PhaseManager"));
 }
 
-void UMassSimulationSubsystem::PostInitProperties()
-{
-	REDIRECT_OBJECT_TO_VLOG(PhaseManager, this);
-
-	Super::PostInitProperties();
-
-	if (HasAnyFlags(RF_ClassDefaultObject) == false)
-	{
-#if WITH_EDITOR
-		GET_MASS_CONFIG_VALUE(GetOnTickSchematicChanged()).AddUObject(this, &UMassSimulationSubsystem::RebuildTickPipeline);
-#endif // WITH_EDITOR
-	}
-}
-
 void UMassSimulationSubsystem::BeginDestroy()
 {
 #if WITH_EDITOR
@@ -80,9 +65,6 @@ void UMassSimulationSubsystem::Initialize(FSubsystemCollectionBase& Collection)
 
 void UMassSimulationSubsystem::Deinitialize()
 {
-#if WITH_EDITOR
-	GET_MASS_CONFIG_VALUE(GetOnTickSchematicChanged()).RemoveAll(this);
-#endif // WITH_EDITOR
 	GetOnProcessingPhaseStarted(EMassProcessingPhase::PrePhysics).RemoveAll(this);
 	StopSimulation();
 
diff --git a/Engine/Plugins/Runtime/MassGameplay/Source/MassSimulation/Public/MassSimulationLocalCoordinator.h b/Engine/Plugins/Runtime/MassGameplay/Source/MassSimulation/Public/MassSimulationLocalCoordinator.h
deleted file mode 100644
diff --git a/Engine/Plugins/Runtime/MassGameplay/Source/MassSimulation/Public/MassSimulationSettings.h b/Engine/Plugins/Runtime/MassGameplay/Source/MassSimulation/Public/MassSimulationSettings.h
--- a/Engine/Plugins/Runtime/MassGameplay/Source/MassSimulation/Public/MassSimulationSettings.h
+++ b/Engine/Plugins/Runtime/MassGameplay/Source/MassSimulation/Public/MassSimulationSettings.h
@@ -7,8 +7,6 @@
 
 #define GET_MASS_CONFIG_VALUE(a) (GetMutableDefault<UMassSimulationSettings>()->a)
 
-class UMassSchematic;
-
 /**
  * Implements the settings for MassSimulation
  */
@@ -16,23 +14,8 @@ UCLASS(config = Mass, defaultconfig, DisplayName = "Mass Simulation")
 class MASSSIMULATION_API UMassSimulationSettings : public UMassModuleSettings
 {
 	GENERATED_BODY()
-public:
-	DECLARE_MULTICAST_DELEGATE(FOnTickSchematicChangedDelegate);
 
-#if WITH_EDITOR
 public:
-	FOnTickSchematicChangedDelegate& GetOnTickSchematicChanged() { return OnTickSchematicChanged; }
-
-protected:
-	virtual void PostInitProperties() override;
-	virtual void PostEditChangeProperty(struct FPropertyChangedEvent& PropertyChangedEvent) override;
-	void OnAssetPropertiesChanged(UMassSchematic* MassSchematic, const FPropertyChangedEvent& PropertyChangedEvent);	
-#endif // WITH_EDITOR
-
-public:
-	UPROPERTY(EditDefaultsOnly, Category = "Runtime", config, meta=(DisplayName="DEPRECATED_TickSchematics"))
-	TArray<TSoftObjectPtr<UMassSchematic>> TickSchematics;
-
 	/** The desired budget in seconds allowed to do actor spawning per frame */
 	UPROPERTY(EditDefaultsOnly, Category = "Runtime", config)
 	double DesiredActorSpawningTimeSlicePerTick = 0.0015;
@@ -52,9 +35,4 @@ public:
 	/** The distance a failed spawned actor needs to move before we retry, default 10 meters */
 	UPROPERTY(EditDefaultsOnly, Category = "Runtime", config)
 	float DesiredActorFailedSpawningRetryMoveDistance = 500.0f;
-
-protected:
-#if WITH_EDITOR
-	FOnTickSchematicChangedDelegate OnTickSchematicChanged;
-#endif // WITH_EDITOR
 };
diff --git a/Engine/Plugins/Runtime/MassGameplay/Source/MassSimulation/Public/MassSimulationSubsystem.h b/Engine/Plugins/Runtime/MassGameplay/Source/MassSimulation/Public/MassSimulationSubsystem.h
--- a/Engine/Plugins/Runtime/MassGameplay/Source/MassSimulation/Public/MassSimulationSubsystem.h
+++ b/Engine/Plugins/Runtime/MassGameplay/Source/MassSimulation/Public/MassSimulationSubsystem.h
@@ -9,7 +9,6 @@
 #include "MassSimulationSubsystem.generated.h"
 
 
-class AMassSimulationLocalCoordinator;
 class UMassEntitySubsystem;
 class UMassProcessingPhaseManager;
 
@@ -34,7 +33,6 @@ public:
 	bool IsSimulationStarted() const { return bSimulationStarted; }
 
 protected:
-	virtual void PostInitProperties() override;
 	virtual void Initialize(FSubsystemCollectionBase& Collection) override;
 	virtual void PostInitialize() override;
 	virtual void OnWorldBeginPlay(UWorld& InWorld) override;
diff --git a/Engine/Plugins/Runtime/MassGameplay/Source/MassSpawner/Public/MassSpawner.h b/Engine/Plugins/Runtime/MassGameplay/Source/MassSpawner/Public/MassSpawner.h
--- a/Engine/Plugins/Runtime/MassGameplay/Source/MassSpawner/Public/MassSpawner.h
+++ b/Engine/Plugins/Runtime/MassGameplay/Source/MassSpawner/Public/MassSpawner.h
@@ -8,7 +8,7 @@
 #include "MassEntitySpawnDataGeneratorBase.h"
 #include "MassSpawner.generated.h"
 
-class UMassSchematic;
+
 struct FStreamableHandle;
 
 DECLARE_DYNAMIC_MULTICAST_DELEGATE(FMassSpawnerOnSpawningFinishedEvent);
