[MassEntity] fixed memory leak from batched commands.
Also replaced Memcpy by CopyScriptStruct when setting fragment values in batch since the memory is not relocated in this case (command payload copied to fragment)

diff --git a/Engine/Plugins/Runtime/MassEntity/Source/MassEntity/Private/MassArchetypeData.cpp b/Engine/Plugins/Runtime/MassEntity/Source/MassEntity/Private/MassArchetypeData.cpp
--- a/Engine/Plugins/Runtime/MassEntity/Source/MassEntity/Private/MassArchetypeData.cpp
+++ b/Engine/Plugins/Runtime/MassEntity/Source/MassEntity/Private/MassArchetypeData.cpp
@@ -1121,7 +1121,7 @@ void FMassArchetypeData::BatchSetFragmentValues(TConstArrayView<FMassArchetypeEn
 			void* Dst = FragmentConfigs[FragmentIndex].GetFragmentData(Chunk.GetRawMemory(), EntityRange.SubchunkStart);
 			const void* Src = FragmentPayload.GetDataAt(EntitiesHandled);
 
-			FMemory::Memcpy(Dst, Src, FragmentType.GetStructureSize() * EntityRange.Length);
+			FragmentType.CopyScriptStruct(Dst, Src, EntityRange.Length);
 		}
 
 		EntitiesHandled += EntityRange.Length;
diff --git a/Engine/Plugins/Runtime/MassEntity/Source/MassEntity/Private/MassCommandBuffer.cpp b/Engine/Plugins/Runtime/MassEntity/Source/MassEntity/Private/MassCommandBuffer.cpp
--- a/Engine/Plugins/Runtime/MassEntity/Source/MassEntity/Private/MassCommandBuffer.cpp
+++ b/Engine/Plugins/Runtime/MassEntity/Source/MassEntity/Private/MassCommandBuffer.cpp
@@ -76,6 +76,11 @@ std::atomic<uint32> FMassBatchedCommand::CommandsCounter;
 FMassCommandBuffer::~FMassCommandBuffer()
 {
 	ensureMsgf(HasPendingCommands() == false, TEXT("Destroying FMassCommandBuffer while there are still unprocessed commands. These operations will never be performed now."));
+
+	for (FMassBatchedCommand*& Command : CommandInstances)
+	{
+		delete Command;
+	}
 }
 
 void FMassCommandBuffer::Flush(UMassEntitySubsystem& EntitySystem)
