MassDebugger: Added names to archetypes
- Add name to each archetype if possible (debug name, first fragment, hash as fallback)
- Organize archetypes based on names (allows to see what has changed for each derived archetype)
- Improved the fragment display
- Allow to diff archetypes using selection

diff --git a/Engine/Plugins/Runtime/MassEntity/Source/MassEntityDebugger/Private/MassDebuggerModel.cpp b/Engine/Plugins/Runtime/MassEntity/Source/MassEntityDebugger/Private/MassDebuggerModel.cpp
--- a/Engine/Plugins/Runtime/MassEntity/Source/MassEntityDebugger/Private/MassDebuggerModel.cpp
+++ b/Engine/Plugins/Runtime/MassEntity/Source/MassEntityDebugger/Private/MassDebuggerModel.cpp
@@ -156,12 +156,80 @@ FMassDebuggerArchetypeData::FMassDebuggerArchetypeData(const FMassArchetypeHandl
 	SharedFragments = FMassDebugger::GetArchetypeSharedFragmentValues(ArchetypeHandle);
 
 	// @todo should ensure we're using same hashing as the EntitySubsystem here
-	FamilyHash = Composition.CalculateHash();
-	FullHash = HashCombine(FamilyHash, GetTypeHash(SharedFragments));
+	CompositionHash = Composition.CalculateHash();
+	FullHash = HashCombine(CompositionHash, GetTypeHash(SharedFragments));
 
-	BytesToHexLower(reinterpret_cast<const uint8*>(&FullHash), sizeof(FullHash), Label);
+	FString FullHashAsString;
+	BytesToHexLower(reinterpret_cast<const uint8*>(&FullHash), sizeof(FullHash), FullHashAsString);
+	HashLabel = FText::FromString(FullHashAsString);
 
 	FMassDebugger::GetArchetypeEntityStats(ArchetypeHandle, EntitiesCount, EntitiesCountPerChunk, ChunksCount);
+
+	const TConstArrayView<FName> DebugNames = FMassDebugger::GetArchetypeDebugNames(ArchetypeHandle);
+
+	if (DebugNames.IsEmpty())
+	{
+		// This archetype has no associated debug names, use hash as name.
+		FString HashAsString; 
+		BytesToHexLower(reinterpret_cast<const uint8*>(&CompositionHash), sizeof(CompositionHash), HashAsString);
+		PrimaryDebugName = HashAsString;
+
+		// Use first fragment as name
+		if (FMassFragmentBitSet::FIndexIterator It = Composition.Fragments.GetIndexIterator())
+		{
+			const FName FirstStructName = Composition.Fragments.DebugGetStructTypeName(*It);
+			TStringBuilder<256> StringBuilder;
+			StringBuilder.Append(FirstStructName.ToString());
+			StringBuilder.Append(TEXT("..."));
+			Label = FText::FromString(StringBuilder.ToString());
+		}
+		else
+		{
+			Label = FText::FromString(HashAsString);
+		}
+		
+		LabelLong = Label;
+	}
+	else
+	{
+		PrimaryDebugName = DebugNames[0].ToString();
+
+		TStringBuilder<256> StringBuilder;
+
+		// Short label for lists
+		StringBuilder.Reset();
+		StringBuilder.Append(DebugNames[0].ToString());
+		if (DebugNames.Num() > 1)
+		{
+			StringBuilder.Append(TEXT("..."));
+		}
+		Label = FText::FromString(StringBuilder.ToString());
+
+		// Longer label for info display
+		StringBuilder.Reset();
+		for (int i = 0; i < DebugNames.Num(); i++)
+		{
+			if (i > 0)
+			{
+				StringBuilder.Append(TEXT(", "));
+			}
+			StringBuilder.Append(DebugNames[i].ToString());
+		}
+		LabelLong = FText::FromString(StringBuilder.ToString());
+
+		// Label tooltip
+		StringBuilder.Reset();
+		for (int i = 0; i < DebugNames.Num(); i++)
+		{
+			if (i > 0)
+			{
+				StringBuilder.Append(TEXT("\n"));
+			}
+			StringBuilder.Append(DebugNames[i].ToString());
+		}
+		LabelTooltip = FText::FromString(StringBuilder.ToString());
+	}
+
 #endif // WITH_MASSENTITY_DEBUG
 }
 
@@ -433,6 +501,24 @@ void FMassDebuggerModel::StoreArchetypes(const UMassEntitySubsystem& EntitySubsy
 			ArchetypeDistances[k][i] = Distance;
 		}
 	}
+
+	// Add archetypes that share same primary name under the same entry. 
+	TMap<FString, TSharedPtr<FMassDebuggerArchetypeData>> ArchetypeNameMap;
+	for (int32 Index = 0; Index < CachedArchetypes.Num(); Index++)
+	{
+		TSharedPtr<FMassDebuggerArchetypeData> ArchetypeData = CachedArchetypes[Index];
+		if (const TSharedPtr<FMassDebuggerArchetypeData>* Representative = ArchetypeNameMap.Find(ArchetypeData->PrimaryDebugName))
+		{
+			(*Representative)->Children.Add(ArchetypeData);
+			ArchetypeData->Parent = *Representative; 
+			CachedArchetypes.RemoveAt(Index);
+			Index--;
+		}
+		else
+		{
+			ArchetypeNameMap.Add(ArchetypeData->PrimaryDebugName, ArchetypeData);
+		}
+	}
 }
 
 FText FMassDebuggerModel::GetDisplayName() const
diff --git a/Engine/Plugins/Runtime/MassEntity/Source/MassEntityDebugger/Private/MassDebuggerModel.h b/Engine/Plugins/Runtime/MassEntity/Source/MassEntityDebugger/Private/MassDebuggerModel.h
--- a/Engine/Plugins/Runtime/MassEntity/Source/MassEntityDebugger/Private/MassDebuggerModel.h
+++ b/Engine/Plugins/Runtime/MassEntity/Source/MassEntityDebugger/Private/MassDebuggerModel.h
@@ -59,18 +59,39 @@ struct FMassDebuggerArchetypeData
 
 	FMassArchetypeCompositionDescriptor Composition;
 	FMassArchetypeSharedFragmentValues SharedFragments;
-	uint32 FamilyHash = 0;
+
+	/** Hash of the Compositions. */
+	uint32 CompositionHash = 0;
+	/** Combined hash of composition and shared fragments. */
 	uint32 FullHash = 0;
 
+	/** Number of active entities of the archetype. */
 	int32 EntitiesCount = 0;
+	/** Number of entities that fit per chunk. */
 	int32 EntitiesCountPerChunk = 0;
+	/** Number of allocated chunks. */
 	int32 ChunksCount = 0;
 
+	/** Child debugger data (same as parent, but changed in some way) */
 	TArray<TSharedPtr<FMassDebuggerArchetypeData>> Children;
+	/** Parent debugger data. */
+	TWeakPtr<FMassDebuggerArchetypeData> Parent;
 
-	int32 Index = INDEX_NONE;
-	FString Label;
 
+	/** Index in FMassDebuggerModel::CachedArchetypes */
+	int32 Index = INDEX_NONE;
+	/** Display label */
+	FText Label;
+	/** Display label */
+	FText LabelLong;
+	/** Display label tooltip */
+	FText LabelTooltip;
+	/** FullHash as a display string */
+	FText HashLabel;
+	/** Primary debug name, used for grouping derived archetypes. */
+	FString PrimaryDebugName;
+
+	/** True if the archetype is selected. */
 	bool bIsSelected = false;
 
 	int32 GetTotalBitsUsedCount() const;
diff --git a/Engine/Plugins/Runtime/MassEntity/Source/MassEntityDebugger/Private/MassDebuggerStyle.cpp b/Engine/Plugins/Runtime/MassEntity/Source/MassEntityDebugger/Private/MassDebuggerStyle.cpp
--- a/Engine/Plugins/Runtime/MassEntity/Source/MassEntityDebugger/Private/MassDebuggerStyle.cpp
+++ b/Engine/Plugins/Runtime/MassEntity/Source/MassEntityDebugger/Private/MassDebuggerStyle.cpp
@@ -11,6 +11,7 @@
 #include "UObject/ConstructorHelpers.h"
 #include "Engine/Font.h"
 #include "Components/Widget.h"
+#include "Styling/StyleColors.h"
 
 
 #define TTF_CORE_FONT(RelativePath, ...) FSlateFontInfo(RootToCoreContentDir(RelativePath, TEXT(".ttf")), __VA_ARGS__)
@@ -45,8 +46,16 @@ public:
 		Set("MassDebug.Font.Large", StateTitle);
 
 		Set("MassDebug.Fragment", new BOX_BRUSH("Common/LightGroupBorder", FMargin(4.0f / 16.0f), FLinearColor(.2f, .2f, .9f)));
+		Set("MassDebug.Fragment.Added", new BOX_BRUSH("Common/LightGroupBorder", FMargin(4.0f / 16.0f), FLinearColor(.2f, .2f, .9f)));
+		Set("MassDebug.Fragment.Removed", new BOX_BRUSH("Common/LightGroupBorder", FMargin(4.0f / 16.0f), FLinearColor(.9f, .2f, .2f)));
 		Set("MassDebug.Fragment.ReadOnly", new BOX_BRUSH("Common/LightGroupBorder", FMargin(4.0f / 16.0f), FLinearColor(.5f, .5f, .5f)));
 		Set("MassDebug.Fragment.ReadWrite", new BOX_BRUSH("Common/LightGroupBorder", FMargin(4.0f / 16.0f), FLinearColor(0.f, .7f, 0.f)));
+
+		Set("MassDebug.Label.Background", new FSlateRoundedBoxBrush(FStyleColors::Foreground, 4.0f));
+
+		Set("MassDebug.Label.Text", FTextBlockStyle(NormalText)
+			.SetFont(FSlateFontInfo(FONT(7, "Bold")))
+			.SetColorAndOpacity(FStyleColors::Background));
 	}
 };
 
diff --git a/Engine/Plugins/Runtime/MassEntity/Source/MassEntityDebugger/Private/SMassArchetype.cpp b/Engine/Plugins/Runtime/MassEntity/Source/MassEntityDebugger/Private/SMassArchetype.cpp
--- a/Engine/Plugins/Runtime/MassEntity/Source/MassEntityDebugger/Private/SMassArchetype.cpp
+++ b/Engine/Plugins/Runtime/MassEntity/Source/MassEntityDebugger/Private/SMassArchetype.cpp
@@ -12,10 +12,8 @@
 //----------------------------------------------------------------------//
 // SMassArchetype
 //----------------------------------------------------------------------//
-void SMassArchetype::Construct(const FArguments& InArgs, TSharedPtr<FMassDebuggerArchetypeData> InArchetypeData)
+void SMassArchetype::Construct(const FArguments& InArgs, TSharedPtr<FMassDebuggerArchetypeData> InArchetypeData, TSharedPtr<FMassDebuggerArchetypeData> InBaseArchetypeData, const EMassBitSetDiffPrune Prune)
 {
-	using UE::Mass::Debugger::UI::AddBitSet;
-
 	if (!InArchetypeData)
 	{
 		return;
@@ -23,17 +21,24 @@ void SMassArchetype::Construct(const FArguments& InArgs, TSharedPtr<FMassDebugge
 	
 	ArchetypeData = InArchetypeData;
 
-	FMassDebuggerArchetypeData& ArchetypeDebugData = *InArchetypeData.Get();
+	const FMassDebuggerArchetypeData* BaseArchetypeDebugData = InBaseArchetypeData.Get();
+	const FMassDebuggerArchetypeData& ArchetypeDebugData = *InArchetypeData.Get();
 
+	if (BaseArchetypeDebugData == &ArchetypeDebugData)
+	{
+		BaseArchetypeDebugData = nullptr;
+	}
+	
 	TSharedRef<SVerticalBox> Box = SNew(SVerticalBox);
 
 	const TArray<FText> LabelBits = {
 		LOCTEXT("MassArchetypeLabel", "Archetype")
-		, FText::FromString(InArchetypeData->Label)
+		, InArchetypeData->LabelLong
 	};
 
 	Box->AddSlot()
 	.AutoHeight()
+	.Padding(0, 4)
 	[
 		SNew(SRichTextBlock)
 		.Text(FText::Join(FText::FromString(TEXT(": ")), LabelBits))
@@ -41,6 +46,15 @@ void SMassArchetype::Construct(const FArguments& InArgs, TSharedPtr<FMassDebugge
 		.TextStyle(FAppStyle::Get(), "LargeText")
 	];
 
+	Box->AddSlot()
+	.AutoHeight()
+	.Padding(0, 4)
+	[
+		SNew(STextBlock)
+		.Text(InArchetypeData->HashLabel)
+		.Font(FCoreStyle::GetDefaultFontStyle("Mono", 9))
+	];
+	
 	FString ArchetypeDescription = FString::Printf(TEXT("EntitiesCount: %d\nEntitiesCountPerChunk: %d\nChunksCount: %d")
 		, ArchetypeDebugData.EntitiesCount
 		, ArchetypeDebugData.EntitiesCountPerChunk
@@ -48,6 +62,7 @@ void SMassArchetype::Construct(const FArguments& InArgs, TSharedPtr<FMassDebugge
 	
 	Box->AddSlot()
 	.AutoHeight()
+	.Padding(0, 4)
 	[
 		SNew(STextBlock)
 		.Text(FText::FromString(ArchetypeDescription))
@@ -56,10 +71,22 @@ void SMassArchetype::Construct(const FArguments& InArgs, TSharedPtr<FMassDebugge
 	const FMassArchetypeCompositionDescriptor& Composition = ArchetypeData->Composition;
 	const FSlateBrush* Brush = FMassDebuggerStyle::GetBrush("MassDebug.Fragment");
 
-	AddBitSet(Box, Composition.Fragments, TEXT("Fragments"), Brush);
-	AddBitSet(Box, Composition.Tags, TEXT("Tags"), Brush);
-	AddBitSet(Box, Composition.ChunkFragments, TEXT("Chunk Fragments"), Brush);
-	AddBitSet(Box, Composition.SharedFragments, TEXT("Shared Fragments"), Brush);
+	if (BaseArchetypeDebugData != nullptr)
+	{
+		const FMassArchetypeCompositionDescriptor& ParentComposition = BaseArchetypeDebugData->Composition;
+
+		UE::Mass::Debugger::UI::AddBitSetDiff(Box, ParentComposition.Fragments, Composition.Fragments, TEXT("Fragments"), Brush, Prune);
+		UE::Mass::Debugger::UI::AddBitSetDiff(Box, ParentComposition.Tags, Composition.Tags, TEXT("Tags"), Brush, Prune);
+		UE::Mass::Debugger::UI::AddBitSetDiff(Box, ParentComposition.ChunkFragments, Composition.ChunkFragments, TEXT("Chunk Fragments"), Brush, Prune);
+		UE::Mass::Debugger::UI::AddBitSetDiff(Box, ParentComposition.SharedFragments, Composition.SharedFragments, TEXT("Shared Fragments"), Brush, Prune);
+	}
+	else
+	{
+		UE::Mass::Debugger::UI::AddBitSet(Box, Composition.Fragments, TEXT("Fragments"), Brush);
+		UE::Mass::Debugger::UI::AddBitSet(Box, Composition.Tags, TEXT("Tags"), Brush);
+		UE::Mass::Debugger::UI::AddBitSet(Box, Composition.ChunkFragments, TEXT("Chunk Fragments"), Brush);
+		UE::Mass::Debugger::UI::AddBitSet(Box, Composition.SharedFragments, TEXT("Shared Fragments"), Brush);
+	}
 
 	ChildSlot
 	[
diff --git a/Engine/Plugins/Runtime/MassEntity/Source/MassEntityDebugger/Private/SMassArchetype.h b/Engine/Plugins/Runtime/MassEntity/Source/MassEntityDebugger/Private/SMassArchetype.h
--- a/Engine/Plugins/Runtime/MassEntity/Source/MassEntityDebugger/Private/SMassArchetype.h
+++ b/Engine/Plugins/Runtime/MassEntity/Source/MassEntityDebugger/Private/SMassArchetype.h
@@ -2,7 +2,7 @@
 #pragma once
 
 #include "Widgets/SCompoundWidget.h"
-
+#include "SMassBitSet.h"
 
 struct FMassDebuggerArchetypeData;
 
@@ -13,7 +13,7 @@ public:
 	SLATE_BEGIN_ARGS(SMassArchetype){}
 	SLATE_END_ARGS()
 
-	void Construct(const FArguments& InArgs, TSharedPtr<FMassDebuggerArchetypeData> InArchetypeData);
+	void Construct(const FArguments& InArgs, TSharedPtr<FMassDebuggerArchetypeData> InArchetypeData, TSharedPtr<FMassDebuggerArchetypeData> InBaseArchetypeData, const EMassBitSetDiffPrune Prune);
 
 protected:
 	TSharedPtr<FMassDebuggerArchetypeData> ArchetypeData;
diff --git a/Engine/Plugins/Runtime/MassEntity/Source/MassEntityDebugger/Private/SMassArchetypesView.cpp b/Engine/Plugins/Runtime/MassEntity/Source/MassEntityDebugger/Private/SMassArchetypesView.cpp
--- a/Engine/Plugins/Runtime/MassEntity/Source/MassEntityDebugger/Private/SMassArchetypesView.cpp
+++ b/Engine/Plugins/Runtime/MassEntity/Source/MassEntityDebugger/Private/SMassArchetypesView.cpp
@@ -3,6 +3,9 @@
 #include "SMassArchetypesView.h"
 #include "MassDebuggerModel.h"
 #include "SMassArchetype.h"
+#include "MassDebuggerStyle.h"
+#include "Styling/StyleColors.h"
+#include "Widgets/Layout/SScrollBox.h"		
 
 #define LOCTEXT_NAMESPACE "SMassDebugger"
 
@@ -11,46 +14,43 @@ namespace UE::Mass::Debugger::UI::Private
 {
 	FLinearColor GetArchetypeDistanceColor(float Distance)
 	{
-		const FLinearColor MinDistance = FLinearColor::Green;
-		const FLinearColor MaxDistance = FLinearColor::Red;
-		
-		return (1.f - Distance) * MinDistance + Distance * MaxDistance;
+		const FLinearColor Foreground = FStyleColors::Foreground.GetSpecifiedColor();
+		const FLinearColor MinDistance = FMath::Lerp(Foreground, FLinearColor::Green, 0.75f);
+		const FLinearColor MaxDistance = FMath::Lerp(Foreground, FLinearColor::Red, 0.75f);
+
+		return FMath::Lerp(MinDistance, MaxDistance, Distance);
 	}
+
+	static const FName ColumnName_ArcheType = TEXT("Archetype");
+	static const FName ColumnName_Hash = TEXT("Hash");
+	static const FName ColumnName_Entities = TEXT("Entities");
+
 }
 
 using FMassDebuggerArchetypeDataPtr = TSharedPtr<FMassDebuggerArchetypeData, ESPMode::ThreadSafe>;
 
 //----------------------------------------------------------------------//
-// SMassArchetypeTableRowBase
+// SMassArchetypeTableRow
 //----------------------------------------------------------------------//
-class SMassArchetypeTableRowBase : public STableRow<FMassDebuggerArchetypeDataPtr>
+class SMassArchetypeTableRow : public SMultiColumnTableRow<FMassDebuggerArchetypeDataPtr>
 {
 public:
-	SLATE_BEGIN_ARGS(SMassArchetypeTableRowBase) { }
-	SLATE_END_ARGS()
-
-protected:
-	void Construct(const FArguments& InArgs, const TSharedPtr<STableViewBase>& InOwnerTableView, const FMassDebuggerArchetypeDataPtr InEntryItem)
+	void Construct(const FArguments& InArgs, const TSharedPtr<STableViewBase>& InOwnerTableView, const FMassDebuggerArchetypeDataPtr InEntryItem, const FMassDebuggerArchetypeDataPtr InBaseItem, TSharedPtr<FMassDebuggerModel> InDebuggerModel)
 	{
 		Item = InEntryItem;
-		STableRow<FMassDebuggerArchetypeDataPtr>::Construct(STableRow<FMassDebuggerArchetypeDataPtr>::FArguments(), InOwnerTableView.ToSharedRef());
+		DebuggerModel = InDebuggerModel;
+		bIsBase = InEntryItem == InBaseItem;
+		SMultiColumnTableRow<FMassDebuggerArchetypeDataPtr>::Construct(FSuperRowType::FArguments(), InOwnerTableView.ToSharedRef());
 	}
 
-	FMassDebuggerArchetypeDataPtr Item;
-};
-
-//----------------------------------------------------------------------//
-// SMassArchetypeTableRow
-//----------------------------------------------------------------------//
-class SMassArchetypeTableRow : public SMassArchetypeTableRowBase
-{
-public:
-	void Construct(const FArguments& InArgs, const TSharedPtr<STableViewBase>& InOwnerTableView, const FMassDebuggerArchetypeDataPtr InEntryItem, TSharedPtr<FMassDebuggerModel> DebuggerModel)
+	// SMultiColumnTableRow overrides
+	virtual TSharedRef<SWidget> GenerateWidgetForColumn(const FName& ColumnName) override
 	{
-		SMassArchetypeTableRowBase::Construct(InArgs, InOwnerTableView, InEntryItem);
+		TSharedPtr<SHorizontalBox> Contents = SNew(SHorizontalBox);
 
-		TSharedPtr<SHorizontalBox> Box = SNew(SHorizontalBox);
-		Box->AddSlot()
+		if (ColumnName == UE::Mass::Debugger::UI::Private::ColumnName_ArcheType)
+		{
+			Contents->AddSlot()
 			.VAlign(VAlign_Fill)
 			.HAlign(HAlign_Left)
 			.AutoWidth()
@@ -61,61 +61,86 @@ public:
 				.BaseIndentLevel(0)
 			];
 
-		if (DebuggerModel->SelectedArchetypes.Num())
-		{
-			if (InEntryItem->bIsSelected)
+			FLinearColor Color = FLinearColor::White;
+			if (DebuggerModel->SelectedArchetypes.Num())
 			{
-				Box->AddSlot()
-				[
-					SNew(STextBlock)
-					.Text(FText::FromString(InEntryItem->Label))
-					.ColorAndOpacity(FLinearColor::Green)
-				];
+				if (Item->bIsSelected)
+				{
+					const FLinearColor Foreground = FStyleColors::Foreground.GetSpecifiedColor();
+					Color = FMath::Lerp(Foreground, FLinearColor::Green, 0.75f);
+				}
+				else
+				{
+					const float DistanceToSelected = DebuggerModel->MinDistanceToSelectedArchetypes(Item);
+					Color = UE::Mass::Debugger::UI::Private::GetArchetypeDistanceColor(DistanceToSelected);
+					if (Item->EntitiesCount == 0)
+					{
+						Color.A *= 0.5f;
+					}
+				}
 			}
 			else
 			{
-				const float DistanceToSelected = DebuggerModel->MinDistanceToSelectedArchetypes(InEntryItem);
-				const FLinearColor LerpedColor = UE::Mass::Debugger::UI::Private::GetArchetypeDistanceColor(DistanceToSelected);
+				Color = FLinearColor::White;
+				if (Item->EntitiesCount == 0)
+				{
+					Color.A *= 0.5f;
+				}
+			}
+			Contents->AddSlot()
+			[
+				SNew(STextBlock)
+				.Text(Item->Label)
+				.ToolTipText(Item->LabelTooltip)
+				.ColorAndOpacity(Color)
+			];
 
-				Box->AddSlot()
+			if (bIsBase)
+			{
+				Contents->AddSlot()
+				.VAlign(VAlign_Center)
+				.Padding(FMargin(6,0))
+				.AutoWidth()
+				[
+					SNew(SBorder)
+					.BorderImage(FMassDebuggerStyle::GetBrush("MassDebug.Label.Background"))
+					.Padding(FMargin(4, 1))
+					[
+						SNew(STextBlock)
+						.TextStyle(FMassDebuggerStyle::Get(), TEXT("MassDebug.Label.Text"))
+						.Text(LOCTEXT("ArchetypeDiffBase", "BASE"))
+						.ToolTipText(LOCTEXT("ArchetypeDiffBaseTootip", "Base archetype used for diffing."))
+					]
+			];
+			}
+		}
+		else if (ColumnName == UE::Mass::Debugger::UI::Private::ColumnName_Hash)
+		{
+			Contents->AddSlot()
+				.Padding(5, 0)
 				[
 					SNew(STextBlock)
-					.Text(FText::FromString(InEntryItem->Label))
-					.ColorAndOpacity(LerpedColor)
+					.Font(FCoreStyle::GetDefaultFontStyle("Mono", 9))
+					.Text(Item->HashLabel)
 				];
-			}
 		}
-		else
+		else if (ColumnName == UE::Mass::Debugger::UI::Private::ColumnName_Entities)
 		{
-			Box->AddSlot()
-			[
-				SNew(STextBlock)
-				.Text(FText::FromString(InEntryItem->Label))
-			];
+			Contents->AddSlot()
+				.Padding(5, 0)
+				[
+					SNew(STextBlock)
+					.Text(FText::FromString(FString::Printf(TEXT("%d"), Item->EntitiesCount)))
+					.Justification(ETextJustify::Right)
+				];
 		}
 
-		ChildSlot
-		[
-			Box.ToSharedRef()
-		];
+		return Contents.ToSharedRef();
 	}
-};
-
-//----------------------------------------------------------------------//
-// SMassArchetypeDetailTableRow
-//----------------------------------------------------------------------//
-class SMassArchetypeDetailTableRow : public SMassArchetypeTableRowBase
-{
-public:
-	void Construct(const FArguments& InArgs, const TSharedPtr<STableViewBase>& InOwnerTableView, const FMassDebuggerArchetypeDataPtr InEntryItem)
-	{
-		SMassArchetypeTableRowBase::Construct(InArgs, InOwnerTableView, InEntryItem);
 
-		ChildSlot
-			[
-				SNew(SMassArchetype, InEntryItem)
-			];
-	}
+	FMassDebuggerArchetypeDataPtr Item;
+	bool bIsBase = false;
+	TSharedPtr<FMassDebuggerModel> DebuggerModel;
 };
 
 //----------------------------------------------------------------------//
@@ -127,37 +152,99 @@ void SMassArchetypesView::Construct(const FArguments& InArgs, TSharedRef<FMassDe
 
 	ChildSlot
 		[
-			SNew(SHorizontalBox)
-			+ SHorizontalBox::Slot()
-			.FillWidth(0.2f)
+			SNew(SSplitter)
+			.Orientation(Orient_Horizontal)
+			
+			+ SSplitter::Slot()
+			.Value(.4f)
+			.MinSize(350.0f)
 			[
-				SAssignNew(ArchetypesTreeView, STreeView<TSharedPtr<FMassDebuggerArchetypeData>>)
-				.TreeItemsSource(&TreeViewSource)
-				.OnGenerateRow_Lambda([this](TSharedPtr<FMassDebuggerArchetypeData> Item, const TSharedPtr<STableViewBase>& OwnerTable)
-					{
-						return SNew(SMassArchetypeTableRow, OwnerTable, Item, DebuggerModel);
-					})
-				.OnGetChildren_Lambda([](TSharedPtr<FMassDebuggerArchetypeData> InItem, TArray<TSharedPtr<FMassDebuggerArchetypeData>>& OutChildren)
+				SNew(SVerticalBox)
+				+ SVerticalBox::Slot()
+				[
+					SAssignNew(ArchetypesTreeView, STreeView<TSharedPtr<FMassDebuggerArchetypeData>>)
+					.TreeItemsSource(&TreeViewSource)
+					.HeaderRow(SNew(SHeaderRow)
+						+SHeaderRow::Column(UE::Mass::Debugger::UI::Private::ColumnName_ArcheType)
+							.DefaultLabel(LOCTEXT("ArchetypeColumnHeader", "Archetype"))
+							.FillWidth(0.70f)
+						+SHeaderRow::Column(UE::Mass::Debugger::UI::Private::ColumnName_Hash)
+							.DefaultLabel(LOCTEXT("HashColumnHeader", "Hash"))
+							.FillSized(100)
+						+SHeaderRow::Column(UE::Mass::Debugger::UI::Private::ColumnName_Entities)
+							.DefaultLabel(LOCTEXT("EntitiesColumnHeader", "Entities"))
+							.FillSized(60)
+					)
+					.OnGenerateRow_Lambda([this](TSharedPtr<FMassDebuggerArchetypeData> Item, const TSharedPtr<STableViewBase>& OwnerTable)
+						{
+							return SNew(SMassArchetypeTableRow, OwnerTable, Item, DiffBase, DebuggerModel);
+						})
+					.OnGetChildren_Lambda([](TSharedPtr<FMassDebuggerArchetypeData> InItem, TArray<TSharedPtr<FMassDebuggerArchetypeData>>& OutChildren)
+						{
+							if (InItem->Children.Num())
+							{
+								OutChildren.Append(InItem->Children);
+							}
+						})
+					.OnSelectionChanged(this, &SMassArchetypesView::HandleSelectionChanged)
+				]
+
+				+ SVerticalBox::Slot()
+				.AutoHeight()
+				.Padding(3)
+				[
+					SNew(STextBlock)
+					.Text_Lambda([this]()
 					{
-						if (InItem->Children.Num())
+						int32 ActiveCount = 0;
+						int32 TotalCount = 0;
+						for (const TSharedPtr<FMassDebuggerArchetypeData>& ArchetypeData : DebuggerModel->CachedArchetypes)
 						{
-							OutChildren.Append(InItem->Children);
+							TotalCount++;
+							if (ArchetypeData->EntitiesCount > 0)
+							{
+								ActiveCount++;
+							}
+							
+							for (const TSharedPtr<FMassDebuggerArchetypeData>& ChildArchetypeData : ArchetypeData->Children)
+							{
+								TotalCount++;
+								if (ChildArchetypeData->EntitiesCount > 0)
+								{
+									ActiveCount++;
+								}
+							}
 						}
+						
+						return FText::Format(LOCTEXT("ArchetypeInfoFormat", "{0} archetypes, {1} active, {2} selected"),
+							TotalCount, ActiveCount, DebuggerModel->SelectedArchetypes.Num());
 					})
-				.OnSelectionChanged(this, &SMassArchetypesView::HandleSelectionChanged)
+				]					
 			]
 
-			+ SHorizontalBox::Slot()
+			+ SSplitter::Slot()
+			.Value(.6f)
 			[
-				SAssignNew(SelectedArchetypesListWidget, SListView<TSharedPtr<FMassDebuggerArchetypeData>>)
-				.ListItemsSource(&DebuggerModel->SelectedArchetypes)
-				.SelectionMode(ESelectionMode::None)
-				.OnGenerateRow_Lambda([](TSharedPtr<FMassDebuggerArchetypeData> Item, const TSharedPtr<STableViewBase>& OwnerTable)
-					{
-						return SNew(SMassArchetypeDetailTableRow, OwnerTable, Item);
-					})
+				SAssignNew(SelectedArchetypesView, SScrollBox)
 			]
 		];
+
+	RebuildSelectedView();
+}
+
+void SMassArchetypesView::RebuildSelectedView()
+{
+	SelectedArchetypesView->ClearChildren();
+
+	const EMassBitSetDiffPrune Prune = DebuggerModel->SelectedArchetypes.Num() > 1 ? EMassBitSetDiffPrune::Same : EMassBitSetDiffPrune::None;
+	
+	for (const TSharedPtr<FMassDebuggerArchetypeData>& ArchetypeData : DebuggerModel->SelectedArchetypes)
+	{
+		SelectedArchetypesView->AddSlot()
+			[
+				SNew(SMassArchetype, ArchetypeData, DiffBase, Prune)
+			];
+	}
 }
 
 void SMassArchetypesView::HandleSelectionChanged(TSharedPtr<FMassDebuggerArchetypeData> InNode, ESelectInfo::Type InSelectInfo)
@@ -174,32 +261,20 @@ void SMassArchetypesView::HandleSelectionChanged(TSharedPtr<FMassDebuggerArchety
 
 void SMassArchetypesView::OnRefresh()
 {
+	TreeViewSource.Reset();
+	
 	if (DebuggerModel)
 	{
-		TreeViewSource.Reset();
-		
-		TMap<uint32, TSharedPtr<FMassDebuggerArchetypeData>> ArchetypeRepresentativeMap;
-		for (const TSharedPtr<FMassDebuggerArchetypeData>& ArchetypeData : DebuggerModel->CachedArchetypes)
-		{
-			if (TSharedPtr<FMassDebuggerArchetypeData>* Representative = ArchetypeRepresentativeMap.Find(ArchetypeData->FamilyHash))
-			{
-				(*Representative)->Children.Add(ArchetypeData);
-			}
-			else
-			{
-				ArchetypeRepresentativeMap.Add(ArchetypeData->FamilyHash, ArchetypeData);
-			}
-		}
+		TreeViewSource = DebuggerModel->CachedArchetypes;
 
-		for (auto& KVP : ArchetypeRepresentativeMap)
+		for (const TSharedPtr<FMassDebuggerArchetypeData>& ArchetypeData : TreeViewSource)
 		{
-			TreeViewSource.Add(KVP.Value);
-			ArchetypesTreeView->SetItemExpansion(KVP.Value, true);
+			ArchetypesTreeView->SetItemExpansion(ArchetypeData, true);
 		}
 	}
 
 	ArchetypesTreeView->RequestListRefresh();
-	SelectedArchetypesListWidget->RequestListRefresh();
+	RebuildSelectedView();
 }
 
 void SMassArchetypesView::OnProcessorsSelected(TConstArrayView<TSharedPtr<FMassDebuggerProcessorData>> SelectedProcessors, ESelectInfo::Type SelectInfo)
@@ -226,8 +301,19 @@ void SMassArchetypesView::OnArchetypesSelected(TConstArrayView<TSharedPtr<FMassD
 			ArchetypesTreeView->RequestScrollIntoView(SelectedArchetypes[0]);
 		}
 	}
+
+	DiffBase = nullptr;
+	if (SelectedArchetypes.Num() == 1)
+	{
+		DiffBase = SelectedArchetypes[0]->Parent.Pin();
+	}
+	else if (SelectedArchetypes.Num() > 1)
+	{
+		DiffBase = SelectedArchetypes[0];
+	}
+	
 	ArchetypesTreeView->RebuildList();
-	SelectedArchetypesListWidget->RequestListRefresh();
+	RebuildSelectedView();
 }
 
 #undef LOCTEXT_NAMESPACE
diff --git a/Engine/Plugins/Runtime/MassEntity/Source/MassEntityDebugger/Private/SMassArchetypesView.h b/Engine/Plugins/Runtime/MassEntity/Source/MassEntityDebugger/Private/SMassArchetypesView.h
--- a/Engine/Plugins/Runtime/MassEntity/Source/MassEntityDebugger/Private/SMassArchetypesView.h
+++ b/Engine/Plugins/Runtime/MassEntity/Source/MassEntityDebugger/Private/SMassArchetypesView.h
@@ -4,6 +4,8 @@
 #include "SMassDebuggerViewBase.h"
 #include "Widgets/Views/STreeView.h"
 #include "Widgets/Views/SListView.h"
+#include "Widgets/Views/SHeaderRow.h"
+#include "Widgets/Layout/SScrollBox.h"
 
 struct FMassDebuggerModel;
 struct FMassDebuggerArchetypeData;
@@ -18,7 +20,7 @@ public:
 	}
 	SLATE_END_ARGS()
 
-	void Construct(const FArguments& InArgs, TSharedRef<FMassDebuggerModel> DebuggerModel);
+	void Construct(const FArguments& InArgs, TSharedRef<FMassDebuggerModel> InDebuggerModel);
 
 	void HandleSelectionChanged(TSharedPtr<FMassDebuggerArchetypeData> InNode, ESelectInfo::Type InSelectInfo);
 
@@ -27,8 +29,11 @@ protected:
 	virtual void OnProcessorsSelected(TConstArrayView<TSharedPtr<FMassDebuggerProcessorData>> SelectedProcessors, ESelectInfo::Type SelectInfo) override;
 	virtual void OnArchetypesSelected(TConstArrayView<TSharedPtr<FMassDebuggerArchetypeData>> SelectedArchetypes, ESelectInfo::Type SelectInfo) override;
 
+	void RebuildSelectedView();
+	
 	TSharedPtr<STreeView<TSharedPtr<FMassDebuggerArchetypeData>>> ArchetypesTreeView;
-	TSharedPtr<SListView<TSharedPtr<FMassDebuggerArchetypeData>>> SelectedArchetypesListWidget;
-
+	TSharedPtr<SScrollBox> SelectedArchetypesView;
+	TSharedPtr<SHeaderRow> ArchetypesHeaderRow;
+	TSharedPtr<FMassDebuggerArchetypeData> DiffBase;
 	TArray<TSharedPtr<FMassDebuggerArchetypeData>> TreeViewSource;
 };
diff --git a/Engine/Plugins/Runtime/MassEntity/Source/MassEntityDebugger/Private/SMassBitSet.h b/Engine/Plugins/Runtime/MassEntity/Source/MassEntityDebugger/Private/SMassBitSet.h
--- a/Engine/Plugins/Runtime/MassEntity/Source/MassEntityDebugger/Private/SMassBitSet.h
+++ b/Engine/Plugins/Runtime/MassEntity/Source/MassEntityDebugger/Private/SMassBitSet.h
@@ -2,14 +2,17 @@
 #pragma once
 
 #include "Widgets/SCompoundWidget.h"
-#include "Widgets/Text/SRichTextBlock.h"
 #include "Widgets/Text/STextBlock.h"
 #include "MassDebuggerStyle.h"
 #include "Widgets/Layout/SWrapBox.h"
 #include "Widgets/Layout/SBorder.h"
+#include "Widgets/Layout/SBox.h"
+#include "Widgets/Images/SImage.h"
 #include "Widgets/SBoxPanel.h"
 #include "MassProcessingTypes.h"
 
+#define LOCTEXT_NAMESPACE "SMassDebugger"
+
 struct FMassDebuggerQueryData;
 
 enum class EMassBitSetDisplayMode 
@@ -56,8 +59,13 @@ public:
 			+ SHorizontalBox::Slot()
 			.AutoWidth()
 			[
-				SNew(STextBlock)
-				.Text(FText::FromString(Label))
+				SNew(SBox)
+				.MinDesiredWidth(150)
+				.Padding(FMargin(0.0f, InArgs._SlotPadding.Get().Top, 0.0f, InArgs._SlotPadding.Get().Bottom))
+				[
+					SNew(STextBlock)
+					.Text(FText::FromString(Label))
+				]
 			]
 			+ SHorizontalBox::Slot()
 			[
@@ -91,6 +99,155 @@ protected:
 	}
 };
 
+
+enum class EMassBitSetDiffPrune : uint8
+{
+	None,	// No pruning
+	Same,	// Prune tags that exists in base
+};
+
+template<typename TBitSet>
+class SMassBitSetDiff : public SCompoundWidget
+{
+public:
+	SLATE_BEGIN_ARGS(SMassBitSetDiff)
+		: _SlotPadding(5.f)
+		, _TextColor(FLinearColor(1.f, 1.f, 1.f))
+		, _BackgroundBrush(FMassDebuggerStyle::GetBrush("MassDebug.Fragment"))
+	{}
+		SLATE_ATTRIBUTE(FMargin, SlotPadding)
+		SLATE_ATTRIBUTE(FLinearColor, TextColor)
+		SLATE_ATTRIBUTE(const FSlateBrush*, BackgroundBrush)
+	SLATE_END_ARGS()
+
+	void Construct(const SMassBitSetDiff::FArguments& InArgs, const FString& Label, const TBitSet& BaseBitSet, const TBitSet& BitSet, const EMassBitSetDiffPrune Prune)
+	{
+		TSharedRef<SWrapBox> ButtonBox = SNew(SWrapBox).UseAllottedSize(true);
+		
+		const FSlateBrush* Brush = InArgs._BackgroundBrush.Get();
+		AddBitSet(InArgs, ButtonBox, BaseBitSet, BitSet, Brush, Prune);
+
+		ChildSlot
+		[
+			SNew(SHorizontalBox)
+			+ SHorizontalBox::Slot()
+			.AutoWidth()
+			[
+				SNew(SBox)
+				.MinDesiredWidth(150)
+				.Padding(FMargin(0.0f, InArgs._SlotPadding.Get().Top, 0.0f, InArgs._SlotPadding.Get().Bottom))
+				[
+					SNew(STextBlock)
+					.Text(FText::FromString(Label))
+				]
+			]
+			+ SHorizontalBox::Slot()
+			[
+				ButtonBox
+			]
+		];
+	}
+
+protected:
+	void AddBitSet(const SMassBitSetDiff::FArguments& InArgs, TSharedRef<SWrapBox>& ButtonBox, const TBitSet& BaseBitSet, const TBitSet& BitSet, const FSlateBrush* Brush, const EMassBitSetDiffPrune Prune)
+	{
+#if WITH_MASSENTITY_DEBUG
+		TBitSet CombinedBitSet = BaseBitSet + BitSet;
+		int32 SameCount = 0;
+
+		for (auto It = CombinedBitSet.GetIndexIterator(); It; ++It)
+		{
+			const bool bInBase = BaseBitSet.IsBitSet(*It);
+			const bool bInCurr = BitSet.IsBitSet(*It);
+
+			if (Prune == EMassBitSetDiffPrune::Same && bInBase && bInCurr)
+			{
+				SameCount++;
+				continue;
+			}
+			
+			FString Name = CombinedBitSet.DebugGetStructTypeName(*It).ToString();
+			FLinearColor BorderColor = FLinearColor::White;
+			FLinearColor TextColor = InArgs._TextColor.Get();
+
+			FText Tooltip;
+
+			const FSlateBrush* Icon = nullptr;
+			
+			if (bInBase && bInCurr)
+			{
+				// Same
+				BorderColor.A *= 0.6f;
+				TextColor.A *= 0.75f;
+				Tooltip = LOCTEXT("DiffExistsInBase", "Exists in base.");
+			}
+			else if (!bInBase && bInCurr)
+			{
+				// Added
+				TextColor = FMath::Lerp(TextColor, FLinearColor::White, 0.9f);
+				Icon = FAppStyle::Get().GetBrush("Icons.Plus");
+				Tooltip = LOCTEXT("DiffAddedToBase", "Added to base.");
+			}
+			else
+			{
+				// Removed
+				BorderColor.A *= 0.25f;
+				TextColor.A *= 0.5f;
+				Icon = FAppStyle::Get().GetBrush("Icons.Minus");
+				Tooltip = LOCTEXT("DiffRemovedFromBase", "Removed from base.");
+			}
+			
+			ButtonBox->AddSlot()
+				.Padding(InArgs._SlotPadding)
+				[
+					SNew(SBorder)
+					.BorderImage(Brush)
+					.BorderBackgroundColor(BorderColor)
+					.ColorAndOpacity(TextColor)
+					.ToolTipText(Tooltip)
+					[
+						SNew(SHorizontalBox)
+						+ SHorizontalBox::Slot()
+						.AutoWidth()
+						.Padding(3,0)
+						.VAlign(VAlign_Center)
+						[
+							SNew(SImage)
+							.Image(Icon)
+							.Visibility_Lambda([Icon] { return Icon != nullptr ? EVisibility::Visible : EVisibility::Collapsed; })
+							.DesiredSizeOverride(FVector2D(10.0f, 10.0f))
+						]
+						+ SHorizontalBox::Slot()
+						.AutoWidth()
+						.VAlign(VAlign_Center)
+						[
+							SNew(STextBlock)
+							.Text(FText::FromString(Name))
+						]
+					]
+				];
+		}
+
+		if (SameCount > 0)
+		{
+			ButtonBox->AddSlot()
+			.Padding(InArgs._SlotPadding)
+			[
+				SNew(SHorizontalBox)
+				+ SHorizontalBox::Slot()
+				.AutoWidth()
+				.VAlign(VAlign_Center)
+				[
+					SNew(STextBlock)
+					.Text(FText::Format(LOCTEXT("SameFragmentsFormat", "{0} same fragments"), FText::AsNumber(SameCount)))
+				]
+			];
+		}
+
+#endif // WITH_MASSENTITY_DEBUG
+	}
+};
+
 namespace UE::Mass::Debugger::UI
 {
 	template<typename TBitSet>
@@ -100,6 +257,7 @@ namespace UE::Mass::Debugger::UI
 		{
 			Box->AddSlot()
 			.AutoHeight()
+			.Padding(0, 8)
 			[
 				SNew(SMassBitSet<TBitSet>, Label, BitSetAccess)
 				.BackgroundBrush(Brush)
@@ -107,4 +265,23 @@ namespace UE::Mass::Debugger::UI
 			];
 		}
 	}
-}
+
+	template<typename TBitSet>
+	void AddBitSetDiff(TSharedRef<SVerticalBox>& Box, const TBitSet& BaseBitSet, const TBitSet& BitSet, const FString& Label, const FSlateBrush* Brush, const EMassBitSetDiffPrune Prune)
+	{
+		if (BaseBitSet.IsEmpty() == false || BitSet.IsEmpty() == false)
+		{
+			Box->AddSlot()
+			.AutoHeight()
+			.Padding(0, 8)
+			[
+				SNew(SMassBitSetDiff<TBitSet>, Label, BaseBitSet, BitSet, Prune)
+				.BackgroundBrush(Brush)
+				.SlotPadding(5)
+			];
+		}
+	}
+
+} // UE::Mass::Debugger::UI
+
+#undef LOCTEXT_NAMESPACE
\ No newline at end of file
