[Mass] added some LLM scopes for processor execution, command creation, command flush and phase delegates

diff --git a/Engine/Plugins/Runtime/MassEntity/Source/MassEntity/Private/MassCommandBuffer.cpp b/Engine/Plugins/Runtime/MassEntity/Source/MassEntity/Private/MassCommandBuffer.cpp
--- a/Engine/Plugins/Runtime/MassEntity/Source/MassEntity/Private/MassCommandBuffer.cpp
+++ b/Engine/Plugins/Runtime/MassEntity/Source/MassEntity/Private/MassCommandBuffer.cpp
@@ -96,6 +96,7 @@ void FMassCommandBuffer::Flush(UMassEntitySubsystem& EntitySystem)
 
 	{
 		UE_MT_SCOPED_WRITE_ACCESS(PendingBatchCommandsDetector);
+		LLM_SCOPE_BYNAME(TEXT("Mass/FlushCommands"));
 
 		// array used to group commands depending on their operations. Based on EMassCommandOperationType
 		// @todo I'm opened to suggestions on how to better implement this
diff --git a/Engine/Plugins/Runtime/MassEntity/Source/MassEntity/Private/MassProcessingPhase.cpp b/Engine/Plugins/Runtime/MassEntity/Source/MassEntity/Private/MassProcessingPhase.cpp
--- a/Engine/Plugins/Runtime/MassEntity/Source/MassEntity/Private/MassProcessingPhase.cpp
+++ b/Engine/Plugins/Runtime/MassEntity/Source/MassEntity/Private/MassProcessingPhase.cpp
@@ -58,8 +58,10 @@ void FMassProcessingPhase::ExecuteTick(float DeltaTime, ELevelTick TickType, ENa
 	TRACE_CPUPROFILER_EVENT_SCOPE_STR(*FString::Printf(TEXT("FMassProcessingPhase::ExecuteTick %s"), *UEnum::GetValueAsString(Phase)));
 
 	Manager->OnPhaseStart(*this);
-	
-	OnPhaseStart.Broadcast(DeltaTime);
+	{
+		LLM_SCOPE_BYNAME(TEXT("Mass/PhaseStartDelegate"));
+		OnPhaseStart.Broadcast(DeltaTime);
+	}
 
 	check(PhaseProcessor);
 	
@@ -93,7 +95,10 @@ void FMassProcessingPhase::ExecuteTick(float DeltaTime, ELevelTick TickType, ENa
 	{
 		UE::Mass::Executor::Run(*PhaseProcessor, Context);
 
-		OnPhaseEnd.Broadcast(DeltaTime);
+		{
+			LLM_SCOPE_BYNAME(TEXT("Mass/PhaseEndDelegate"));
+			OnPhaseEnd.Broadcast(DeltaTime);
+		}
 		Manager->OnPhaseEnd(*this);
 		bIsDuringMassProcessing = false;
 	}
@@ -102,7 +107,10 @@ void FMassProcessingPhase::ExecuteTick(float DeltaTime, ELevelTick TickType, ENa
 void FMassProcessingPhase::OnParallelExecutionDone(const float DeltaTime)
 {
 	bIsDuringMassProcessing = false;
-	OnPhaseEnd.Broadcast(DeltaTime);
+	{
+		LLM_SCOPE_BYNAME(TEXT("Mass/PhaseEndDelegate"));
+		OnPhaseEnd.Broadcast(DeltaTime);
+	}
 	check(Manager);
 	Manager->OnPhaseEnd(*this);
 }
diff --git a/Engine/Plugins/Runtime/MassEntity/Source/MassEntity/Private/MassProcessor.cpp b/Engine/Plugins/Runtime/MassEntity/Source/MassEntity/Private/MassProcessor.cpp
--- a/Engine/Plugins/Runtime/MassEntity/Source/MassEntity/Private/MassProcessor.cpp
+++ b/Engine/Plugins/Runtime/MassEntity/Source/MassEntity/Private/MassProcessor.cpp
@@ -154,6 +154,9 @@ void UMassProcessor::PostInitProperties()
 void UMassProcessor::CallExecute(UMassEntitySubsystem& EntitySubsystem, FMassExecutionContext& Context)
 {
 	TRACE_CPUPROFILER_EVENT_SCOPE_TEXT(*StatId);
+	LLM_SCOPE_BYNAME(TEXT("Mass/ExecuteProcessor"));
+	// Not using a more specific scope by default (i.e. LLM_SCOPE_BYNAME(*StatId)) since LLM is more strict regarding the provided string (no spaces or '_')
+
 #if WITH_MASSENTITY_DEBUG
 	Context.DebugSetExecutionDesc(FString::Printf(TEXT("%s (%s)"), *GetProcessorName(), *ToString(EntitySubsystem.GetWorld()->GetNetMode())));
 #endif
diff --git a/Engine/Plugins/Runtime/MassEntity/Source/MassEntity/Public/MassCommandBuffer.h b/Engine/Plugins/Runtime/MassEntity/Source/MassEntity/Public/MassCommandBuffer.h
--- a/Engine/Plugins/Runtime/MassEntity/Source/MassEntity/Public/MassCommandBuffer.h
+++ b/Engine/Plugins/Runtime/MassEntity/Source/MassEntity/Public/MassCommandBuffer.h
@@ -22,6 +22,7 @@ public:
 	template< template<typename... TArgs> typename TCommand, typename... TArgs >
 	void PushCommand(const FMassEntityHandle Entity, TArgs&&... InArgs)
 	{
+		LLM_SCOPE_BYNAME(TEXT("Mass/PushCommand"));
 		TCommand<TArgs...>& Instance = CreateOrAddCommand<TCommand<TArgs...>>();
 		Instance.Add(Entity, Forward<TArgs>(InArgs)...);
 	}
@@ -29,6 +30,7 @@ public:
 	template<typename TCommand, typename... TArgs>
 	void PushCommand(TArgs&&... InArgs)
 	{
+		LLM_SCOPE_BYNAME(TEXT("Mass/PushCommand"));
 		TCommand& Instance = CreateOrAddCommand<TCommand>();
 		Instance.Add(Forward<TArgs>(InArgs)...);
 	}
@@ -37,12 +39,14 @@ public:
 	template< typename TCommand>
 	void PushCommand(const FMassEntityHandle Entity)
 	{
+		LLM_SCOPE_BYNAME(TEXT("Mass/PushCommand"));
 		CreateOrAddCommand<TCommand>().Add(Entity);
 	}
 
 	template< typename TCommand>
 	void PushCommand(TConstArrayView<FMassEntityHandle> Entities)
 	{
+		LLM_SCOPE_BYNAME(TEXT("Mass/PushCommand"));
 		CreateOrAddCommand<TCommand>().Add(Entities);
 	}
 
