Moved shared fragment from being unique per archetype to unique per chunk. Saves lots of unnecessary archetypes.
Archetype are there for queries to filter on fragment types and we do not need to filter on the shared fragment values.
Performance wize, it saved 0.5ms over 5.7ms locally.

diff --git a/Engine/Plugins/Runtime/MassEntity/Source/MassEntity/Private/MassArchetypeData.cpp b/Engine/Plugins/Runtime/MassEntity/Source/MassEntity/Private/MassArchetypeData.cpp
--- a/Engine/Plugins/Runtime/MassEntity/Source/MassEntity/Private/MassArchetypeData.cpp
+++ b/Engine/Plugins/Runtime/MassEntity/Source/MassEntity/Private/MassArchetypeData.cpp
@@ -20,7 +20,7 @@ bool FMassArchetypeData::HasFragmentType(const UScriptStruct* FragmentType) cons
 	return (FragmentType && CompositionDescriptor.Fragments.Contains(*FragmentType));
 }
 
-void FMassArchetypeData::Initialize(const FMassArchetypeCompositionDescriptor& InCompositionDescriptor, const FMassArchetypeSharedFragmentValues& InSharedFragmentValues)
+void FMassArchetypeData::Initialize(const FMassArchetypeCompositionDescriptor& InCompositionDescriptor)
 {
 	CompositionDescriptor.Fragments = InCompositionDescriptor.Fragments;
 	ConfigureFragments();
@@ -41,7 +41,6 @@ void FMassArchetypeData::Initialize(const FMassArchetypeCompositionDescriptor& I
 
 	// Share fragments
 	CompositionDescriptor.SharedFragments = InCompositionDescriptor.SharedFragments;
-	SharedFragmentValues = InSharedFragmentValues;
 
 	EntityListOffsetWithinChunk = 0;
 }
@@ -64,7 +63,6 @@ void FMassArchetypeData::InitializeWithSimilar(const FMassArchetypeData& BaseArc
 		NumEntitiesPerChunk = BaseArchetype.NumEntitiesPerChunk;
 	}
 	ChunkFragmentsTemplate = BaseArchetype.ChunkFragmentsTemplate;
-	SharedFragmentValues = BaseArchetype.GetSharedFragmentValues();
 
 	EntityListOffsetWithinChunk = 0;
 }
@@ -118,9 +116,9 @@ void FMassArchetypeData::ConfigureFragments()
 	}
 }
 
-void FMassArchetypeData::AddEntity(FMassEntityHandle Entity)
+void FMassArchetypeData::AddEntity(FMassEntityHandle Entity, const FMassArchetypeSharedFragmentValues& SharedFragmentValues)
 {
-	const int32 AbsoluteIndex = AddEntityInternal(Entity);
+	const int32 AbsoluteIndex = AddEntityInternal(Entity, SharedFragmentValues);
 
 	// Initialize fragments
 	const int32 ChunkIndex = AbsoluteIndex / NumEntitiesPerChunk;
@@ -133,8 +131,11 @@ void FMassArchetypeData::AddEntity(FMassEntityHandle Entity)
 	}
 }
 
-int32 FMassArchetypeData::AddEntityInternal(FMassEntityHandle Entity)
+int32 FMassArchetypeData::AddEntityInternal(FMassEntityHandle Entity, const FMassArchetypeSharedFragmentValues& SharedFragmentValues)
 {
+	checkf(SharedFragmentValues.IsSorted(), TEXT("Expecting shared fragment values to be previously sorted"));
+	checkfSlow(SharedFragmentValues.HasExactFragmentTypesMatch(CompositionDescriptor.SharedFragments), TEXT("Expecting values for every specified shared fragment in the archetype and only those"))
+
 	int32 IndexWithinChunk = 0;
 	int32 AbsoluteIndex = 0;
 	int32 ChunkIndex = 0;
@@ -155,7 +156,7 @@ int32 FMassArchetypeData::AddEntityInternal(FMassEntityHandle Entity)
 				EmptyAbsoluteIndex = AbsoluteIndex;
 			}
 		}
-		else if (Chunk.GetNumInstances() < NumEntitiesPerChunk)
+		else if (Chunk.GetNumInstances() < NumEntitiesPerChunk && Chunk.GetSharedFragmentValues().IsEquivalent(SharedFragmentValues))
 		{
 			IndexWithinChunk = Chunk.GetNumInstances();
 			AbsoluteIndex += IndexWithinChunk;
@@ -175,12 +176,12 @@ int32 FMassArchetypeData::AddEntityInternal(FMassEntityHandle Entity)
 		if (EmptyChunkIndex != INDEX_NONE)
 		{
 			DestinationChunk = &Chunks[EmptyChunkIndex];
-			DestinationChunk->Recycle(ChunkFragmentsTemplate);
+			DestinationChunk->Recycle(ChunkFragmentsTemplate, SharedFragmentValues);
 			AbsoluteIndex = EmptyAbsoluteIndex;
 		}
 		else
 		{
-			DestinationChunk = &Chunks.Emplace_GetRef(GetChunkAllocSize(), ChunkFragmentsTemplate);
+			DestinationChunk = &Chunks.Emplace_GetRef(GetChunkAllocSize(), ChunkFragmentsTemplate, SharedFragmentValues);
 		}
 
 		check(DestinationChunk);
@@ -361,7 +362,7 @@ void FMassArchetypeData::MoveEntityToAnotherArchetype(const FMassEntityHandle En
 	const int32 IndexWithinChunk = AbsoluteIndex % NumEntitiesPerChunk;
 	FMassArchetypeChunk& Chunk = Chunks[ChunkIndex];
 
-	const int32 NewAbsoluteIndex = NewArchetype.AddEntityInternal(Entity);
+	const int32 NewAbsoluteIndex = NewArchetype.AddEntityInternal(Entity, Chunk.GetSharedFragmentValues());
 	const int32 NewChunkIndex = NewAbsoluteIndex / NewArchetype.NumEntitiesPerChunk;
 	const int32 NewIndexWithinChunk = NewAbsoluteIndex % NewArchetype.NumEntitiesPerChunk;
 	FMassArchetypeChunk& NewChunk = NewArchetype.Chunks[NewChunkIndex];
@@ -376,13 +377,19 @@ void FMassArchetypeData::ExecuteFunction(FMassExecutionContext& RunContext, cons
 	// mz@todo to be removed
 	RunContext.SetCurrentArchetypesTagBitSet(GetTagBitSet());
 
-	BindConstSharedFragmentRequirements(RunContext, RequirementMapping.ConstSharedFragments);
-	BindSharedFragmentRequirements(RunContext, RequirementMapping.SharedFragments);
-
+	uint32 PrevSharedFragmentValuesHash = UINT32_MAX;
 	for (FMassArchetypeChunkIterator ChunkIterator(EntityRangeContainer); ChunkIterator; ++ChunkIterator)
 	{
 		FMassArchetypeChunk& Chunk = Chunks[ChunkIterator->ChunkIndex];
-		
+
+		const uint32 SharedFragmentValuesHash = GetTypeHash(Chunk.GetSharedFragmentValues());
+		if (PrevSharedFragmentValuesHash != SharedFragmentValuesHash)
+		{
+			PrevSharedFragmentValuesHash = SharedFragmentValuesHash;
+			BindConstSharedFragmentRequirements(RunContext, Chunk.GetSharedFragmentValues(), RequirementMapping.ConstSharedFragments);
+			BindSharedFragmentRequirements(RunContext, Chunk.GetSharedFragmentValues(), RequirementMapping.SharedFragments);
+		}
+
 		const int32 ChunkLength = ChunkIterator->Length > 0 ? ChunkIterator->Length : (Chunk.GetNumInstances() - ChunkIterator->SubchunkStart);
 		if (ChunkLength)
 		{
@@ -402,13 +409,19 @@ void FMassArchetypeData::ExecuteFunction(FMassExecutionContext& RunContext, cons
 	// mz@todo to be removed
 	RunContext.SetCurrentArchetypesTagBitSet(GetTagBitSet());
 
-	BindConstSharedFragmentRequirements(RunContext, RequirementMapping.ConstSharedFragments);
-	BindSharedFragmentRequirements(RunContext, RequirementMapping.SharedFragments);
-
-	if(!ArchetypeCondition || ArchetypeCondition(RunContext))
+	if (!ArchetypeCondition || ArchetypeCondition(RunContext))
 	{
+		uint32 PrevSharedFragmentValuesHash = UINT32_MAX;
 		for (FMassArchetypeChunk& Chunk : Chunks)
 		{
+			const uint32 SharedFragmentValuesHash = GetTypeHash(Chunk.GetSharedFragmentValues());
+			if (PrevSharedFragmentValuesHash != SharedFragmentValuesHash)
+			{
+				PrevSharedFragmentValuesHash = SharedFragmentValuesHash;
+				BindConstSharedFragmentRequirements(RunContext, Chunk.GetSharedFragmentValues(), RequirementMapping.ConstSharedFragments);
+				BindSharedFragmentRequirements(RunContext, Chunk.GetSharedFragmentValues(), RequirementMapping.SharedFragments);
+			}
+
 			if (Chunk.GetNumInstances())
 			{
 				RunContext.SetCurrentChunkSerialModificationNumber(Chunk.GetSerialModificationNumber());
@@ -431,8 +444,8 @@ void FMassArchetypeData::ExecutionFunctionForChunk(FMassExecutionContext RunCont
 
 	if (ChunkLength)
 	{
-		BindConstSharedFragmentRequirements(RunContext, RequirementMapping.ChunkFragments);
-		BindSharedFragmentRequirements(RunContext, RequirementMapping.ChunkFragments);
+		BindConstSharedFragmentRequirements(RunContext, Chunk.GetSharedFragmentValues(), RequirementMapping.ChunkFragments);
+		BindSharedFragmentRequirements(RunContext, Chunk.GetSharedFragmentValues(), RequirementMapping.ChunkFragments);
 
 		RunContext.SetCurrentArchetypesTagBitSet(GetTagBitSet());
 		RunContext.SetCurrentChunkSerialModificationNumber(Chunk.GetSerialModificationNumber());
@@ -450,73 +463,80 @@ void FMassArchetypeData::CompactEntities(const double TimeAllowed)
 {
 	const double TimeAllowedEnd = FPlatformTime::Seconds() + TimeAllowed;
 
-	TArray<FMassArchetypeChunk*> SortedChunks;
+	TMap<uint32, TArray<FMassArchetypeChunk*>> SortedChunksBySharedValues;
 	for (FMassArchetypeChunk& Chunk : Chunks)
 	{
 		// Skip already full chunks
 		const int32 NumInstances = Chunk.GetNumInstances();
 		if (NumInstances > 0 && NumInstances < NumEntitiesPerChunk)
 		{
+			const uint32 SharedFragmentHash = GetTypeHash(Chunk.GetSharedFragmentValues());
+			TArray<FMassArchetypeChunk*>& SortedChunks = SortedChunksBySharedValues.FindOrAddByHash(SharedFragmentHash, SharedFragmentHash, TArray<FMassArchetypeChunk*>());
 			SortedChunks.Add(&Chunk);
 		}
 	}
 
-	// Check if there is anything to compact at all
-	if (SortedChunks.Num() <= 1)
-	{
-		return;
-	}
-
-	SortedChunks.Sort([](const FMassArchetypeChunk& LHS, const FMassArchetypeChunk& RHS)
+	for (TPair<uint32, TArray<FMassArchetypeChunk*>>& Pair : SortedChunksBySharedValues)
 	{
-		return LHS.GetNumInstances() < RHS.GetNumInstances();
-	});
+		TArray<FMassArchetypeChunk*>& SortedChunks = Pair.Value;
 
-	int32 ChunkToFillSortedIdx = 0;
-	int32 ChunkToEmptySortedIdx = SortedChunks.Num() - 1;
-	while (ChunkToFillSortedIdx < ChunkToEmptySortedIdx && FPlatformTime::Seconds() < TimeAllowedEnd)
-	{
-		while (ChunkToFillSortedIdx < SortedChunks.Num() && SortedChunks[ChunkToFillSortedIdx]->GetNumInstances() == NumEntitiesPerChunk)
+		// Check if there is anything to compact at all
+		if (SortedChunks.Num() <= 1)
 		{
-			ChunkToFillSortedIdx++;
+			continue;
 		}
-		while (ChunkToEmptySortedIdx >= 0 && SortedChunks[ChunkToEmptySortedIdx]->GetNumInstances() == 0)
+
+		SortedChunks.Sort([](const FMassArchetypeChunk& LHS, const FMassArchetypeChunk& RHS)
 		{
-			ChunkToEmptySortedIdx--;
-		}
-		if (ChunkToFillSortedIdx >= ChunkToEmptySortedIdx)
+			return LHS.GetNumInstances() < RHS.GetNumInstances();
+		});
+
+		int32 ChunkToFillSortedIdx = 0;
+		int32 ChunkToEmptySortedIdx = SortedChunks.Num() - 1;
+		while (ChunkToFillSortedIdx < ChunkToEmptySortedIdx && FPlatformTime::Seconds() < TimeAllowedEnd)
 		{
-			break;
-		}
+			while (ChunkToFillSortedIdx < SortedChunks.Num() && SortedChunks[ChunkToFillSortedIdx]->GetNumInstances() == NumEntitiesPerChunk)
+			{
+				ChunkToFillSortedIdx++;
+			}
+			while (ChunkToEmptySortedIdx >= 0 && SortedChunks[ChunkToEmptySortedIdx]->GetNumInstances() == 0)
+			{
+				ChunkToEmptySortedIdx--;
+			}
+			if (ChunkToFillSortedIdx >= ChunkToEmptySortedIdx)
+			{
+				break;
+			}
 
-		FMassArchetypeChunk* ChunkToFill = SortedChunks[ChunkToFillSortedIdx];
-		FMassArchetypeChunk* ChunkToEmpty = SortedChunks[ChunkToEmptySortedIdx];
-		const int32 NumberOfEntitiesToMove = FMath::Min(NumEntitiesPerChunk - ChunkToFill->GetNumInstances(), ChunkToEmpty->GetNumInstances());
-		const int32 FromIndex = ChunkToEmpty->GetNumInstances() - NumberOfEntitiesToMove;
-		const int32 ToIndex = ChunkToFill->GetNumInstances();
-		check(NumberOfEntitiesToMove > 0);
+			FMassArchetypeChunk* ChunkToFill = SortedChunks[ChunkToFillSortedIdx];
+			FMassArchetypeChunk* ChunkToEmpty = SortedChunks[ChunkToEmptySortedIdx];
+			const int32 NumberOfEntitiesToMove = FMath::Min(NumEntitiesPerChunk - ChunkToFill->GetNumInstances(), ChunkToEmpty->GetNumInstances());
+			const int32 FromIndex = ChunkToEmpty->GetNumInstances() - NumberOfEntitiesToMove;
+			const int32 ToIndex = ChunkToFill->GetNumInstances();
+			check(NumberOfEntitiesToMove > 0);
 
-		for (const FMassArchetypeFragmentConfig& FragmentConfig : FragmentConfigs)
-		{
-			void* FromFragmentPtr = FragmentConfig.GetFragmentData(ChunkToEmpty->GetRawMemory(), FromIndex);
-			void* ToFragmentPtr = FragmentConfig.GetFragmentData(ChunkToFill->GetRawMemory(), ToIndex);
-			// Move all entries
-			FMemory::Memcpy(ToFragmentPtr, FromFragmentPtr, FragmentConfig.FragmentType->GetStructureSize() * NumberOfEntitiesToMove);
-		}
+			for (const FMassArchetypeFragmentConfig& FragmentConfig : FragmentConfigs)
+			{
+				void* FromFragmentPtr = FragmentConfig.GetFragmentData(ChunkToEmpty->GetRawMemory(), FromIndex);
+				void* ToFragmentPtr = FragmentConfig.GetFragmentData(ChunkToFill->GetRawMemory(), ToIndex);
+				// Move all entries
+				FMemory::Memcpy(ToFragmentPtr, FromFragmentPtr, FragmentConfig.FragmentType->GetStructureSize() * NumberOfEntitiesToMove);
+			}
 
-		FMassEntityHandle* FromEntity = &ChunkToEmpty->GetEntityArrayElementRef(EntityListOffsetWithinChunk, FromIndex);
-		FMassEntityHandle* ToEntity = &ChunkToFill->GetEntityArrayElementRef(EntityListOffsetWithinChunk, ToIndex);
-		FMemory::Memcpy(ToEntity, FromEntity, NumberOfEntitiesToMove * sizeof(FMassEntityHandle));
-		ChunkToFill->AddMultipleInstances(NumberOfEntitiesToMove);
-		ChunkToEmpty->RemoveMultipleInstances(NumberOfEntitiesToMove);
+			FMassEntityHandle* FromEntity = &ChunkToEmpty->GetEntityArrayElementRef(EntityListOffsetWithinChunk, FromIndex);
+			FMassEntityHandle* ToEntity = &ChunkToFill->GetEntityArrayElementRef(EntityListOffsetWithinChunk, ToIndex);
+			FMemory::Memcpy(ToEntity, FromEntity, NumberOfEntitiesToMove * sizeof(FMassEntityHandle));
+			ChunkToFill->AddMultipleInstances(NumberOfEntitiesToMove);
+			ChunkToEmpty->RemoveMultipleInstances(NumberOfEntitiesToMove);
 
-		const int32 ChunkToFillIdx = UE_PTRDIFF_TO_INT32(ChunkToFill - &Chunks[0]);
-		check(ChunkToFillIdx >= 0 && ChunkToFillIdx < Chunks.Num());
-		const int32 AbsoluteIndex = ChunkToFillIdx * NumEntitiesPerChunk + ToIndex;
+			const int32 ChunkToFillIdx = UE_PTRDIFF_TO_INT32(ChunkToFill - &Chunks[0]);
+			check(ChunkToFillIdx >= 0 && ChunkToFillIdx < Chunks.Num());
+			const int32 AbsoluteIndex = ChunkToFillIdx * NumEntitiesPerChunk + ToIndex;
 
-		for (int32 i = 0; i < NumberOfEntitiesToMove; i++, ++ToEntity)
-		{
-			EntityMap.FindChecked(ToEntity->Index) = AbsoluteIndex + i;
+			for (int32 i = 0; i < NumberOfEntitiesToMove; i++, ++ToEntity)
+			{
+				EntityMap.FindChecked(ToEntity->Index) = AbsoluteIndex + i;
+			}
 		}
 	}
 }
@@ -563,6 +583,13 @@ void FMassArchetypeData::GetRequirementsChunkFragmentMapping(TConstArrayView<FMa
 
 void FMassArchetypeData::GetRequirementsConstSharedFragmentMapping(TConstArrayView<FMassFragmentRequirementDescription> Requirements, FMassFragmentIndicesMapping& OutFragmentIndices) const
 {
+	if (Chunks.Num() == 0)
+	{
+		return;
+	}
+	// All shared fragment values for this archetype should have deterministic indices, so anyone will work to calculate them
+	const FMassArchetypeSharedFragmentValues& SharedFragmentValues = Chunks[0].GetSharedFragmentValues();
+
 	OutFragmentIndices.Reset(Requirements.Num());
 	for (const FMassFragmentRequirementDescription& Requirement : Requirements)
 	{
@@ -577,6 +604,14 @@ void FMassArchetypeData::GetRequirementsConstSharedFragmentMapping(TConstArrayVi
 
 void FMassArchetypeData::GetRequirementsSharedFragmentMapping(TConstArrayView<FMassFragmentRequirementDescription> Requirements, FMassFragmentIndicesMapping& OutFragmentIndices) const
 {
+	if (Chunks.Num() == 0)
+	{
+		return;
+	}
+
+	// All shared fragment values for this archetype should have deterministic indices, so anyone will work to calculate them
+	const FMassArchetypeSharedFragmentValues& SharedFragmentValues = Chunks[0].GetSharedFragmentValues();
+
 	OutFragmentIndices.Reset(Requirements.Num());
 	for (const FMassFragmentRequirementDescription& Requirement : Requirements)
 	{
@@ -663,7 +698,7 @@ void FMassArchetypeData::BindChunkFragmentRequirements(FMassExecutionContext& Ru
 	}
 }
 
-void FMassArchetypeData::BindConstSharedFragmentRequirements(FMassExecutionContext& RunContext, const FMassFragmentIndicesMapping& FragmentsMapping)
+void FMassArchetypeData::BindConstSharedFragmentRequirements(FMassExecutionContext& RunContext, const FMassArchetypeSharedFragmentValues& SharedFragmentValues, const FMassFragmentIndicesMapping& FragmentsMapping)
 {
 	if (FragmentsMapping.Num() > 0)
 	{
@@ -689,7 +724,7 @@ void FMassArchetypeData::BindConstSharedFragmentRequirements(FMassExecutionConte
 	}
 }
 
-void FMassArchetypeData::BindSharedFragmentRequirements(FMassExecutionContext& RunContext, const FMassFragmentIndicesMapping& FragmentsMapping)
+void FMassArchetypeData::BindSharedFragmentRequirements(FMassExecutionContext& RunContext, const FMassArchetypeSharedFragmentValues& SharedFragmentValues, const FMassFragmentIndicesMapping& FragmentsMapping)
 {
 	if (FragmentsMapping.Num() > 0)
 	{
@@ -727,7 +762,6 @@ SIZE_T FMassArchetypeData::GetAllocatedSize() const
 	}
 
 	return sizeof(FMassArchetypeData) +
-		SharedFragmentValues.GetAllocatedSize() +
 		ChunkFragmentsTemplate.GetAllocatedSize() +
 		FragmentConfigs.GetAllocatedSize() +
 		Chunks.GetAllocatedSize() +
@@ -872,14 +906,14 @@ void FMassArchetypeData::REMOVEME_GetArrayViewForFragmentInChunk(int32 ChunkInde
 //////////////////////////////////////////////////////////////////////
 // FMassArchetypeData batched api
 
-void FMassArchetypeData::BatchAddEntities(TConstArrayView<FMassEntityHandle> Entities, TArray<FMassArchetypeEntityCollection::FArchetypeEntityRange>& OutNewRanges)
+void FMassArchetypeData::BatchAddEntities(TConstArrayView<FMassEntityHandle> Entities, const FMassArchetypeSharedFragmentValues& SharedFragmentValues, TArray<FMassArchetypeEntityCollection::FArchetypeEntityRange>& OutNewRanges)
 {
 	FMassArchetypeEntityCollection::FArchetypeEntityRange ResultSubchunk;
 	ResultSubchunk.ChunkIndex = 0;
 	int32 NumberMoved = 0;
 	do 
 	{
-		ResultSubchunk = PrepareNextEntitiesSpanInternal(MakeArrayView(Entities.GetData() + NumberMoved, Entities.Num() - NumberMoved), ResultSubchunk.ChunkIndex);
+		ResultSubchunk = PrepareNextEntitiesSpanInternal(MakeArrayView(Entities.GetData() + NumberMoved, Entities.Num() - NumberMoved), SharedFragmentValues, ResultSubchunk.ChunkIndex);
 		check(Chunks.IsValidIndex(ResultSubchunk.ChunkIndex) && Chunks[ResultSubchunk.ChunkIndex].IsValidSubChunk(ResultSubchunk.SubchunkStart, ResultSubchunk.Length));
 		
 		for (const FMassArchetypeFragmentConfig& FragmentConfig : FragmentConfigs)
@@ -932,6 +966,7 @@ void FMassArchetypeData::BatchMoveEntitiesToAnotherArchetype(const FMassArchetyp
 			const int32 IndexWithinChunk = EntityRange.SubchunkStart + NumberMoved;
 
 			ResultSubChunk = NewArchetype.PrepareNextEntitiesSpanInternal(MakeArrayView(DyingEntityPtr + NumberMoved, EntityRange.Length - NumberMoved)
+				, Chunk.GetSharedFragmentValues()
 				, ResultSubChunk.ChunkIndex);
 
 			FMassArchetypeChunk& NewChunk = NewArchetype.Chunks[ResultSubChunk.ChunkIndex];
@@ -957,8 +992,11 @@ void FMassArchetypeData::BatchMoveEntitiesToAnotherArchetype(const FMassArchetyp
 	}
 }
 
-FMassArchetypeEntityCollection::FArchetypeEntityRange FMassArchetypeData::PrepareNextEntitiesSpanInternal(TConstArrayView<FMassEntityHandle> Entities, const int32 StartingChunk)
+FMassArchetypeEntityCollection::FArchetypeEntityRange FMassArchetypeData::PrepareNextEntitiesSpanInternal(TConstArrayView<FMassEntityHandle> Entities, const FMassArchetypeSharedFragmentValues& SharedFragmentValues, int32 StartingChunk)
 {
+	checkf(SharedFragmentValues.IsSorted(), TEXT("Expecting shared fragment values to be previously sorted"));
+	checkfSlow(SharedFragmentValues.HasExactFragmentTypesMatch(CompositionDescriptor.SharedFragments), TEXT("Expecting values for every specified shared fragment in the archetype and only those"))
+
 	int32 StartIndexWithinChunk = INDEX_NONE;
 	int32 AbsoluteStartIndex = 0;
 
@@ -969,7 +1007,7 @@ FMassArchetypeEntityCollection::FArchetypeEntityRange FMassArchetypeData::Prepar
 	for (; ChunkIndex < Chunks.Num(); ++ChunkIndex)
 	{
 		FMassArchetypeChunk& Chunk = Chunks[ChunkIndex];
-		if (Chunk.GetNumInstances() < NumEntitiesPerChunk)
+		if (Chunk.GetNumInstances() < NumEntitiesPerChunk && Chunk.GetSharedFragmentValues().IsEquivalent(SharedFragmentValues))
 		{
 			StartIndexWithinChunk = Chunk.GetNumInstances();
 			AbsoluteStartIndex = ChunkIndex * NumEntitiesPerChunk + StartIndexWithinChunk;
@@ -978,7 +1016,7 @@ FMassArchetypeEntityCollection::FArchetypeEntityRange FMassArchetypeData::Prepar
 
 			if (StartIndexWithinChunk == 0)
 			{
-				Chunk.Recycle(ChunkFragmentsTemplate);
+				Chunk.Recycle(ChunkFragmentsTemplate, SharedFragmentValues);
 			}
 			break;
 		}
@@ -991,7 +1029,7 @@ FMassArchetypeEntityCollection::FArchetypeEntityRange FMassArchetypeData::Prepar
 		AbsoluteStartIndex = Chunks.Num() * NumEntitiesPerChunk;
 		StartIndexWithinChunk = 0;
 
-		DestinationChunk = &Chunks.Emplace_GetRef(GetChunkAllocSize(), ChunkFragmentsTemplate);
+		DestinationChunk = &Chunks.Emplace_GetRef(GetChunkAllocSize(), ChunkFragmentsTemplate, SharedFragmentValues);
 	}
 
 	check(DestinationChunk);
@@ -1127,4 +1165,3 @@ void FMassArchetypeData::BatchSetFragmentValues(TConstArrayView<FMassArchetypeEn
 		EntitiesHandled += EntityRange.Length;
 	}
 }
-
diff --git a/Engine/Plugins/Runtime/MassEntity/Source/MassEntity/Private/MassArchetypeData.h b/Engine/Plugins/Runtime/MassEntity/Source/MassEntity/Private/MassArchetypeData.h
--- a/Engine/Plugins/Runtime/MassEntity/Source/MassEntity/Private/MassArchetypeData.h
+++ b/Engine/Plugins/Runtime/MassEntity/Source/MassEntity/Private/MassArchetypeData.h
@@ -24,11 +24,13 @@ private:
 	int32 NumInstances = 0;
 	int32 SerialModificationNumber = 0;
 	TArray<FInstancedStruct> ChunkFragmentData;
+	FMassArchetypeSharedFragmentValues SharedFragmentValues;
 
 public:
-	explicit FMassArchetypeChunk(int32 InAllocSize, TConstArrayView<FInstancedStruct> InChunkFragmentTemplates)
+	explicit FMassArchetypeChunk(int32 InAllocSize, TConstArrayView<FInstancedStruct> InChunkFragmentTemplates, FMassArchetypeSharedFragmentValues InSharedFragmentValues)
 		: AllocSize(InAllocSize)
 		, ChunkFragmentData(InChunkFragmentTemplates)
+		, SharedFragmentValues(InSharedFragmentValues)
 	{
 		RawMemory = (uint8*)FMemory::Malloc(AllocSize);
 	}
@@ -105,11 +107,12 @@ public:
 			});
 	}
 
-	void Recycle(TConstArrayView<FInstancedStruct> InChunkFragmentsTemplate)
+	void Recycle(TConstArrayView<FInstancedStruct> InChunkFragmentsTemplate, const FMassArchetypeSharedFragmentValues& InSharedFragmentValues)
 	{
 		checkf(NumInstances == 0, TEXT("Recycling a chunk that is not empty."));
 		SerialModificationNumber++;
 		ChunkFragmentData = InChunkFragmentsTemplate;
+		SharedFragmentValues = InSharedFragmentValues;
 		
 		// If this chunk previously had entity and it does not anymore, we might have to reallocate the memory as it was freed to save memory
 		if (RawMemory == nullptr)
@@ -126,6 +129,8 @@ public:
 #if WITH_MASSENTITY_DEBUG
 	int32 DebugGetChunkFragmentCount() const { return ChunkFragmentData.Num(); }
 #endif // WITH_MASSENTITY_DEBUG
+
+	const FMassArchetypeSharedFragmentValues& GetSharedFragmentValues() const { return SharedFragmentValues; }
 };
 
 // Information for a single fragment type in an archetype
@@ -147,7 +152,6 @@ struct FMassArchetypeData
 private:
 	// One-stop-shop variable describing the archetype's fragment and tag composition 
 	FMassArchetypeCompositionDescriptor CompositionDescriptor;
-	FMassArchetypeSharedFragmentValues SharedFragmentValues;
 
 	// Pre-created default chunk fragment templates
 	TArray<FInstancedStruct> ChunkFragmentsTemplate;
@@ -182,19 +186,29 @@ public:
 	const FMassSharedFragmentBitSet& GetSharedFragmentBitSet() const { return CompositionDescriptor.SharedFragments; }
 
 	const FMassArchetypeCompositionDescriptor& GetCompositionDescriptor() const { return CompositionDescriptor; }
-	const FMassArchetypeSharedFragmentValues& GetSharedFragmentValues() const { return SharedFragmentValues; }
+	FORCEINLINE const FMassArchetypeSharedFragmentValues& GetSharedFragmentValues(int32 EntityIndex) const
+	{ 
+		const int32 AbsoluteIndex = EntityMap.FindChecked(EntityIndex);
+		const int32 ChunkIndex = AbsoluteIndex / NumEntitiesPerChunk;
+
+		return Chunks[ChunkIndex].GetSharedFragmentValues();
+	}
+	FORCEINLINE const FMassArchetypeSharedFragmentValues& GetSharedFragmentValues(FMassEntityHandle Entity) const
+	{
+		return GetSharedFragmentValues(Entity.Index);
+	}
 
 	/** Method to iterate on all the fragment types */
 	void ForEachFragmentType(TFunction< void(const UScriptStruct* /*FragmentType*/)> Function) const;
 	bool HasFragmentType(const UScriptStruct* FragmentType) const;
 	bool HasTagType(const UScriptStruct* FragmentType) const { check(FragmentType); return CompositionDescriptor.Tags.Contains(*FragmentType); }
 
-	bool IsEquivalent(const FMassArchetypeCompositionDescriptor& OtherCompositionDescriptor, const FMassArchetypeSharedFragmentValues& OtherSharedFragmentValues) const
+	bool IsEquivalent(const FMassArchetypeCompositionDescriptor& OtherCompositionDescriptor) const
 	{
-		return CompositionDescriptor.IsEquivalent(OtherCompositionDescriptor) && SharedFragmentValues.IsEquivalent(OtherSharedFragmentValues);
+		return CompositionDescriptor.IsEquivalent(OtherCompositionDescriptor);
 	}
 
-	void Initialize(const FMassArchetypeCompositionDescriptor& InCompositionDescriptor, const FMassArchetypeSharedFragmentValues& InSharedFragmentValues);
+	void Initialize(const FMassArchetypeCompositionDescriptor& InCompositionDescriptor);
 
 	/** 
 	 * A special way of initializing an archetype resulting in a copy of SiblingArchetype's setup with OverrideTags
@@ -202,7 +216,7 @@ public:
 	 */
 	void InitializeWithSimilar(const FMassArchetypeData& BaseArchetype, FMassArchetypeCompositionDescriptor&& NewComposition);
 
-	void AddEntity(FMassEntityHandle Entity);
+	void AddEntity(FMassEntityHandle Entity, const FMassArchetypeSharedFragmentValues& InSharedFragmentValues);
 	void RemoveEntity(FMassEntityHandle Entity);
 
 	bool HasFragmentDataForEntity(const UScriptStruct* FragmentType, int32 EntityIndex) const;
@@ -320,12 +334,12 @@ public:
 	//////////////////////////////////////////////////////////////////////
 	// batched api
 	void BatchDestroyEntityChunks(FMassArchetypeEntityCollection::FConstEntityRangeArrayView EntityRangeContainer, TArray<FMassEntityHandle>& OutEntitiesRemoved);
-	void BatchAddEntities(TConstArrayView<FMassEntityHandle> Entities, TArray<FMassArchetypeEntityCollection::FArchetypeEntityRange>& OutNewRanges);
+	void BatchAddEntities(TConstArrayView<FMassEntityHandle> Entities, const FMassArchetypeSharedFragmentValues& SharedFragmentValues, TArray<FMassArchetypeEntityCollection::FArchetypeEntityRange>& OutNewRanges);
 	void BatchMoveEntitiesToAnotherArchetype(const FMassArchetypeEntityCollection& EntityCollection, FMassArchetypeData& NewArchetype, TArray<FMassEntityHandle>& OutEntitesBeingMoved, TArray<FMassArchetypeEntityCollection::FArchetypeEntityRange>* OutNewChunks = nullptr);
 	void BatchSetFragmentValues(TConstArrayView<FMassArchetypeEntityCollection::FArchetypeEntityRange> EntityCollection, const FMassGenericPayloadViewSlice& Payload);
 
 protected:
-	FMassArchetypeEntityCollection::FArchetypeEntityRange PrepareNextEntitiesSpanInternal(TConstArrayView<FMassEntityHandle> Entities, const int32 StartingChunk = 0);
+	FMassArchetypeEntityCollection::FArchetypeEntityRange PrepareNextEntitiesSpanInternal(TConstArrayView<FMassEntityHandle> Entities, const FMassArchetypeSharedFragmentValues& InSharedFragmentValues, const int32 StartingChunk = 0);
 	void BatchRemoveEntitiesInternal(const int32 ChunkIndex, const int32 StartIndexWithinChunk, const int32 NumberToRemove);
 
 	struct FTransientChunkLocation
@@ -344,11 +358,11 @@ protected:
 
 	void BindEntityRequirements(FMassExecutionContext& RunContext, const FMassFragmentIndicesMapping& EntityFragmentsMapping, FMassArchetypeChunk& Chunk, const int32 SubchunkStart, const int32 SubchunkLength);
 	void BindChunkFragmentRequirements(FMassExecutionContext& RunContext, const FMassFragmentIndicesMapping& ChunkFragmentsMapping, FMassArchetypeChunk& Chunk);
-	void BindConstSharedFragmentRequirements(FMassExecutionContext& RunContext, const FMassFragmentIndicesMapping& ChunkFragmentsMapping);
-	void BindSharedFragmentRequirements(FMassExecutionContext& RunContext, const FMassFragmentIndicesMapping& ChunkFragmentsMapping);
+	void BindConstSharedFragmentRequirements(FMassExecutionContext& RunContext, const FMassArchetypeSharedFragmentValues& SharedFragmentValues, const FMassFragmentIndicesMapping& ChunkFragmentsMapping);
+	void BindSharedFragmentRequirements(FMassExecutionContext& RunContext, const FMassArchetypeSharedFragmentValues& SharedFragmentValues, const FMassFragmentIndicesMapping& ChunkFragmentsMapping);
 
 private:
-	int32 AddEntityInternal(FMassEntityHandle Entity);
+	int32 AddEntityInternal(FMassEntityHandle Entity, const FMassArchetypeSharedFragmentValues& InSharedFragmentValues);
 	void RemoveEntityInternal(const int32 AbsoluteIndex);
 };
 
diff --git a/Engine/Plugins/Runtime/MassEntity/Source/MassEntity/Private/MassDebugger.cpp b/Engine/Plugins/Runtime/MassEntity/Source/MassEntity/Private/MassDebugger.cpp
--- a/Engine/Plugins/Runtime/MassEntity/Source/MassEntity/Private/MassDebugger.cpp
+++ b/Engine/Plugins/Runtime/MassEntity/Source/MassEntity/Private/MassDebugger.cpp
@@ -255,12 +255,6 @@ const FMassArchetypeCompositionDescriptor& FMassDebugger::GetArchetypeCompositio
 	return ArchetypeData.CompositionDescriptor;
 }
 
-const FMassArchetypeSharedFragmentValues& FMassDebugger::GetArchetypeSharedFragmentValues(const FMassArchetypeHandle& ArchetypeHandle)
-{
-	const FMassArchetypeData& ArchetypeData = FMassArchetypeHelper::ArchetypeDataFromHandleChecked(ArchetypeHandle);
-	return ArchetypeData.SharedFragmentValues;
-}
-
 void FMassDebugger::GetArchetypeEntityStats(const FMassArchetypeHandle& ArchetypeHandle
 	, int32& OutEntitiesCount, int32& OutNumEntitiesPerChunk, int32& OutChunksCount)
 {
diff --git a/Engine/Plugins/Runtime/MassEntity/Source/MassEntity/Private/MassEntitySubsystem.cpp b/Engine/Plugins/Runtime/MassEntity/Source/MassEntity/Private/MassEntitySubsystem.cpp
--- a/Engine/Plugins/Runtime/MassEntity/Source/MassEntity/Private/MassEntitySubsystem.cpp
+++ b/Engine/Plugins/Runtime/MassEntity/Source/MassEntity/Private/MassEntitySubsystem.cpp
@@ -155,7 +155,7 @@ FMassArchetypeHandle UMassEntitySubsystem::CreateArchetype(TConstArrayView<const
 	}
 
 	const FMassArchetypeCompositionDescriptor Composition(FMassFragmentBitSet(FragmentList), Tags, ChunkFragments, FMassSharedFragmentBitSet());
-	return CreateArchetype(Composition, FMassArchetypeSharedFragmentValues(), ArchetypeDebugName);
+	return CreateArchetype(Composition, ArchetypeDebugName);
 }
 
 FMassArchetypeHandle UMassEntitySubsystem::CreateArchetype(const TSharedPtr<FMassArchetypeData>& SourceArchetype, const FMassFragmentBitSet& AddedFragments, const FName ArchetypeDebugName)
@@ -164,19 +164,19 @@ FMassArchetypeHandle UMassEntitySubsystem::CreateArchetype(const TSharedPtr<FMas
 	checkf(AddedFragments.IsEmpty() == false, TEXT("%s Adding an empty fragment list to an archetype is not supported."), ANSI_TO_TCHAR(__FUNCTION__));
 
 	const FMassArchetypeCompositionDescriptor Composition(AddedFragments + SourceArchetype->GetFragmentBitSet(), SourceArchetype->GetTagBitSet(), SourceArchetype->GetChunkFragmentBitSet(), SourceArchetype->GetSharedFragmentBitSet());
-	return CreateArchetype(Composition, SourceArchetype->GetSharedFragmentValues(), ArchetypeDebugName);
+	return CreateArchetype(Composition, ArchetypeDebugName);
 }
 
-FMassArchetypeHandle UMassEntitySubsystem::CreateArchetype(const FMassArchetypeCompositionDescriptor& Composition, const FMassArchetypeSharedFragmentValues& SharedFragmentValues, const FName ArchetypeDebugName)
+FMassArchetypeHandle UMassEntitySubsystem::CreateArchetype(const FMassArchetypeCompositionDescriptor& Composition, const FName ArchetypeDebugName)
 {
-	const uint32 TypeHash = HashCombine(Composition.CalculateHash(), GetTypeHash(SharedFragmentValues));
+	const uint32 TypeHash = Composition.CalculateHash();
 
 	TArray<TSharedPtr<FMassArchetypeData>>& HashRow = FragmentHashToArchetypeMap.FindOrAdd(TypeHash);
 
 	TSharedPtr<FMassArchetypeData> ArchetypeDataPtr;
 	for (const TSharedPtr<FMassArchetypeData>& Ptr : HashRow)
 	{
-		if (Ptr->IsEquivalent(Composition, SharedFragmentValues))
+		if (Ptr->IsEquivalent(Composition))
 		{
 			// Keep track of all names for this archetype.
 			if (!ArchetypeDebugName.IsNone())
@@ -192,7 +192,7 @@ FMassArchetypeHandle UMassEntitySubsystem::CreateArchetype(const FMassArchetypeC
 	{
 		// Create a new archetype
 		FMassArchetypeData* NewArchetype = new FMassArchetypeData();
-		NewArchetype->Initialize(Composition, SharedFragmentValues);
+		NewArchetype->Initialize(Composition);
 		if (!ArchetypeDebugName.IsNone())
 		{
 			NewArchetype->AddUniqueDebugName(ArchetypeDebugName);
@@ -229,14 +229,14 @@ FMassArchetypeHandle UMassEntitySubsystem::InternalCreateSimilarArchetype(const
 
 FMassArchetypeHandle UMassEntitySubsystem::InternalCreateSimilarArchetype(const FMassArchetypeData& SourceArchetypeRef, FMassArchetypeCompositionDescriptor&& NewComposition)
 {
-	const uint32 TypeHash = HashCombine(NewComposition.CalculateHash(), GetTypeHash(SourceArchetypeRef.GetSharedFragmentValues()));
+	const uint32 TypeHash = NewComposition.CalculateHash();
 
 	TArray<TSharedPtr<FMassArchetypeData>>& HashRow = FragmentHashToArchetypeMap.FindOrAdd(TypeHash);
 
 	TSharedPtr<FMassArchetypeData> ArchetypeDataPtr;
 	for (const TSharedPtr<FMassArchetypeData>& Ptr : HashRow)
 	{
-		if (Ptr->IsEquivalent(NewComposition, SourceArchetypeRef.GetSharedFragmentValues()))
+		if (Ptr->IsEquivalent(NewComposition))
 		{
 			ArchetypeDataPtr = Ptr;
 			break;
@@ -309,25 +309,25 @@ void UMassEntitySubsystem::DoEntityCompaction(const double TimeAllowed)
 	}
 }
 
-FMassEntityHandle UMassEntitySubsystem::CreateEntity(const FMassArchetypeHandle& ArchetypeHandle)
+FMassEntityHandle UMassEntitySubsystem::CreateEntity(const FMassArchetypeHandle& ArchetypeHandle, const FMassArchetypeSharedFragmentValues& SharedFragmentValues)
 {
 	check(ArchetypeHandle.IsValid());
 
 	const FMassEntityHandle Entity = ReserveEntity();
-	InternalBuildEntity(Entity, ArchetypeHandle);
+	InternalBuildEntity(Entity, ArchetypeHandle, SharedFragmentValues);
 	return Entity;
 }
 
-FMassEntityHandle UMassEntitySubsystem::CreateEntity(TConstArrayView<FInstancedStruct> FragmentInstanceList, const FName ArchetypeDebugName)
+FMassEntityHandle UMassEntitySubsystem::CreateEntity(TConstArrayView<FInstancedStruct> FragmentInstanceList, const FMassArchetypeSharedFragmentValues& SharedFragmentValues, const FName ArchetypeDebugName)
 {
 	check(FragmentInstanceList.Num() > 0);
 
 	const FMassArchetypeHandle& ArchetypeHandle = CreateArchetype(FMassArchetypeCompositionDescriptor(FragmentInstanceList,
-		FMassTagBitSet(), FMassChunkFragmentBitSet(), FMassSharedFragmentBitSet()), FMassArchetypeSharedFragmentValues(), ArchetypeDebugName);
+		FMassTagBitSet(), FMassChunkFragmentBitSet(), FMassSharedFragmentBitSet()), ArchetypeDebugName);
 	check(ArchetypeHandle.IsValid());
 
 	const FMassEntityHandle Entity = ReserveEntity();
-	InternalBuildEntity(Entity, ArchetypeHandle);
+	InternalBuildEntity(Entity, ArchetypeHandle, SharedFragmentValues);
 
 	const FEntityData& EntityData = Entities[Entity.Index];
 	EntityData.CurrentArchetype->SetFragmentsData(Entity, FragmentInstanceList);
@@ -353,21 +353,20 @@ void UMassEntitySubsystem::ReleaseReservedEntity(FMassEntityHandle Entity)
 	InternalReleaseEntity(Entity);
 }
 
-void UMassEntitySubsystem::BuildEntity(FMassEntityHandle Entity, const FMassArchetypeHandle& ArchetypeHandle)
+void UMassEntitySubsystem::BuildEntity(FMassEntityHandle Entity, const FMassArchetypeHandle& ArchetypeHandle, const FMassArchetypeSharedFragmentValues& SharedFragmentValues)
 {
 	checkf(!IsEntityBuilt(Entity), TEXT("Expecting an entity that is not already built"));
 	check(ArchetypeHandle.IsValid());
 
-	InternalBuildEntity(Entity, ArchetypeHandle);
+	InternalBuildEntity(Entity, ArchetypeHandle, SharedFragmentValues);
 }
 
-void UMassEntitySubsystem::BuildEntity(FMassEntityHandle Entity, TConstArrayView<FInstancedStruct> FragmentInstanceList, FMassArchetypeSharedFragmentValues SharedFragmentValues)
+void UMassEntitySubsystem::BuildEntity(FMassEntityHandle Entity, TConstArrayView<FInstancedStruct> FragmentInstanceList, const FMassArchetypeSharedFragmentValues& SharedFragmentValues)
 {
 	check(FragmentInstanceList.Num() > 0);
 	checkf(!IsEntityBuilt(Entity), TEXT("Expecting an entity that is not already built"));
 
-	SharedFragmentValues.Sort();
-	SharedFragmentValues.CacheHash();
+	checkf(SharedFragmentValues.IsSorted(), TEXT("Expecting shared fragment values to be previously sorted"));
 	FMassArchetypeCompositionDescriptor Composition(FragmentInstanceList, FMassTagBitSet(), FMassChunkFragmentBitSet(), FMassSharedFragmentBitSet());
 	for (const FConstSharedStruct& SharedFragment : SharedFragmentValues.GetConstSharedFragments())
 	{
@@ -378,10 +377,10 @@ void UMassEntitySubsystem::BuildEntity(FMassEntityHandle Entity, TConstArrayView
 		Composition.SharedFragments.Add(*SharedFragment.GetScriptStruct());
 	}
 
-	const FMassArchetypeHandle& ArchetypeHandle = CreateArchetype(Composition, SharedFragmentValues);
+	const FMassArchetypeHandle& ArchetypeHandle = CreateArchetype(Composition);
 	check(ArchetypeHandle.IsValid());
 
-	InternalBuildEntity(Entity, ArchetypeHandle);
+	InternalBuildEntity(Entity, ArchetypeHandle, SharedFragmentValues);
 
 	const FEntityData& EntityData = Entities[Entity.Index];
 	EntityData.CurrentArchetype->SetFragmentsData(Entity, FragmentInstanceList);
@@ -411,7 +410,7 @@ void UMassEntitySubsystem::BatchBuildEntities(const FMassArchetypeEntityCollecti
 	FMassArchetypeEntityCollection::FEntityRangeArray TargetArchetypeEntityRanges;
 
 	// "built" entities case, this is verified during FMassArchetypeEntityCollectionWithPayload construction
-	FMassArchetypeHandle TargetArchetypeHandle = CreateArchetype(Composition, SharedFragmentValues, ArchetypeDebugName);
+	FMassArchetypeHandle TargetArchetypeHandle = CreateArchetype(Composition, ArchetypeDebugName);
 	check(TargetArchetypeHandle.IsValid());
 
 	// there are some extra steps in creating EncodedEntities from the original given entity handles and then back
@@ -432,7 +431,7 @@ void UMassEntitySubsystem::BatchBuildEntities(const FMassArchetypeEntityCollecti
 		EntityData.CurrentArchetype = TargetArchetypeHandle.DataPtr;
 	}
 
-	TargetArchetypeHandle.DataPtr->BatchAddEntities(EntityHandles, TargetArchetypeEntityRanges);
+	TargetArchetypeHandle.DataPtr->BatchAddEntities(EntityHandles, SharedFragmentValues, TargetArchetypeEntityRanges);
 
 	if (EncodedEntitiesWithPayload.GetPayload().IsEmpty() == false)
 	{
@@ -451,7 +450,7 @@ void UMassEntitySubsystem::BatchBuildEntities(const FMassArchetypeEntityCollecti
 	}
 }
 
-TSharedRef<UMassEntitySubsystem::FEntityCreationContext> UMassEntitySubsystem::BatchCreateEntities(const FMassArchetypeHandle& ArchetypeHandle, const int32 Count, TArray<FMassEntityHandle>& OutEntities)
+TSharedRef<UMassEntitySubsystem::FEntityCreationContext> UMassEntitySubsystem::BatchCreateEntities(const FMassArchetypeHandle& ArchetypeHandle, const FMassArchetypeSharedFragmentValues& SharedFragmentValues, const int32 Count, TArray<FMassEntityHandle>& OutEntities)
 {
 	TRACE_CPUPROFILER_EVENT_SCOPE(Mass_BatchCreateEntities);
 
@@ -472,7 +471,7 @@ TSharedRef<UMassEntitySubsystem::FEntityCreationContext> UMassEntitySubsystem::B
 		EntityData.CurrentArchetype = ArchetypeHandle.DataPtr;
 		EntityData.SerialNumber = Result.SerialNumber;
 		
-		ArchetypeData.AddEntity(Result);
+		ArchetypeData.AddEntity(Result, SharedFragmentValues);
 	}
 		
 	FEntityCreationContext* CreationContext = new FEntityCreationContext(Count);
@@ -617,7 +616,7 @@ void UMassEntitySubsystem::AddCompositionToEntity_GetDelta(FMassEntityHandle Ent
 		NewDescriptor.Fragments += InDescriptor.Fragments;
 		NewDescriptor.Tags += InDescriptor.Tags;
 
-		const FMassArchetypeHandle NewArchetypeHandle = CreateArchetype(NewDescriptor, OldArchetype->GetSharedFragmentValues());
+		const FMassArchetypeHandle NewArchetypeHandle = CreateArchetype(NewDescriptor);
 
 		if (ensure(NewArchetypeHandle.DataPtr != EntityData.CurrentArchetype))
 		{
@@ -654,7 +653,7 @@ void UMassEntitySubsystem::RemoveCompositionFromEntity(FMassEntityHandle Entity,
 			ensureMsgf(OldArchetype->GetCompositionDescriptor().HasAll(InDescriptor), TEXT("Some of the elements being removed are already missing from entity\'s composition."));
 			ObserverManager.OnPreCompositionRemoved(Entity, InDescriptor);
 
-			const FMassArchetypeHandle NewArchetypeHandle = CreateArchetype(NewDescriptor, OldArchetype->GetSharedFragmentValues());
+			const FMassArchetypeHandle NewArchetypeHandle = CreateArchetype(NewDescriptor);
 
 			if (ensure(NewArchetypeHandle.DataPtr != EntityData.CurrentArchetype))
 			{
@@ -674,11 +673,11 @@ const FMassArchetypeCompositionDescriptor& UMassEntitySubsystem::GetArchetypeCom
 	return ArchetypeData.GetCompositionDescriptor();
 }
 
-void UMassEntitySubsystem::InternalBuildEntity(FMassEntityHandle Entity, const FMassArchetypeHandle& ArchetypeHandle)
+void UMassEntitySubsystem::InternalBuildEntity(FMassEntityHandle Entity, const FMassArchetypeHandle& ArchetypeHandle, const FMassArchetypeSharedFragmentValues& SharedFragmentValues)
 {
 	FEntityData& EntityData = Entities[Entity.Index];
 	EntityData.CurrentArchetype = ArchetypeHandle.DataPtr;
-	EntityData.CurrentArchetype->AddEntity(Entity);
+	EntityData.CurrentArchetype->AddEntity(Entity, SharedFragmentValues);
 }
 
 void UMassEntitySubsystem::InternalReleaseEntity(FMassEntityHandle Entity)
@@ -760,7 +759,7 @@ void UMassEntitySubsystem::RemoveFragmentListFromEntity(FMassEntityHandle Entity
 	{
 		// If all the fragments got removed this will result in fetching of the empty archetype
 		const FMassArchetypeCompositionDescriptor NewComposition(OldArchetype->GetFragmentBitSet() - FragmentsToRemove, OldArchetype->GetTagBitSet(), OldArchetype->GetChunkFragmentBitSet(), OldArchetype->GetSharedFragmentBitSet());
-		const FMassArchetypeHandle NewArchetypeHandle = CreateArchetype(NewComposition, OldArchetype->GetSharedFragmentValues());
+		const FMassArchetypeHandle NewArchetypeHandle = CreateArchetype(NewComposition);
 
 		// Move the entity over
 		FMassArchetypeData& NewArchetype = FMassArchetypeHelper::ArchetypeDataFromHandleChecked(NewArchetypeHandle);
diff --git a/Engine/Plugins/Runtime/MassEntity/Source/MassEntity/Private/MassEntityView.cpp b/Engine/Plugins/Runtime/MassEntity/Source/MassEntity/Private/MassEntityView.cpp
--- a/Engine/Plugins/Runtime/MassEntity/Source/MassEntity/Private/MassEntityView.cpp
+++ b/Engine/Plugins/Runtime/MassEntity/Source/MassEntity/Private/MassEntityView.cpp
@@ -42,26 +42,26 @@ void* FMassEntityView::GetFragmentPtrChecked(const UScriptStruct& FragmentType)
 
 const void* FMassEntityView::GetConstSharedFragmentPtr(const UScriptStruct& FragmentType) const
 {
-	const FConstSharedStruct* SharedFragment = Archetype->GetSharedFragmentValues().GetConstSharedFragments().FindByPredicate(FStructTypeEqualOperator(&FragmentType));
+	const FConstSharedStruct* SharedFragment = Archetype->GetSharedFragmentValues(Entity).GetConstSharedFragments().FindByPredicate(FStructTypeEqualOperator(&FragmentType));
 	return (SharedFragment != nullptr) ? SharedFragment->GetMemory() : nullptr;
 }
 
 const void* FMassEntityView::GetConstSharedFragmentPtrChecked(const UScriptStruct& FragmentType) const
 {
-	const FConstSharedStruct* SharedFragment = Archetype->GetSharedFragmentValues().GetConstSharedFragments().FindByPredicate(FStructTypeEqualOperator(&FragmentType));
+	const FConstSharedStruct* SharedFragment = Archetype->GetSharedFragmentValues(Entity).GetConstSharedFragments().FindByPredicate(FStructTypeEqualOperator(&FragmentType));
 	check(SharedFragment != nullptr);
 	return SharedFragment->GetMemory();
 }
 
 void* FMassEntityView::GetSharedFragmentPtr(const UScriptStruct& FragmentType) const
 {
-	const FSharedStruct* SharedFragment = Archetype->GetSharedFragmentValues().GetSharedFragments().FindByPredicate(FStructTypeEqualOperator(&FragmentType));
+	const FSharedStruct* SharedFragment = Archetype->GetSharedFragmentValues(Entity).GetSharedFragments().FindByPredicate(FStructTypeEqualOperator(&FragmentType));
 	return (SharedFragment != nullptr) ? SharedFragment->GetMutableMemory() : nullptr;
 }
 
 void* FMassEntityView::GetSharedFragmentPtrChecked(const UScriptStruct& FragmentType) const
 {
-	const FSharedStruct* SharedFragment = Archetype->GetSharedFragmentValues().GetSharedFragments().FindByPredicate(FStructTypeEqualOperator(&FragmentType));
+	const FSharedStruct* SharedFragment = Archetype->GetSharedFragmentValues(Entity).GetSharedFragments().FindByPredicate(FStructTypeEqualOperator(&FragmentType));
 	check(SharedFragment != nullptr);
 	return SharedFragment->GetMutableMemory();
 }
diff --git a/Engine/Plugins/Runtime/MassEntity/Source/MassEntity/Public/MassDebugger.h b/Engine/Plugins/Runtime/MassEntity/Source/MassEntity/Public/MassDebugger.h
--- a/Engine/Plugins/Runtime/MassEntity/Source/MassEntity/Public/MassDebugger.h
+++ b/Engine/Plugins/Runtime/MassEntity/Source/MassEntity/Public/MassDebugger.h
@@ -53,7 +53,6 @@ struct MASSENTITY_API FMassDebugger
 
 	static TArray<FMassArchetypeHandle> GetAllArchetypes(const UMassEntitySubsystem& EntitySubsystem);
 	static const FMassArchetypeCompositionDescriptor& GetArchetypeComposition(const FMassArchetypeHandle& ArchetypeHandle);
-	static const FMassArchetypeSharedFragmentValues& GetArchetypeSharedFragmentValues(const FMassArchetypeHandle& ArchetypeHandle);
 
 	static void GetArchetypeEntityStats(const FMassArchetypeHandle& ArchetypeHandle, int32& OutEntitiesCount, int32& OutNumEntitiesPerChunk, int32& OutChunksCount);
 	static const TConstArrayView<FName> GetArchetypeDebugNames(const FMassArchetypeHandle& ArchetypeHandle);
diff --git a/Engine/Plugins/Runtime/MassEntity/Source/MassEntity/Public/MassEntitySubsystem.h b/Engine/Plugins/Runtime/MassEntity/Source/MassEntity/Public/MassEntitySubsystem.h
--- a/Engine/Plugins/Runtime/MassEntity/Source/MassEntity/Public/MassEntitySubsystem.h
+++ b/Engine/Plugins/Runtime/MassEntity/Source/MassEntity/Public/MassEntitySubsystem.h
@@ -95,11 +95,10 @@ public:
 	 * CreateArchetype from a composition descriptor and initial values
 	 *
 	 * @param Composition of fragment, tag and chunk fragment types
-	 * @param SharedFragmentValues are the actual pointer to shared fragments
 	 * @param ArchetypeDebugName Name to identify the archetype while debugging
 	 * @return a handle of a new archetype 
 	 */
-	FMassArchetypeHandle CreateArchetype(const FMassArchetypeCompositionDescriptor& Composition, const FMassArchetypeSharedFragmentValues& SharedFragmentValues, const FName ArchetypeDebugName = FName());
+	FMassArchetypeHandle CreateArchetype(const FMassArchetypeCompositionDescriptor& Composition, const FName ArchetypeDebugName = FName());
 
 	/** 
 	 *  Creates an archetype like SourceArchetype + InFragments. 
@@ -134,15 +133,17 @@ public:
 	/**
 	 * Creates fully built entity ready to be used by the subsystem
 	 * @param Archetype you want this entity to be
+	 * @param SharedFragmentValues to be associated with the entity
 	 * @return FMassEntityHandle id of the newly created entity */
-	FMassEntityHandle CreateEntity(const FMassArchetypeHandle& ArchetypeHandle);
+	FMassEntityHandle CreateEntity(const FMassArchetypeHandle& ArchetypeHandle, const FMassArchetypeSharedFragmentValues& SharedFragmentValues = {});
 
 	/**
 	 * Creates fully built entity ready to be used by the subsystem
 	 * @param FragmentInstanceList is the fragments to create the entity from and initialize values
+	 * @param SharedFragmentValues to be associated with the entity
 	 * @param ArchetypeDebugName Name to identify the archetype while debugging
 	 * @return FMassEntityHandle id of the newly created entity */
-	FMassEntityHandle CreateEntity(TConstArrayView<FInstancedStruct> FragmentInstanceList, const FName ArchetypeDebugName = FName());
+	FMassEntityHandle CreateEntity(TConstArrayView<FInstancedStruct> FragmentInstanceList, const FMassArchetypeSharedFragmentValues& SharedFragmentValues = {}, const FName ArchetypeDebugName = FName());
 
 	/**
 	 * A dedicated structure for ensuring the "on entities creation" observers get notified only once all other 
@@ -162,12 +163,18 @@ public:
 		TFunction<void(FEntityCreationContext&)> OnSpawningFinished;
 	};
 
-	/** A version of CreateEntity that's creating a number of entities (Count) at one go
-	 *  @param Archetype you want this entity to be
-	 *  @param Count number of entities to create
-	 *  @param OutEntities the newly created entities are appended to given array, i.e. the pre-existing content of OutEntities won't be affected by the call
-	 *  @return a creation context that will notify all the interested observers about newly created fragments once the context is released */
-	TSharedRef<FEntityCreationContext> BatchCreateEntities(const FMassArchetypeHandle& ArchetypeHandle, const int32 Count, TArray<FMassEntityHandle>& OutEntities);
+	/**
+	 * A version of CreateEntity that's creating a number of entities (Count) at one go
+	 * @param Archetype you want this entity to be
+	 * @param SharedFragmentValues to be associated with the entities
+	 * @param Count number of entities to create
+	 * @param OutEntities the newly created entities are appended to given array, i.e. the pre-existing content of OutEntities won't be affected by the call
+	 * @return a creation context that will notify all the interested observers about newly created fragments once the context is released */
+	TSharedRef<FEntityCreationContext> BatchCreateEntities(const FMassArchetypeHandle& ArchetypeHandle, const FMassArchetypeSharedFragmentValues& SharedFragmentValues, const int32 Count, TArray<FMassEntityHandle>& OutEntities);
+	FORCEINLINE TSharedRef<FEntityCreationContext> BatchCreateEntities(const FMassArchetypeHandle& ArchetypeHandle, const int32 Count, TArray<FMassEntityHandle>& OutEntities)
+	{
+		return BatchCreateEntities(ArchetypeHandle, FMassArchetypeSharedFragmentValues(), Count, OutEntities);
+	}
 
 	/**
 	 * Destroys a fully built entity, use ReleaseReservedEntity if entity was not yet built.
@@ -182,14 +189,18 @@ public:
 	/**
 	 * Builds an entity for it to be ready to be used by the subsystem
 	 * @param Entity to build which was retrieved with ReserveEntity() method
-	 * @param Archetype you want this entity to be*/
-	void BuildEntity(FMassEntityHandle Entity, const FMassArchetypeHandle& ArchetypeHandle);
+	 * @param Archetype you want this entity to be
+	 * @param SharedFragmentValues to be associated with the entity
+	 */
+	void BuildEntity(FMassEntityHandle Entity, const FMassArchetypeHandle& ArchetypeHandle, const FMassArchetypeSharedFragmentValues& SharedFragmentValues = {});
 
 	/**
 	 * Builds an entity for it to be ready to be used by the subsystem
 	 * @param Entity to build which was retrieved with ReserveEntity() method
-	 * @param FragmentInstanceList is the fragments to create the entity from and initialize values*/
-	void BuildEntity(FMassEntityHandle Entity, TConstArrayView<FInstancedStruct> FragmentInstanceList, FMassArchetypeSharedFragmentValues SharedFragmentValues = {});
+	 * @param FragmentInstanceList is the fragments to create the entity from and initialize values
+	 * @param SharedFragmentValues to be associated with the entity
+	 */
+	void BuildEntity(FMassEntityHandle Entity, TConstArrayView<FInstancedStruct> FragmentInstanceList, const FMassArchetypeSharedFragmentValues& SharedFragmentValues = {});
 
 	/*
 	 * Releases a previously reserved entity that was not yet built, otherwise call DestroyEntity
@@ -220,8 +231,8 @@ public:
 	void RemoveTagFromEntity(FMassEntityHandle Entity, const UScriptStruct* TagType);
 	void SwapTagsForEntity(FMassEntityHandle Entity, const UScriptStruct* FromFragmentType, const UScriptStruct* ToFragmentType);
 
-	void BatchBuildEntities(const FMassArchetypeEntityCollectionWithPayload& EncodedEntitiesWithPayload, const FMassFragmentBitSet& FragmentsAffected, const FMassArchetypeSharedFragmentValues& SharedFragmentValues, const FName ArchetypeDebugName = FName());
-	void BatchBuildEntities(const FMassArchetypeEntityCollectionWithPayload& EncodedEntitiesWithPayload, FMassArchetypeCompositionDescriptor&& Composition, const FMassArchetypeSharedFragmentValues& SharedFragmentValues, const FName ArchetypeDebugName = FName());
+	void BatchBuildEntities(const FMassArchetypeEntityCollectionWithPayload& EncodedEntitiesWithPayload, const FMassFragmentBitSet& FragmentsAffected, const FMassArchetypeSharedFragmentValues& SharedFragmentValues = {}, const FName ArchetypeDebugName = FName());
+	void BatchBuildEntities(const FMassArchetypeEntityCollectionWithPayload& EncodedEntitiesWithPayload, FMassArchetypeCompositionDescriptor&& Composition, const FMassArchetypeSharedFragmentValues& SharedFragmentValues = {}, const FName ArchetypeDebugName = FName());
 	void BatchChangeTagsForEntities(TConstArrayView<FMassArchetypeEntityCollection> EntityCollections, const FMassTagBitSet& TagsToAdd, const FMassTagBitSet& TagsToRemove);
 	void BatchChangeFragmentCompositionForEntities(TConstArrayView<FMassArchetypeEntityCollection> EntityCollections, const FMassFragmentBitSet& FragmentsToAdd, const FMassFragmentBitSet& FragmentsToRemove);
 	void BatchAddFragmentInstancesForEntities(TConstArrayView<FMassArchetypeEntityCollectionWithPayload> EntityCollections, const FMassFragmentBitSet& FragmentsAffected);
@@ -377,7 +388,7 @@ protected:
 	FMassArchetypeHandle InternalCreateSimilarArchetype(const FMassArchetypeData& SourceArchetypeRef, FMassArchetypeCompositionDescriptor&& NewComposition);
 
 private:
-	void InternalBuildEntity(FMassEntityHandle Entity, const FMassArchetypeHandle& ArchetypeHandle);
+	void InternalBuildEntity(FMassEntityHandle Entity, const FMassArchetypeHandle& ArchetypeHandle, const FMassArchetypeSharedFragmentValues& SharedFragmentValues);
 	void InternalReleaseEntity(FMassEntityHandle Entity);
 
 	/** 
diff --git a/Engine/Plugins/Runtime/MassEntity/Source/MassEntity/Public/MassEntityTypes.h b/Engine/Plugins/Runtime/MassEntity/Source/MassEntity/Public/MassEntityTypes.h
--- a/Engine/Plugins/Runtime/MassEntity/Source/MassEntity/Public/MassEntityTypes.h
+++ b/Engine/Plugins/Runtime/MassEntity/Source/MassEntity/Public/MassEntityTypes.h
@@ -224,39 +224,14 @@ struct MASSENTITY_API FMassArchetypeSharedFragmentValues
 	static constexpr uint32 EmptyInstanceHash = 0;
 
 	FMassArchetypeSharedFragmentValues() = default;
+	FMassArchetypeSharedFragmentValues(const FMassArchetypeSharedFragmentValues& OtherFragmentValues) = default;
+	FMassArchetypeSharedFragmentValues(FMassArchetypeSharedFragmentValues&& OtherFragmentValues) = default;
+	FMassArchetypeSharedFragmentValues& operator=(const FMassArchetypeSharedFragmentValues& OtherFragmentValues) = default;
+	FMassArchetypeSharedFragmentValues& operator=(FMassArchetypeSharedFragmentValues&& OtherFragmentValues) = default;
 
-	FMassArchetypeSharedFragmentValues(const FMassArchetypeSharedFragmentValues& OtherFragmentValues)
+	FORCEINLINE bool HasExactFragmentTypesMatch(const FMassSharedFragmentBitSet& InSharedFragmentBitSet) const
 	{
-		ConstSharedFragments = OtherFragmentValues.ConstSharedFragments;
-		SharedFragments = OtherFragmentValues.SharedFragments;
-		HashCache = OtherFragmentValues.HashCache;
-		bSorted = OtherFragmentValues.bSorted;
-	}
-
-	FMassArchetypeSharedFragmentValues(FMassArchetypeSharedFragmentValues&& OtherFragmentValues)
-	{
-		ConstSharedFragments = MoveTemp(OtherFragmentValues.ConstSharedFragments);
-		SharedFragments = MoveTemp(OtherFragmentValues.SharedFragments);
-		HashCache = OtherFragmentValues.HashCache;
-		bSorted = OtherFragmentValues.bSorted;
-	}
-
-	FMassArchetypeSharedFragmentValues& operator=(const FMassArchetypeSharedFragmentValues& OtherFragmentValues)
-	{
-		ConstSharedFragments = OtherFragmentValues.ConstSharedFragments;
-		SharedFragments = OtherFragmentValues.SharedFragments;
-		HashCache = OtherFragmentValues.HashCache;
-		bSorted = OtherFragmentValues.bSorted;
-		return *this;
-	}
-
-	FMassArchetypeSharedFragmentValues& operator=(FMassArchetypeSharedFragmentValues&& OtherFragmentValues)
-	{
-		ConstSharedFragments = MoveTemp(OtherFragmentValues.ConstSharedFragments);
-		SharedFragments = MoveTemp(OtherFragmentValues.SharedFragments);
-		HashCache = OtherFragmentValues.HashCache;
-		bSorted = OtherFragmentValues.bSorted;
-		return *this;
+		return SharedFragmentBitSet == InSharedFragmentBitSet;
 	}
 
 	FORCEINLINE bool IsEquivalent(const FMassArchetypeSharedFragmentValues& OtherSharedFragmentValues) const
@@ -267,12 +242,16 @@ struct MASSENTITY_API FMassArchetypeSharedFragmentValues
 	FORCEINLINE FConstSharedStruct& AddConstSharedFragment(const FConstSharedStruct& Fragment)
 	{
 		DirtyHashCache();
+		check(Fragment.GetScriptStruct());
+		SharedFragmentBitSet.Add(*Fragment.GetScriptStruct());
 		return ConstSharedFragments.Add_GetRef(Fragment);
 	}
 
 	FORCEINLINE FSharedStruct AddSharedFragment(const FSharedStruct& Fragment)
 	{
 		DirtyHashCache();
+		check(Fragment.GetScriptStruct());
+		SharedFragmentBitSet.Add(*Fragment.GetScriptStruct());
 		return SharedFragments.Add_GetRef(Fragment);
 	}
 
@@ -329,6 +308,7 @@ protected:
 	mutable uint32 HashCache = UINT32_MAX;
 	mutable bool bSorted = true; // When no element in the array, consider already sorted
 	
+	FMassSharedFragmentBitSet SharedFragmentBitSet;
 	TArray<FConstSharedStruct> ConstSharedFragments;
 	TArray<FSharedStruct> SharedFragments;
 };
diff --git a/Engine/Plugins/Runtime/MassEntity/Source/MassEntityDebugger/Private/MassDebuggerModel.cpp b/Engine/Plugins/Runtime/MassEntity/Source/MassEntityDebugger/Private/MassDebuggerModel.cpp
--- a/Engine/Plugins/Runtime/MassEntity/Source/MassEntityDebugger/Private/MassDebuggerModel.cpp
+++ b/Engine/Plugins/Runtime/MassEntity/Source/MassEntityDebugger/Private/MassDebuggerModel.cpp
@@ -153,11 +153,10 @@ FMassDebuggerArchetypeData::FMassDebuggerArchetypeData(const FMassArchetypeHandl
 {
 #if WITH_MASSENTITY_DEBUG
 	Composition = FMassDebugger::GetArchetypeComposition(ArchetypeHandle);
-	SharedFragments = FMassDebugger::GetArchetypeSharedFragmentValues(ArchetypeHandle);
 
 	// @todo should ensure we're using same hashing as the EntitySubsystem here
 	CompositionHash = Composition.CalculateHash();
-	FullHash = HashCombine(CompositionHash, GetTypeHash(SharedFragments));
+	FullHash = CompositionHash;
 
 	FString FullHashAsString;
 	BytesToHexLower(reinterpret_cast<const uint8*>(&FullHash), sizeof(FullHash), FullHashAsString);
diff --git a/Engine/Plugins/Runtime/MassGameplay/Source/MassLOD/Private/MassSimulationLOD.cpp b/Engine/Plugins/Runtime/MassGameplay/Source/MassLOD/Private/MassSimulationLOD.cpp
--- a/Engine/Plugins/Runtime/MassGameplay/Source/MassLOD/Private/MassSimulationLOD.cpp
+++ b/Engine/Plugins/Runtime/MassGameplay/Source/MassLOD/Private/MassSimulationLOD.cpp
@@ -87,12 +87,11 @@ void UMassSimulationLODProcessor::ConfigureQueries()
 	EntityQueryCalculateLOD.SetChunkFilter(FMassSimulationVariableTickSharedFragment::ShouldCalculateLODForChunk);
 
 	EntityQueryAdjustDistances = EntityQuery;
-	EntityQueryAdjustDistances.SetArchetypeFilter([](const FMassExecutionContext& Context)
+	EntityQueryAdjustDistances.SetChunkFilter([](const FMassExecutionContext& Context)
 	{
 		const FMassSimulationLODSharedFragment& LODSharedFragment = Context.GetSharedFragment<FMassSimulationLODSharedFragment>();
-		return LODSharedFragment.bHasAdjustedDistancesFromCount;
+		return LODSharedFragment.bHasAdjustedDistancesFromCount && FMassSimulationVariableTickSharedFragment::ShouldAdjustLODFromCountForChunk(Context);
 	});
-	EntityQueryAdjustDistances.SetChunkFilter(&FMassSimulationVariableTickSharedFragment::ShouldAdjustLODFromCountForChunk);
 
 	EntityQueryVariableTick.AddRequirement<FMassSimulationLODFragment>(EMassFragmentAccess::ReadOnly);
 	EntityQueryVariableTick.AddRequirement<FMassSimulationVariableTickFragment>(EMassFragmentAccess::ReadWrite);
@@ -104,7 +103,7 @@ void UMassSimulationLODProcessor::ConfigureQueries()
 	EntityQuerySetLODTag.AddRequirement<FMassSimulationLODFragment>(EMassFragmentAccess::ReadOnly);
 	EntityQuerySetLODTag.AddRequirement<FMassSimulationVariableTickFragment>(EMassFragmentAccess::ReadWrite, EMassFragmentPresence::None);
 	EntityQuerySetLODTag.AddConstSharedRequirement<FMassSimulationLODParameters>();
-	EntityQuerySetLODTag.SetArchetypeFilter([](const FMassExecutionContext& Context)
+	EntityQuerySetLODTag.SetChunkFilter([](const FMassExecutionContext& Context)
 	{
 		const FMassSimulationLODParameters& LODParams = Context.GetConstSharedFragment<FMassSimulationLODParameters>();
 		return LODParams.bSetLODTags;
diff --git a/Engine/Plugins/Runtime/MassGameplay/Source/MassReplication/Private/MassReplicationProcessor.cpp b/Engine/Plugins/Runtime/MassGameplay/Source/MassReplication/Private/MassReplicationProcessor.cpp
--- a/Engine/Plugins/Runtime/MassGameplay/Source/MassReplication/Private/MassReplicationProcessor.cpp
+++ b/Engine/Plugins/Runtime/MassGameplay/Source/MassReplication/Private/MassReplicationProcessor.cpp
@@ -47,7 +47,7 @@ void UMassReplicationProcessor::ConfigureQueries()
 	AdjustLODDistancesQuery.AddRequirement<FMassReplicationViewerInfoFragment>(EMassFragmentAccess::ReadOnly);
 	AdjustLODDistancesQuery.AddRequirement<FMassReplicationLODFragment>(EMassFragmentAccess::ReadWrite);
 	AdjustLODDistancesQuery.AddSharedRequirement<FMassReplicationSharedFragment>(EMassFragmentAccess::ReadWrite);
-	AdjustLODDistancesQuery.SetArchetypeFilter([](const FMassExecutionContext& Context)
+	AdjustLODDistancesQuery.SetChunkFilter([](const FMassExecutionContext& Context)
 	{
 		const FMassReplicationSharedFragment& LODSharedFragment = Context.GetSharedFragment<FMassReplicationSharedFragment>();
 		return LODSharedFragment.bHasAdjustedDistancesFromCount;
@@ -89,7 +89,7 @@ void UMassReplicationProcessor::PrepareExecution(UMassEntitySubsystem& EntitySub
 		if (!RepSharedFragment.bEntityQueryInitialized)
 		{
 			RepSharedFragment.EntityQuery = EntityQuery;
-			RepSharedFragment.EntityQuery.SetArchetypeFilter([&RepSharedFragment](const FMassExecutionContext& Context)
+			RepSharedFragment.EntityQuery.SetChunkFilter([&RepSharedFragment](const FMassExecutionContext& Context)
 			{
 				const FMassReplicationSharedFragment& CurRepSharedFragment = Context.GetSharedFragment<FMassReplicationSharedFragment>();
 				return &CurRepSharedFragment == &RepSharedFragment;
diff --git a/Engine/Plugins/Runtime/MassGameplay/Source/MassRepresentation/Private/MassVisualizationLODProcessor.cpp b/Engine/Plugins/Runtime/MassGameplay/Source/MassRepresentation/Private/MassVisualizationLODProcessor.cpp
--- a/Engine/Plugins/Runtime/MassGameplay/Source/MassRepresentation/Private/MassVisualizationLODProcessor.cpp
+++ b/Engine/Plugins/Runtime/MassGameplay/Source/MassRepresentation/Private/MassVisualizationLODProcessor.cpp
@@ -31,7 +31,7 @@ void UMassVisualizationLODProcessor::ConfigureQueries()
 	CloseEntityQuery.RegisterWithProcessor(*this);
 
 	CloseEntityAdjustDistanceQuery = CloseEntityQuery;
-	CloseEntityAdjustDistanceQuery.SetArchetypeFilter([](const FMassExecutionContext& Context)
+	CloseEntityAdjustDistanceQuery.SetChunkFilter([](const FMassExecutionContext& Context)
 	{
 		const FMassVisualizationLODSharedFragment& LODSharedFragment = Context.GetSharedFragment<FMassVisualizationLODSharedFragment>();
 		return LODSharedFragment.bHasAdjustedDistancesFromCount;
diff --git a/Engine/Plugins/Runtime/MassGameplay/Source/MassSpawner/Private/MassEntityTemplateRegistry.cpp b/Engine/Plugins/Runtime/MassGameplay/Source/MassSpawner/Private/MassEntityTemplateRegistry.cpp
--- a/Engine/Plugins/Runtime/MassGameplay/Source/MassSpawner/Private/MassEntityTemplateRegistry.cpp
+++ b/Engine/Plugins/Runtime/MassGameplay/Source/MassSpawner/Private/MassEntityTemplateRegistry.cpp
@@ -156,7 +156,7 @@ void UMassEntityTemplateRegistry::InitializeEntityTemplate(FMassEntityTemplate&
 	// Sort anything there is to sort for later comparison purposes
 	OutTemplate.Sort();
 
-	const FMassArchetypeHandle ArchetypeHandle = EntitySys->CreateArchetype(OutTemplate.GetCompositionDescriptor(), OutTemplate.GetSharedFragmentValues(), FName(OutTemplate.GetTemplateName()));
+	const FMassArchetypeHandle ArchetypeHandle = EntitySys->CreateArchetype(OutTemplate.GetCompositionDescriptor(), FName(OutTemplate.GetTemplateName()));
 	OutTemplate.SetArchetype(ArchetypeHandle);
 }
 
diff --git a/Engine/Plugins/Runtime/MassGameplay/Source/MassSpawner/Private/MassSpawnerSubsystem.cpp b/Engine/Plugins/Runtime/MassGameplay/Source/MassSpawner/Private/MassSpawnerSubsystem.cpp
--- a/Engine/Plugins/Runtime/MassGameplay/Source/MassSpawner/Private/MassSpawnerSubsystem.cpp
+++ b/Engine/Plugins/Runtime/MassGameplay/Source/MassSpawner/Private/MassSpawnerSubsystem.cpp
@@ -170,7 +170,7 @@ void UMassSpawnerSubsystem::DoSpawning(const FMassEntityTemplate& EntityTemplate
 	// 4. "OnEntitiesCreated" notifies will be sent out once the CreationContext gets destroyed (via its destructor).
 
 	TArray<FMassEntityHandle> SpawnedEntities;
-	TSharedRef<UMassEntitySubsystem::FEntityCreationContext> CreationContext = EntitySystem->BatchCreateEntities(EntityTemplate.GetArchetype(), NumToSpawn, SpawnedEntities);
+	TSharedRef<UMassEntitySubsystem::FEntityCreationContext> CreationContext = EntitySystem->BatchCreateEntities(EntityTemplate.GetArchetype(), EntityTemplate.GetSharedFragmentValues(), NumToSpawn, SpawnedEntities);
 
 	TConstArrayView<FInstancedStruct> FragmentInstances = EntityTemplate.GetInitialFragmentValues();
 	EntitySystem->BatchSetEntityFragmentsValues(CreationContext->GetEntityCollection(), FragmentInstances);
