Removed the Archetype filter
Optimization of the SharedFragmentValues bindigns when no entities are in the chunk
Small optimization when the Archetype has no entities

diff --git a/Engine/Plugins/Runtime/MassEntity/Source/MassEntity/Private/MassArchetypeData.cpp b/Engine/Plugins/Runtime/MassEntity/Source/MassEntity/Private/MassArchetypeData.cpp
--- a/Engine/Plugins/Runtime/MassEntity/Source/MassEntity/Private/MassArchetypeData.cpp
+++ b/Engine/Plugins/Runtime/MassEntity/Source/MassEntity/Private/MassArchetypeData.cpp
@@ -374,6 +374,11 @@ void FMassArchetypeData::MoveEntityToAnotherArchetype(const FMassEntityHandle En
 
 void FMassArchetypeData::ExecuteFunction(FMassExecutionContext& RunContext, const FMassExecuteFunction& Function, const FMassQueryRequirementIndicesMapping& RequirementMapping, FMassArchetypeEntityCollection::FConstEntityRangeArrayView EntityRangeContainer)
 {
+	if (GetNumEntities() == 0)
+	{
+		return;
+	}
+
 	// mz@todo to be removed
 	RunContext.SetCurrentArchetypesTagBitSet(GetTagBitSet());
 
@@ -382,17 +387,17 @@ void FMassArchetypeData::ExecuteFunction(FMassExecutionContext& RunContext, cons
 	{
 		FMassArchetypeChunk& Chunk = Chunks[ChunkIterator->ChunkIndex];
 
-		const uint32 SharedFragmentValuesHash = GetTypeHash(Chunk.GetSharedFragmentValues());
-		if (PrevSharedFragmentValuesHash != SharedFragmentValuesHash)
-		{
-			PrevSharedFragmentValuesHash = SharedFragmentValuesHash;
-			BindConstSharedFragmentRequirements(RunContext, Chunk.GetSharedFragmentValues(), RequirementMapping.ConstSharedFragments);
-			BindSharedFragmentRequirements(RunContext, Chunk.GetSharedFragmentValues(), RequirementMapping.SharedFragments);
-		}
-
 		const int32 ChunkLength = ChunkIterator->Length > 0 ? ChunkIterator->Length : (Chunk.GetNumInstances() - ChunkIterator->SubchunkStart);
 		if (ChunkLength)
 		{
+			const uint32 SharedFragmentValuesHash = GetTypeHash(Chunk.GetSharedFragmentValues());
+			if (PrevSharedFragmentValuesHash != SharedFragmentValuesHash)
+			{
+				PrevSharedFragmentValuesHash = SharedFragmentValuesHash;
+				BindConstSharedFragmentRequirements(RunContext, Chunk.GetSharedFragmentValues(), RequirementMapping.ConstSharedFragments);
+				BindSharedFragmentRequirements(RunContext, Chunk.GetSharedFragmentValues(), RequirementMapping.SharedFragments);
+			}
+
 			checkf((ChunkIterator->SubchunkStart + ChunkLength) <= Chunk.GetNumInstances() && ChunkLength > 0, TEXT("Invalid subchunk, it is going over the number of instances in the chunk or it is empty."));
 
 			RunContext.SetCurrentChunkSerialModificationNumber(Chunk.GetSerialModificationNumber());
@@ -404,15 +409,20 @@ void FMassArchetypeData::ExecuteFunction(FMassExecutionContext& RunContext, cons
 	}
 }
 
-void FMassArchetypeData::ExecuteFunction(FMassExecutionContext& RunContext, const FMassExecuteFunction& Function, const FMassQueryRequirementIndicesMapping& RequirementMapping, const FMassArchetypeConditionFunction& ArchetypeCondition, const FMassChunkConditionFunction& ChunkCondition)
+void FMassArchetypeData::ExecuteFunction(FMassExecutionContext& RunContext, const FMassExecuteFunction& Function, const FMassQueryRequirementIndicesMapping& RequirementMapping, const FMassChunkConditionFunction& ChunkCondition)
 {
+	if (GetNumEntities() == 0)
+	{
+		return;
+	}
+
 	// mz@todo to be removed
 	RunContext.SetCurrentArchetypesTagBitSet(GetTagBitSet());
 
-	if (!ArchetypeCondition || ArchetypeCondition(RunContext))
+	uint32 PrevSharedFragmentValuesHash = UINT32_MAX;
+	for (FMassArchetypeChunk& Chunk : Chunks)
 	{
-		uint32 PrevSharedFragmentValuesHash = UINT32_MAX;
-		for (FMassArchetypeChunk& Chunk : Chunks)
+		if (Chunk.GetNumInstances())
 		{
 			const uint32 SharedFragmentValuesHash = GetTypeHash(Chunk.GetSharedFragmentValues());
 			if (PrevSharedFragmentValuesHash != SharedFragmentValuesHash)
@@ -422,16 +432,13 @@ void FMassArchetypeData::ExecuteFunction(FMassExecutionContext& RunContext, cons
 				BindSharedFragmentRequirements(RunContext, Chunk.GetSharedFragmentValues(), RequirementMapping.SharedFragments);
 			}
 
-			if (Chunk.GetNumInstances())
-			{
-				RunContext.SetCurrentChunkSerialModificationNumber(Chunk.GetSerialModificationNumber());
-				BindChunkFragmentRequirements(RunContext, RequirementMapping.ChunkFragments, Chunk);
+			RunContext.SetCurrentChunkSerialModificationNumber(Chunk.GetSerialModificationNumber());
+			BindChunkFragmentRequirements(RunContext, RequirementMapping.ChunkFragments, Chunk);
 
-				if (!ChunkCondition || ChunkCondition(RunContext))
-				{
-					BindEntityRequirements(RunContext, RequirementMapping.EntityFragments, Chunk, 0, Chunk.GetNumInstances());
-					Function(RunContext);
-				}
+			if (!ChunkCondition || ChunkCondition(RunContext))
+			{
+				BindEntityRequirements(RunContext, RequirementMapping.EntityFragments, Chunk, 0, Chunk.GetNumInstances());
+				Function(RunContext);
 			}
 		}
 	}
diff --git a/Engine/Plugins/Runtime/MassEntity/Source/MassEntity/Private/MassArchetypeData.h b/Engine/Plugins/Runtime/MassEntity/Source/MassEntity/Private/MassArchetypeData.h
--- a/Engine/Plugins/Runtime/MassEntity/Source/MassEntity/Private/MassArchetypeData.h
+++ b/Engine/Plugins/Runtime/MassEntity/Source/MassEntity/Private/MassArchetypeData.h
@@ -233,7 +233,7 @@ public:
 	int32 GetChunkCount() const { return Chunks.Num(); }
 
 	void ExecuteFunction(FMassExecutionContext& RunContext, const FMassExecuteFunction& Function, const FMassQueryRequirementIndicesMapping& RequirementMapping, FMassArchetypeEntityCollection::FConstEntityRangeArrayView EntityRangeContainer);
-	void ExecuteFunction(FMassExecutionContext& RunContext, const FMassExecuteFunction& Function, const FMassQueryRequirementIndicesMapping& RequirementMapping, const FMassArchetypeConditionFunction& ArchetypeCondition, const FMassChunkConditionFunction& ChunkCondition);
+	void ExecuteFunction(FMassExecutionContext& RunContext, const FMassExecuteFunction& Function, const FMassQueryRequirementIndicesMapping& RequirementMapping, const FMassChunkConditionFunction& ChunkCondition);
 
 	void ExecutionFunctionForChunk(FMassExecutionContext RunContext, const FMassExecuteFunction& Function, const FMassQueryRequirementIndicesMapping& RequirementMapping, const FMassArchetypeEntityCollection::FArchetypeEntityRange& EntityRange, const FMassChunkConditionFunction& ChunkCondition = FMassChunkConditionFunction());
 
diff --git a/Engine/Plugins/Runtime/MassEntity/Source/MassEntity/Private/MassEntityQuery.cpp b/Engine/Plugins/Runtime/MassEntity/Source/MassEntity/Private/MassEntityQuery.cpp
--- a/Engine/Plugins/Runtime/MassEntity/Source/MassEntity/Private/MassEntityQuery.cpp
+++ b/Engine/Plugins/Runtime/MassEntity/Source/MassEntity/Private/MassEntityQuery.cpp
@@ -190,7 +190,7 @@ void FMassEntityQuery::ForEachEntityChunk(UMassEntitySubsystem& EntitySubsystem,
 		{
 			const FMassArchetypeHandle& ArchetypeHandle = ValidArchetypes[i];
 			FMassArchetypeData& ArchetypeData = FMassArchetypeHelper::ArchetypeDataFromHandleChecked(ArchetypeHandle);
-			ArchetypeData.ExecuteFunction(ExecutionContext, ExecuteFunction, ArchetypeFragmentMapping[i], ArchetypeCondition, ChunkCondition);
+			ArchetypeData.ExecuteFunction(ExecutionContext, ExecuteFunction, ArchetypeFragmentMapping[i], ChunkCondition);
 			ExecutionContext.ClearFragmentViews();
 #if WITH_MASSENTITY_DEBUG
 			NumEntitiesToProcess += ExecutionContext.GetNumEntities();
diff --git a/Engine/Plugins/Runtime/MassEntity/Source/MassEntity/Public/MassArchetypeTypes.h b/Engine/Plugins/Runtime/MassEntity/Source/MassEntity/Public/MassArchetypeTypes.h
--- a/Engine/Plugins/Runtime/MassEntity/Source/MassEntity/Public/MassArchetypeTypes.h
+++ b/Engine/Plugins/Runtime/MassEntity/Source/MassEntity/Public/MassArchetypeTypes.h
@@ -21,7 +21,6 @@ struct FMassArchetypeHelper;
 
 typedef TFunction< void(FMassExecutionContext& /*ExecutionContext*/) > FMassExecuteFunction;
 typedef TFunction< bool(const FMassExecutionContext& /*ExecutionContext*/) > FMassChunkConditionFunction;
-typedef TFunction< bool(const FMassExecutionContext& /*ExecutionContext*/) > FMassArchetypeConditionFunction;
 
 
 //////////////////////////////////////////////////////////////////////
diff --git a/Engine/Plugins/Runtime/MassEntity/Source/MassEntity/Public/MassEntityQuery.h b/Engine/Plugins/Runtime/MassEntity/Source/MassEntity/Public/MassEntityQuery.h
--- a/Engine/Plugins/Runtime/MassEntity/Source/MassEntity/Public/MassEntityQuery.h
+++ b/Engine/Plugins/Runtime/MassEntity/Source/MassEntity/Public/MassEntityQuery.h
@@ -92,22 +92,12 @@ public:
 	 * that this condition won't be applied when a specific entity colleciton is used (via FMassArchetypeEntityCollection )
 	 * The value returned by InFunction controls whether to allow execution (true) or block it (false).
 	 */
-	void SetChunkFilter(const FMassChunkConditionFunction& InFunction) { ChunkCondition = InFunction; }
+	void SetChunkFilter(const FMassChunkConditionFunction& InFunction) { checkf(!HasChunkFilter(), TEXT("Chunk filter needs to be cleared before setting a new one.")); ChunkCondition = InFunction; }
 
 	void ClearChunkFilter() { ChunkCondition.Reset(); }
 
 	bool HasChunkFilter() const { return bool(ChunkCondition); }
 
-	/**
-	 * Sets a archetype filter condition that will applied to each valid archetypes.
-	 * The value returned by InFunction controls whether to allow execution (true) or block it (false).
-	 */
-	void SetArchetypeFilter(const FMassArchetypeConditionFunction& InFunction) { ArchetypeCondition = InFunction; }
-
-	void ClearArchetypeFilter() { ArchetypeCondition.Reset(); }
-
-	bool HasArchetypeFilter() const { return bool(ArchetypeCondition); }
-
 	/** 
 	 * If ArchetypeHandle is among ValidArchetypes then the function retrieves requirements mapping cached for it,
 	 * otherwise an empty mapping will be returned (and the requirements binding will be done the slow way).
@@ -127,13 +117,6 @@ private:
 	 */
 	FMassChunkConditionFunction ChunkCondition;
 
-	/**
-	 * This function represents a condition that will be called for every archetype to be processed before the actual
-	 * execution function is called. The shared fragment requirements are already bound and ready to be used by the time
-	 * ArchetypeCondition is executed.
-	 */
-	FMassArchetypeConditionFunction ArchetypeCondition;
-
 	uint32 EntitySubsystemHash = 0;
 	uint32 ArchetypeDataVersion = 0;
 
