Do incremental query updates
Saves ~2ms globally over 2.4ms at start of the game, not a big deal but could become one if many new archetypes are created.

diff --git a/Engine/Plugins/Runtime/MassEntity/Source/MassEntity/Private/MassArchetypeData.cpp b/Engine/Plugins/Runtime/MassEntity/Source/MassEntity/Private/MassArchetypeData.cpp
--- a/Engine/Plugins/Runtime/MassEntity/Source/MassEntity/Private/MassArchetypeData.cpp
+++ b/Engine/Plugins/Runtime/MassEntity/Source/MassEntity/Private/MassArchetypeData.cpp
@@ -20,8 +20,9 @@ bool FMassArchetypeData::HasFragmentType(const UScriptStruct* FragmentType) cons
 	return (FragmentType && CompositionDescriptor.Fragments.Contains(*FragmentType));
 }
 
-void FMassArchetypeData::Initialize(const FMassArchetypeCompositionDescriptor& InCompositionDescriptor)
+void FMassArchetypeData::Initialize(const FMassArchetypeCompositionDescriptor& InCompositionDescriptor, const uint32 ArchetypeDataVersion)
 {
+	CreatedArchetypeDataVersion = ArchetypeDataVersion;
 	CompositionDescriptor.Fragments = InCompositionDescriptor.Fragments;
 	ConfigureFragments();
 
@@ -45,10 +46,12 @@ void FMassArchetypeData::Initialize(const FMassArchetypeCompositionDescriptor& I
 	EntityListOffsetWithinChunk = 0;
 }
 
-void FMassArchetypeData::InitializeWithSimilar(const FMassArchetypeData& BaseArchetype, FMassArchetypeCompositionDescriptor&& NewComposition)
+void FMassArchetypeData::InitializeWithSimilar(const FMassArchetypeData& BaseArchetype, FMassArchetypeCompositionDescriptor&& NewComposition, const uint32 ArchetypeDataVersion)
 {
 	checkf(IsInitialized() == false, TEXT("Trying to %s but this archetype has already been initialized"));
 
+	CreatedArchetypeDataVersion = ArchetypeDataVersion;
+
 	// note that we're calling this function rarely, so we can be a little bit inefficient here.
 	CompositionDescriptor = MoveTemp(NewComposition);
 	if (CompositionDescriptor.Fragments != NewComposition.Fragments)
diff --git a/Engine/Plugins/Runtime/MassEntity/Source/MassEntity/Private/MassArchetypeData.h b/Engine/Plugins/Runtime/MassEntity/Source/MassEntity/Private/MassArchetypeData.h
--- a/Engine/Plugins/Runtime/MassEntity/Source/MassEntity/Private/MassArchetypeData.h
+++ b/Engine/Plugins/Runtime/MassEntity/Source/MassEntity/Private/MassArchetypeData.h
@@ -171,6 +171,9 @@ private:
 	int32 TotalBytesPerEntity;
 	int32 EntityListOffsetWithinChunk;
 
+	// Archetype version at which this archetype was created, useful for query to do incremental archetype matching
+	uint32 CreatedArchetypeDataVersion;
+
 	// Arrays of names the archetype is referred as.
 	TArray<FName> DebugNames;
 	
@@ -208,13 +211,13 @@ public:
 		return CompositionDescriptor.IsEquivalent(OtherCompositionDescriptor);
 	}
 
-	void Initialize(const FMassArchetypeCompositionDescriptor& InCompositionDescriptor);
+	void Initialize(const FMassArchetypeCompositionDescriptor& InCompositionDescriptor, const uint32 ArchetypeDataVersion);
 
 	/** 
 	 * A special way of initializing an archetype resulting in a copy of SiblingArchetype's setup with OverrideTags
 	 * replacing original tags of SiblingArchetype
 	 */
-	void InitializeWithSimilar(const FMassArchetypeData& BaseArchetype, FMassArchetypeCompositionDescriptor&& NewComposition);
+	void InitializeWithSimilar(const FMassArchetypeData& BaseArchetype, FMassArchetypeCompositionDescriptor&& NewComposition, const uint32 ArchetypeDataVersion);
 
 	void AddEntity(FMassEntityHandle Entity, const FMassArchetypeSharedFragmentValues& InSharedFragmentValues);
 	void RemoveEntity(FMassEntityHandle Entity);
@@ -232,6 +235,8 @@ public:
 
 	int32 GetChunkCount() const { return Chunks.Num(); }
 
+	uint32 GetCreatedArchetypeDataVersion() const { return CreatedArchetypeDataVersion; }
+
 	void ExecuteFunction(FMassExecutionContext& RunContext, const FMassExecuteFunction& Function, const FMassQueryRequirementIndicesMapping& RequirementMapping, FMassArchetypeEntityCollection::FConstEntityRangeArrayView EntityRangeContainer);
 	void ExecuteFunction(FMassExecutionContext& RunContext, const FMassExecuteFunction& Function, const FMassQueryRequirementIndicesMapping& RequirementMapping, const FMassChunkConditionFunction& ChunkCondition);
 
diff --git a/Engine/Plugins/Runtime/MassEntity/Source/MassEntity/Private/MassEntityQuery.cpp b/Engine/Plugins/Runtime/MassEntity/Source/MassEntity/Private/MassEntityQuery.cpp
--- a/Engine/Plugins/Runtime/MassEntity/Source/MassEntity/Private/MassEntityQuery.cpp
+++ b/Engine/Plugins/Runtime/MassEntity/Source/MassEntity/Private/MassEntityQuery.cpp
@@ -71,24 +71,49 @@ void FMassEntityQuery::ReadCommandlineParams()
 void FMassEntityQuery::CacheArchetypes(const UMassEntitySubsystem& InEntitySubsystem)
 {
 	const uint32 InEntitySubsystemHash = PointerHash(&InEntitySubsystem);
-	if (EntitySubsystemHash != InEntitySubsystemHash || InEntitySubsystem.GetArchetypeDataVersion() != ArchetypeDataVersion || IncrementalChangesCount)
+
+	// Do an incremental update if the last updated archetype data version is different 
+    bool bUpdateArchetypes = InEntitySubsystem.GetArchetypeDataVersion() != LastUpdatedArchetypeDataVersion;
+
+	// Force a full update if the entity system changed or if the requirements changed
+	if (EntitySubsystemHash != InEntitySubsystemHash || IncrementalChangesCount)
 	{
-		IncrementalChangesCount = 0;
+		bUpdateArchetypes = true;
+		EntitySubsystemHash = InEntitySubsystemHash;
+		ValidArchetypes.Reset();
+		LastUpdatedArchetypeDataVersion = 0;
+		ArchetypeFragmentMapping.Reset();
 
-		if (CheckValidity())
+		if (IncrementalChangesCount)
 		{
-			SortRequirements();
-
-			EntitySubsystemHash = InEntitySubsystemHash;
-			ValidArchetypes.Reset();
-			InEntitySubsystem.GetValidArchetypes(*this, ValidArchetypes);
-			ArchetypeDataVersion = InEntitySubsystem.GetArchetypeDataVersion();
+			IncrementalChangesCount = 0;
+			if( CheckValidity() )
+			{
+				SortRequirements();
+			}
+			else
+			{
+				bUpdateArchetypes = false;
+				UE_VLOG_UELOG(&InEntitySubsystem, LogMass, Error, TEXT("FMassEntityQuery::CacheArchetypes: requirements not valid: %s"), *FMassDebugger::GetRequirementsDescription(*this));
+			}
+		}
+	}
+	
+	// Process any new archetype that is newer than the LastUpdatedArchetypeDataVersion
+	if (bUpdateArchetypes)
+	{
+		TArray<FMassArchetypeHandle> NewValidArchetypes;
+		InEntitySubsystem.GetValidArchetypes(*this, NewValidArchetypes, LastUpdatedArchetypeDataVersion);
+		LastUpdatedArchetypeDataVersion = InEntitySubsystem.GetArchetypeDataVersion();
+		if (NewValidArchetypes.Num())
+		{
+			const int32 FirstNewArchetype = ValidArchetypes.Num();
+			ValidArchetypes.Append(NewValidArchetypes);
 
 			TRACE_CPUPROFILER_EVENT_SCOPE_STR("Mass RequirementsBinding")
 			const TConstArrayView<FMassFragmentRequirementDescription> LocalRequirements = GetFragmentRequirements();
-			ArchetypeFragmentMapping.Reset(ValidArchetypes.Num());
-			ArchetypeFragmentMapping.AddDefaulted(ValidArchetypes.Num());
-			for (int i = 0; i < ValidArchetypes.Num(); ++i)
+			ArchetypeFragmentMapping.AddDefaulted(NewValidArchetypes.Num());
+			for (int i = FirstNewArchetype; i < ValidArchetypes.Num(); ++i)
 			{
 				FMassArchetypeData& ArchetypeData = FMassArchetypeHelper::ArchetypeDataFromHandleChecked(ValidArchetypes[i]);
 				ArchetypeData.GetRequirementsFragmentMapping(LocalRequirements, ArchetypeFragmentMapping[i].EntityFragments);
@@ -106,10 +131,6 @@ void FMassEntityQuery::CacheArchetypes(const UMassEntitySubsystem& InEntitySubsy
 				}
 			}
 		}
-		else
-		{
-			UE_VLOG_UELOG(&InEntitySubsystem, LogMass, Error, TEXT("FMassEntityQuery::CacheArchetypes: requirements not valid: %s"), *FMassDebugger::GetRequirementsDescription(*this));
-		}
 	}
 }
 
diff --git a/Engine/Plugins/Runtime/MassEntity/Source/MassEntity/Private/MassEntitySubsystem.cpp b/Engine/Plugins/Runtime/MassEntity/Source/MassEntity/Private/MassEntitySubsystem.cpp
--- a/Engine/Plugins/Runtime/MassEntity/Source/MassEntity/Private/MassEntitySubsystem.cpp
+++ b/Engine/Plugins/Runtime/MassEntity/Source/MassEntity/Private/MassEntitySubsystem.cpp
@@ -190,9 +190,12 @@ FMassArchetypeHandle UMassEntitySubsystem::CreateArchetype(const FMassArchetypeC
 
 	if (!ArchetypeDataPtr.IsValid())
 	{
+		// Important to pre-increment the version as the queries will use this value to do incremental updates
+		++ArchetypeDataVersion;
+
 		// Create a new archetype
 		FMassArchetypeData* NewArchetype = new FMassArchetypeData();
-		NewArchetype->Initialize(Composition);
+		NewArchetype->Initialize(Composition, ArchetypeDataVersion);
 		if (!ArchetypeDebugName.IsNone())
 		{
 			NewArchetype->AddUniqueDebugName(ArchetypeDebugName);
@@ -204,8 +207,6 @@ FMassArchetypeHandle UMassEntitySubsystem::CreateArchetype(const FMassArchetypeC
 			checkSlow(FragmentConfig.FragmentType)
 			FragmentTypeToArchetypeMap.FindOrAdd(FragmentConfig.FragmentType).Add(ArchetypeDataPtr);
 		}
-
-		++ArchetypeDataVersion;
 	}
 
 	return FMassArchetypeHelper::ArchetypeHandleFromData(ArchetypeDataPtr);
@@ -245,9 +246,12 @@ FMassArchetypeHandle UMassEntitySubsystem::InternalCreateSimilarArchetype(const
 
 	if (!ArchetypeDataPtr.IsValid())
 	{
+		// Important to pre-increment the version as the queries will use this value to do incremental updates
+		++ArchetypeDataVersion;
+
 		// Create a new archetype
 		FMassArchetypeData* NewArchetype = new FMassArchetypeData();
-		NewArchetype->InitializeWithSimilar(SourceArchetypeRef, MoveTemp(NewComposition));
+		NewArchetype->InitializeWithSimilar(SourceArchetypeRef, MoveTemp(NewComposition), ArchetypeDataVersion);
 		NewArchetype->CopyDebugNamesFrom(SourceArchetypeRef);
 
 		ArchetypeDataPtr = HashRow.Add_GetRef(MakeShareable(NewArchetype));
@@ -257,8 +261,6 @@ FMassArchetypeHandle UMassEntitySubsystem::InternalCreateSimilarArchetype(const
 			checkSlow(FragmentConfig.FragmentType)
 			FragmentTypeToArchetypeMap.FindOrAdd(FragmentConfig.FragmentType).Add(ArchetypeDataPtr);
 		}
-
-		++ArchetypeDataVersion;
 	}
 
 	return FMassArchetypeHelper::ArchetypeHandleFromData(ArchetypeDataPtr);
@@ -1097,7 +1099,7 @@ void UMassEntitySubsystem::CheckIfEntityIsActive(FMassEntityHandle Entity) const
 	checkf(IsEntityBuilt(Entity), TEXT("Entity not yet created(ID: %d, SN:%d)"));
 }
 
-void UMassEntitySubsystem::GetValidArchetypes(const FMassEntityQuery& Query, TArray<FMassArchetypeHandle>& OutValidArchetypes) const
+void UMassEntitySubsystem::GetValidArchetypes(const FMassEntityQuery& Query, TArray<FMassArchetypeHandle>& OutValidArchetypes, const uint32 FromArchetypeDataVersion) const
 {
 	//@TODO: Not optimized yet, but we call this rarely now, so not a big deal.
 
@@ -1120,6 +1122,12 @@ void UMassEntitySubsystem::GetValidArchetypes(const FMassEntityQuery& Query, TAr
 	{
 		FMassArchetypeData& Archetype = *(ArchetypePtr.Get());
 
+		// Only return archetypes with a newer created version than the specified version, this is for incremental query updates
+		if (Archetype.GetCreatedArchetypeDataVersion() <= FromArchetypeDataVersion)
+		{
+			continue;
+		}
+
 		if (Archetype.GetTagBitSet().HasAll(Query.GetRequiredAllTags()) == false)
 		{
 			// missing some required tags, skip.
diff --git a/Engine/Plugins/Runtime/MassEntity/Source/MassEntity/Public/MassEntityQuery.h b/Engine/Plugins/Runtime/MassEntity/Source/MassEntity/Public/MassEntityQuery.h
--- a/Engine/Plugins/Runtime/MassEntity/Source/MassEntity/Public/MassEntityQuery.h
+++ b/Engine/Plugins/Runtime/MassEntity/Source/MassEntity/Public/MassEntityQuery.h
@@ -56,7 +56,7 @@ public:
 	FORCEINLINE void DirtyCachedData()
 	{
 		EntitySubsystemHash = 0;
-		ArchetypeDataVersion = 0;
+		LastUpdatedArchetypeDataVersion = 0;
 	}
 	
 	bool DoesRequireGameThreadExecution() const 
@@ -118,7 +118,7 @@ private:
 	FMassChunkConditionFunction ChunkCondition;
 
 	uint32 EntitySubsystemHash = 0;
-	uint32 ArchetypeDataVersion = 0;
+	uint32 LastUpdatedArchetypeDataVersion = 0;
 
 	TArray<FMassArchetypeHandle> ValidArchetypes;
 	TArray<FMassQueryRequirementIndicesMapping> ArchetypeFragmentMapping;
diff --git a/Engine/Plugins/Runtime/MassEntity/Source/MassEntity/Public/MassEntitySubsystem.h b/Engine/Plugins/Runtime/MassEntity/Source/MassEntity/Public/MassEntitySubsystem.h
--- a/Engine/Plugins/Runtime/MassEntity/Source/MassEntity/Public/MassEntitySubsystem.h
+++ b/Engine/Plugins/Runtime/MassEntity/Source/MassEntity/Public/MassEntitySubsystem.h
@@ -376,7 +376,7 @@ public:
 #endif // WITH_MASSENTITY_DEBUG
 
 protected:
-	void GetValidArchetypes(const FMassEntityQuery& Query, TArray<FMassArchetypeHandle>& OutValidArchetypes) const;
+	void GetValidArchetypes(const FMassEntityQuery& Query, TArray<FMassArchetypeHandle>& OutValidArchetypes, const uint32 FromArchetypeDataVersion) const;
 	
 	/** 
 	 * A "similar" archetype is an archetype exactly the same as SourceArchetype except for one composition aspect 
