Changed how we treat missing dependencies in Mass dependency solving - used to be an error, now we create dummy temporary nodes to maintain before-after dependencies even if used in reference to a non-existing item.

diff --git a/Engine/Plugins/Runtime/MassEntity/Source/MassEntity/Private/MassProcessorDependencySolver.cpp b/Engine/Plugins/Runtime/MassEntity/Source/MassEntity/Private/MassProcessorDependencySolver.cpp
--- a/Engine/Plugins/Runtime/MassEntity/Source/MassEntity/Private/MassProcessorDependencySolver.cpp
+++ b/Engine/Plugins/Runtime/MassEntity/Source/MassEntity/Private/MassProcessorDependencySolver.cpp
@@ -320,23 +320,34 @@ void FProcessorDependencySolver::BuildDependencies()
 {
 	// at this point we have collected all the known processors and groups in AllNodes so we can transpose 
 	// A.ExecuteBefore(B) type of dependencies into B.ExecuteAfter(A)
-	for (FNode& Node : AllNodes)
+	for (int32 i = 0; i < AllNodes.Num(); ++i)
 	{
-		for (const FName& BeforeDependencyName : Node.ExecuteBefore)
+		for (const FName& BeforeDependencyName : AllNodes[i].ExecuteBefore)
 		{
-			int32* DependentNodeIndex = NodeIndexMap.Find(BeforeDependencyName);
-			if (DependentNodeIndex != nullptr)
+			int32 DependentNodeIndex = INDEX_NONE;
+			int32* DependentNodeIndexPtr = NodeIndexMap.Find(BeforeDependencyName);
+			if (DependentNodeIndexPtr == nullptr)
 			{
-				check(AllNodes.IsValidIndex(*DependentNodeIndex));
-				AllNodes[*DependentNodeIndex].ExecuteAfter.Add(Node.Name);
+				// missing dependency. Adding a "dummy" node representing those to still support ordering based on missing groups or processors 
+				// For example, if Processor A and B declare dependency, respectively, "Before C" and "After C" we still 
+				// expect A to come before B regardless of whether C exists or not.
+				
+				DependentNodeIndex = AllNodes.Num();
+				NodeIndexMap.Add(BeforeDependencyName, DependentNodeIndex);
+				AllNodes.Add({ BeforeDependencyName, nullptr, DependentNodeIndex });
+
+				UE_LOG(LogMass, Log, TEXT("Unable to find dependency \"%s\" declared by %s. Creating a dummy dependency node.")
+					, *BeforeDependencyName.ToString(), *AllNodes[i].Name.ToString());
 			}
 			else
 			{
-				UE_LOG(LogMass, Error, TEXT("Unable to find dependency \"%s\" declared by %s"), *BeforeDependencyName.ToString()
-					, *Node.Name.ToString());
+				DependentNodeIndex = *DependentNodeIndexPtr;
 			}
+
+			check(AllNodes.IsValidIndex(DependentNodeIndex));
+			AllNodes[DependentNodeIndex].ExecuteAfter.Add(AllNodes[i].Name);
 		}
-		Node.ExecuteBefore.Reset();
+		AllNodes[i].ExecuteBefore.Reset();
 	}
 
 	// at this point all nodes contain:
@@ -350,45 +361,56 @@ void FProcessorDependencySolver::BuildDependencies()
 	// 
 	for (int32 NodeIndex = 0; NodeIndex < AllNodes.Num(); ++NodeIndex)
 	{
-		FNode& Node = AllNodes[NodeIndex];
-		for (int i = 0; i < Node.ExecuteAfter.Num(); ++i)
+		for (int i = 0; i < AllNodes[NodeIndex].ExecuteAfter.Num(); ++i)
 		{
-			const FName& AfterDependencyName = Node.ExecuteAfter[i];
-			int32* PrerequisiteNodeIndex = NodeIndexMap.Find(AfterDependencyName);
-			if (PrerequisiteNodeIndex != nullptr)
+			const FName& AfterDependencyName = AllNodes[NodeIndex].ExecuteAfter[i];
+			int32* PrerequisiteNodeIndexPtr = NodeIndexMap.Find(AfterDependencyName);
+			int32 PrerequisiteNodeIndex = INDEX_NONE;
+
+			if (PrerequisiteNodeIndexPtr == nullptr)
 			{
-				const FNode& PrerequisiteNode = AllNodes[*PrerequisiteNodeIndex];
+				// missing dependency. Adding a "dummy" node representing those to still support ordering based on missing groups or processors 
+				// For example, if Processor A and B declare dependency, respectively, "Before C" and "After C" we still 
+				// expect A to come before B regardless of whether C exists or not.
 
-				if (PrerequisiteNode.IsGroup())
-				{
-					for (int32 SubNodeIndex : PrerequisiteNode.SubNodeIndices)
-					{
-						Node.ExecuteAfter.AddUnique(AllNodes[SubNodeIndex].Name);
-					}
-				}
-				else
+				PrerequisiteNodeIndex = AllNodes.Num();
+				NodeIndexMap.Add(AfterDependencyName, PrerequisiteNodeIndex);
+				AllNodes.Add({ AfterDependencyName, nullptr, PrerequisiteNodeIndex });
+
+				UE_LOG(LogMass, Log, TEXT("Unable to find dependency \"%s\" declared by %s. Creating a dummy dependency node.")
+					, *AfterDependencyName.ToString(), *AllNodes[NodeIndex].Name.ToString());
+			}
+			else
+			{
+				PrerequisiteNodeIndex = *PrerequisiteNodeIndexPtr;
+			}
+
+			const FNode& PrerequisiteNode = AllNodes[PrerequisiteNodeIndex];
+
+			if (PrerequisiteNode.IsGroup())
+			{
+				for (int32 SubNodeIndex : PrerequisiteNode.SubNodeIndices)
 				{
-					Node.OriginalDependencies.AddUnique(*PrerequisiteNodeIndex);
+					AllNodes[NodeIndex].ExecuteAfter.AddUnique(AllNodes[SubNodeIndex].Name);
 				}
 			}
 			else
 			{
-				UE_LOG(LogMass, Error, TEXT("Unable to find dependency \"%s\" declared by %s"), *AfterDependencyName.ToString()
-					, *Node.Name.ToString());
+				AllNodes[NodeIndex].OriginalDependencies.AddUnique(PrerequisiteNodeIndex);
 			}
 		}
 
 		// if this node is a group push all the dependencies down on all the children
 		// by design all child nodes come after group nodes so the child nodes' dependencies have not been processed yet
-		if (Node.IsGroup() && Node.SubNodeIndices.Num())
+		if (AllNodes[NodeIndex].IsGroup() && AllNodes[NodeIndex].SubNodeIndices.Num())
 		{
-			for (int32 PrerequisiteNodeIndex : Node.OriginalDependencies)
+			for (int32 PrerequisiteNodeIndex : AllNodes[NodeIndex].OriginalDependencies)
 			{
 				checkSlow(PrerequisiteNodeIndex != NodeIndex);
 				// in case of processor nodes we can store it directly
 				if (AllNodes[PrerequisiteNodeIndex].IsGroup() == false)
 				{
-					for (int32 ChildNodeIndex : Node.SubNodeIndices)
+					for (int32 ChildNodeIndex : AllNodes[NodeIndex].SubNodeIndices)
 					{
 						AllNodes[ChildNodeIndex].OriginalDependencies.AddUnique(PrerequisiteNodeIndex);
 					}
@@ -397,7 +419,7 @@ void FProcessorDependencySolver::BuildDependencies()
 				else if (PrerequisiteNodeIndex > NodeIndex)
 				{
 					const FName& PrerequisiteName = AllNodes[PrerequisiteNodeIndex].Name;
-					for (int32 ChildNodeIndex : Node.SubNodeIndices)
+					for (int32 ChildNodeIndex : AllNodes[NodeIndex].SubNodeIndices)
 					{
 						AllNodes[ChildNodeIndex].ExecuteAfter.AddUnique(PrerequisiteName);
 					}
diff --git a/Engine/Plugins/Runtime/MassEntity/Source/MassEntityTestSuite/Private/MassDependencyTest.cpp b/Engine/Plugins/Runtime/MassEntity/Source/MassEntityTestSuite/Private/MassDependencyTest.cpp
--- a/Engine/Plugins/Runtime/MassEntity/Source/MassEntityTestSuite/Private/MassDependencyTest.cpp
+++ b/Engine/Plugins/Runtime/MassEntity/Source/MassEntityTestSuite/Private/MassDependencyTest.cpp
@@ -104,7 +104,6 @@ struct FMissingDependency : FDependencySolverBase
 			UMassTestProcessorBase* Proc = Processors.Add_GetRef(NewObject<UMassTestProcessor_B>());
 			Proc->GetMutableExecutionOrder().ExecuteBefore.Add(TEXT("NonExistingDependency2"));
 			Proc->GetMutableExecutionOrder().ExecuteAfter.Add(GetProcessorName<UMassTestProcessor_C>());
-			Proc->GetMutableExecutionOrder().ExecuteBefore.Add(GetProcessorName<UMassTestProcessor_A>());
 		}
 		{
 			UMassTestProcessorBase* Proc = Processors.Add_GetRef(NewObject<UMassTestProcessor_C>());
@@ -114,11 +113,11 @@ struct FMissingDependency : FDependencySolverBase
 
 	virtual bool InstantTest() override
 	{
-		// we're expecting two occurrences, for NonExistingDependency and NonExistingDependency2
-		GetTestRunner().AddExpectedError(TEXT("Unable to find dependency"), EAutomationExpectedErrorFlags::Contains, 2);
-
 		Solve();
 
+		// event though there's no direct dependency between A and B due to declared dependencies on "NonExistingDependency"
+		// B should come before A
+
 		AITEST_TRUE("C is expected to be the first one", Result[0].Name == GetProcessorName<UMassTestProcessor_C>());
 		AITEST_TRUE("Then B", Result[1].Name == GetProcessorName<UMassTestProcessor_B>());
 		AITEST_TRUE("With A being last", Result[2].Name == GetProcessorName<UMassTestProcessor_A>());
