Separated "entity management" and USubsystem stuff of UMassEntityManager. Now UMassEntityManager hosts a shared ref of newly created FMassEntityManager, and users are free to create their own FMassEnt

diff --git a/Engine/Plugins/AI/MassCrowd/Source/MassCrowd/Private/MassCrowdVisualizationProcessor.cpp b/Engine/Plugins/AI/MassCrowd/Source/MassCrowd/Private/MassCrowdVisualizationProcessor.cpp
--- a/Engine/Plugins/AI/MassCrowd/Source/MassCrowd/Private/MassCrowdVisualizationProcessor.cpp
+++ b/Engine/Plugins/AI/MassCrowd/Source/MassCrowd/Private/MassCrowdVisualizationProcessor.cpp
@@ -73,7 +73,7 @@ void UMassDebugCrowdVisualizationProcessor::Initialize(UObject& Owner)
 	check(World);
 }
 
-void UMassDebugCrowdVisualizationProcessor::Execute(UMassEntitySubsystem& EntitySubsystem, FMassExecutionContext& Context)
+void UMassDebugCrowdVisualizationProcessor::Execute(FMassEntityManager& EntityManager, FMassExecutionContext& Context)
 {
 	if (UE::MassCrowd::bDebugCrowdVisualType)
 	{
@@ -81,7 +81,7 @@ void UMassDebugCrowdVisualizationProcessor::Execute(UMassEntitySubsystem& Entity
 
 		TRACE_CPUPROFILER_EVENT_SCOPE(TEXT("DebugDisplayVisualType"))
 
-		EntityQuery.ForEachEntityChunk(EntitySubsystem, Context, [this](const FMassExecutionContext& Context)
+		EntityQuery.ForEachEntityChunk(EntityManager, Context, [this](const FMassExecutionContext& Context)
 		{
 			const TConstArrayView<FMassRepresentationFragment> VisualizationList = Context.GetFragmentView<FMassRepresentationFragment>();
 			const TConstArrayView<FMassActorFragment> ActorList = Context.GetFragmentView<FMassActorFragment>();
diff --git a/Engine/Plugins/AI/MassCrowd/Source/MassCrowd/Public/MassCrowdNavigationProcessor.h b/Engine/Plugins/AI/MassCrowd/Source/MassCrowd/Public/MassCrowdNavigationProcessor.h
--- a/Engine/Plugins/AI/MassCrowd/Source/MassCrowd/Public/MassCrowdNavigationProcessor.h
+++ b/Engine/Plugins/AI/MassCrowd/Source/MassCrowd/Public/MassCrowdNavigationProcessor.h
@@ -21,7 +21,7 @@ public:
 protected:
 	virtual void Initialize(UObject& Owner) override;
 	virtual void ConfigureQueries() override;
-	virtual void SignalEntities(UMassEntitySubsystem& EntitySubsystem, FMassExecutionContext& Context, FMassSignalNameLookup& EntitySignals) override;
+	virtual void SignalEntities(FMassEntityManager& EntityManager, FMassExecutionContext& Context, FMassSignalNameLookup& EntitySignals) override;
 };
 
 /** Processors that cleans up the lane tracking on entity destruction. */
@@ -35,7 +35,7 @@ public:
 
 protected:
 	virtual void ConfigureQueries() override;
-	virtual void Execute(UMassEntitySubsystem& EntitySubsystem, FMassExecutionContext& Context) override;
+	virtual void Execute(FMassEntityManager& EntityManager, FMassExecutionContext& Context) override;
 
 	FMassEntityQuery EntityQuery;
 };
@@ -52,7 +52,7 @@ public:
 protected:
 	virtual void Initialize(UObject& Owner) override;
 	virtual void ConfigureQueries() override;
-	virtual void Execute(UMassEntitySubsystem& EntitySubsystem, FMassExecutionContext& Context);
+	virtual void Execute(FMassEntityManager& EntityManager, FMassExecutionContext& Context);
 
 	void OnStop(FMassCrowdObstacleFragment& OutObstacle, const float Radius);
 	void OnMove(FMassCrowdObstacleFragment& OutObstacle);
@@ -74,7 +74,7 @@ public:
 
 protected:
 	virtual void ConfigureQueries() override;
-	virtual void Execute(UMassEntitySubsystem& EntitySubsystem, FMassExecutionContext& Context) override;
+	virtual void Execute(FMassEntityManager& EntityManager, FMassExecutionContext& Context) override;
 
 	FMassEntityQuery EntityQuery;
 };
@@ -91,7 +91,7 @@ public:
 protected:
 	virtual void Initialize(UObject& Owner) override;
 	virtual void ConfigureQueries() override;
-	virtual void Execute(UMassEntitySubsystem& EntitySubsystem, FMassExecutionContext& Context) override;
+	virtual void Execute(FMassEntityManager& EntityManager, FMassExecutionContext& Context) override;
 
 	UPROPERTY(Transient)
 	UZoneGraphAnnotationSubsystem* ZoneGraphAnnotationSubsystem;
diff --git a/Engine/Plugins/AI/MassCrowd/Source/MassCrowd/Public/MassCrowdRepresentationActorManagement.h b/Engine/Plugins/AI/MassCrowd/Source/MassCrowd/Public/MassCrowdRepresentationActorManagement.h
--- a/Engine/Plugins/AI/MassCrowd/Source/MassCrowd/Public/MassCrowdRepresentationActorManagement.h
+++ b/Engine/Plugins/AI/MassCrowd/Source/MassCrowd/Public/MassCrowdRepresentationActorManagement.h
@@ -38,7 +38,7 @@ protected:
 	 * @param Priority of this spawn request in comparison with the others, lower value means higher priority
 	 * @return the actor spawned
 	 */
-	virtual AActor* GetOrSpawnActor(UMassRepresentationSubsystem& RepresentationSubsystem, UMassEntitySubsystem& EntitySubsystem, const FMassEntityHandle MassAgent, FMassActorFragment& ActorInfo, const FTransform& Transform, const int16 TemplateActorIndex, FMassActorSpawnRequestHandle& SpawnRequestHandle, const float Priority) const override;
+	virtual AActor* GetOrSpawnActor(UMassRepresentationSubsystem& RepresentationSubsystem, FMassEntityManager& EntitySubsystem, const FMassEntityHandle MassAgent, FMassActorFragment& ActorInfo, const FTransform& Transform, const int16 TemplateActorIndex, FMassActorSpawnRequestHandle& SpawnRequestHandle, const float Priority) const override;
 
 	/**
 	 * Teleports the actor at the specified transform by preserving its velocity and without collision.
diff --git a/Engine/Plugins/Runtime/MassEntity/Source/MassEntity/Private/MassArchetypeTypes.cpp b/Engine/Plugins/Runtime/MassEntity/Source/MassEntity/Private/MassArchetypeTypes.cpp
--- a/Engine/Plugins/Runtime/MassEntity/Source/MassEntity/Private/MassArchetypeTypes.cpp
+++ b/Engine/Plugins/Runtime/MassEntity/Source/MassEntity/Private/MassArchetypeTypes.cpp
@@ -181,7 +181,7 @@ bool FMassArchetypeEntityCollection::IsSame(const FMassArchetypeEntityCollection
 //////////////////////////////////////////////////////////////////////
 // FMassArchetypeEntityCollectionWithPayload
 
-void FMassArchetypeEntityCollectionWithPayload::CreateEntityRangesWithPayload(const UMassEntitySubsystem& EntitySystem, const TConstArrayView<FMassEntityHandle> Entities
+void FMassArchetypeEntityCollectionWithPayload::CreateEntityRangesWithPayload(const FMassEntityManager& EntityManager, const TConstArrayView<FMassEntityHandle> Entities
 	, const FMassArchetypeEntityCollection::EDuplicatesHandling DuplicatesHandling, FMassGenericPayloadView Payload
 	, TArray<FMassArchetypeEntityCollectionWithPayload>& OutEntityCollections)
 {
@@ -227,7 +227,7 @@ void FMassArchetypeEntityCollectionWithPayload::CreateEntityRangesWithPayload(co
 	for (int32 i = 0; i < Entities.Num(); ++i)
 	{
 		const FMassEntityHandle& Entity = Entities[i];
-		const FMassArchetypeHandle ArchetypeHandle = EntitySystem.GetArchetypeForEntity(Entity);
+		const FMassArchetypeHandle ArchetypeHandle = EntityManager.GetArchetypeForEntity(Entity);
 		const FMassArchetypeData* ArchetypePtr = FMassArchetypeHelper::ArchetypeDataFromHandle(ArchetypeHandle);
 		
 		// @todo if FMassArchetypeHandle used indices the look up would be a lot faster
diff --git a/Engine/Plugins/Runtime/MassEntity/Source/MassEntity/Private/MassCommandBuffer.cpp b/Engine/Plugins/Runtime/MassEntity/Source/MassEntity/Private/MassCommandBuffer.cpp
--- a/Engine/Plugins/Runtime/MassEntity/Source/MassEntity/Private/MassCommandBuffer.cpp
+++ b/Engine/Plugins/Runtime/MassEntity/Source/MassEntity/Private/MassCommandBuffer.cpp
@@ -83,7 +83,7 @@ FMassCommandBuffer::~FMassCommandBuffer()
 	}
 }
 
-void FMassCommandBuffer::Flush(UMassEntitySubsystem& EntitySystem)
+void FMassCommandBuffer::Flush(FMassEntityManager& EntityManager)
 {
 	check(!bIsFlushing);
 	TGuardValue FlushingGuard(bIsFlushing, true);
@@ -150,7 +150,7 @@ void FMassCommandBuffer::Flush(UMassEntitySubsystem& EntitySystem)
 			FCsvProfiler::RecordCustomStat(*Name, CSV_CATEGORY_INDEX(MassEntitiesCounters), Command->GetNumOperationsStat(), ECsvCustomStatOp::Accumulate);
 #endif // CSV_PROFILER
 
-			Command->Execute(EntitySystem);
+			Command->Execute(EntityManager);
 			Command->Reset();
 		}
 	}
diff --git a/Engine/Plugins/Runtime/MassEntity/Source/MassEntity/Private/MassDebugger.cpp b/Engine/Plugins/Runtime/MassEntity/Source/MassEntity/Private/MassDebugger.cpp
--- a/Engine/Plugins/Runtime/MassEntity/Source/MassEntity/Private/MassDebugger.cpp
+++ b/Engine/Plugins/Runtime/MassEntity/Source/MassEntity/Private/MassDebugger.cpp
@@ -62,12 +62,12 @@ namespace UE::Mass::Debug
 			[](const TArray<FString>& Params, UWorld* World, FOutputDevice& Ar)
 			{
 				check(World);
-				if (UMassEntitySubsystem* EntitySystem = World->GetSubsystem<UMassEntitySubsystem>())
+				if (UMassEntitySubsystem* EntityManager = World->GetSubsystem<UMassEntitySubsystem>())
 				{
 					int32 Index = INDEX_NONE;
 					if (LexTryParseString<int32>(Index, *Params[0]))
 					{
-						FMassDebugger::OutputEntityDescription(Ar, *EntitySystem, Index);
+						FMassDebugger::OutputEntityDescription(Ar, EntityManager->GetEntityManager(), Index);
 					}
 					else
 					{
@@ -100,7 +100,7 @@ namespace UE::Mass::Debug
 						LexToString(World->WorldType),
 						*ToString(World->GetNetMode()));
 
-					if (UMassEntitySubsystem* EntitySystem = World->GetSubsystem<UMassEntitySubsystem>())
+					if (UMassEntitySubsystem* EntityManager = World->GetSubsystem<UMassEntitySubsystem>())
 					{
 						bool bIncludeEmpty = true;
 						if (Params.Num())
@@ -108,7 +108,7 @@ namespace UE::Mass::Debug
 							LexTryParseString(bIncludeEmpty, *Params[0]);
 						}
 						Ar.Logf(ELogVerbosity::Log, TEXT("Include empty archetypes: %s"), bIncludeEmpty ? TEXT("TRUE") : TEXT("FALSE"));
-						EntitySystem->DebugGetArchetypesStringDetails(Ar, bIncludeEmpty);
+						EntityManager->GetEntityManager().DebugGetArchetypesStringDetails(Ar, bIncludeEmpty);
 					}
 					else
 					{
@@ -130,7 +130,7 @@ namespace UE::Mass::Debug
 				check(InWorld);
 				if (UMassEntitySubsystem* System = InWorld->GetSubsystem<UMassEntitySubsystem>())
 				{
-					System->DebugForceArchetypeDataVersionBump();
+					System->GetMutableEntityManager().DebugForceArchetypeDataVersionBump();
 				}
 			}
 	));
@@ -207,13 +207,13 @@ TConstArrayView<FMassEntityQuery*> FMassDebugger::GetProcessorQueries(const UMas
 	return Processor.OwnedQueries;
 }
 
-TConstArrayView<FMassEntityQuery*> FMassDebugger::GetUpToDateProcessorQueries(const UMassEntitySubsystem& EntitySubsystem, UMassProcessor& Processor)
+TConstArrayView<FMassEntityQuery*> FMassDebugger::GetUpToDateProcessorQueries(const FMassEntityManager& EntityManager, UMassProcessor& Processor)
 {
 	for (FMassEntityQuery* Query : Processor.OwnedQueries)
 	{
 		if (Query)
 		{
-			Query->CacheArchetypes(EntitySubsystem);
+			Query->CacheArchetypes(EntityManager);
 		}
 	}
 
@@ -234,11 +234,11 @@ void FMassDebugger::GetQueryExecutionRequirements(const FMassEntityQuery& Query,
 	Query.ExportRequirements(OutExecutionRequirements);
 }
 
-TArray<FMassArchetypeHandle> FMassDebugger::GetAllArchetypes(const UMassEntitySubsystem& EntitySubsystem)
+TArray<FMassArchetypeHandle> FMassDebugger::GetAllArchetypes(const FMassEntityManager& EntityManager)
 {
 	TArray<FMassArchetypeHandle> Archetypes;
 
-	for (auto& KVP : EntitySubsystem.FragmentHashToArchetypeMap)
+	for (auto& KVP : EntityManager.FragmentHashToArchetypeMap)
 	{
 		for (const TSharedPtr<FMassArchetypeData>& Archetype : KVP.Value)
 		{
@@ -377,40 +377,40 @@ void FMassDebugger::OutputArchetypeDescription(FOutputDevice& Ar, const FMassArc
 	Ar.Logf(TEXT("%s"), ArchetypeHandle.IsValid() ? *FMassArchetypeHelper::ArchetypeDataFromHandleChecked(ArchetypeHandle).DebugGetDescription() : TEXT("INVALID"));
 }
 
-void FMassDebugger::OutputEntityDescription(FOutputDevice& Ar, const UMassEntitySubsystem& EntitySubsystem, const int32 EntityIndex, const TCHAR* InPrefix)
+void FMassDebugger::OutputEntityDescription(FOutputDevice& Ar, const FMassEntityManager& EntityManager, const int32 EntityIndex, const TCHAR* InPrefix)
 {
-	if (EntityIndex >= EntitySubsystem.Entities.Num())
+	if (EntityIndex >= EntityManager.Entities.Num())
 	{
-		Ar.Logf(ELogVerbosity::Log, TEXT("Unable to list fragments values for out of range index in %s"), *GetPathNameSafe(&EntitySubsystem));
+		Ar.Logf(ELogVerbosity::Log, TEXT("Unable to list fragments values for out of range index in EntityManager owned by %s"), *GetPathNameSafe(EntityManager.GetOwner()));
 		return;
 	}
 
-	const UMassEntitySubsystem::FEntityData& EntityData = EntitySubsystem.Entities[EntityIndex];
+	const FMassEntityManager::FEntityData& EntityData = EntityManager.Entities[EntityIndex];
 	if (!EntityData.IsValid())
 	{
-		Ar.Logf(ELogVerbosity::Log, TEXT("Unable to list fragments values for invalid entity in %s"), *GetPathNameSafe(&EntitySubsystem));
+		Ar.Logf(ELogVerbosity::Log, TEXT("Unable to list fragments values for invalid entity in EntityManager owned by %s"), *GetPathNameSafe(EntityManager.GetOwner()));
 	}
 
 	FMassEntityHandle Entity;
 	Entity.Index = EntityIndex;
 	Entity.SerialNumber = EntityData.SerialNumber;
-	OutputEntityDescription(Ar, EntitySubsystem, Entity, InPrefix);
+	OutputEntityDescription(Ar, EntityManager, Entity, InPrefix);
 }
 
-void FMassDebugger::OutputEntityDescription(FOutputDevice& Ar, const UMassEntitySubsystem& EntitySubsystem, const FMassEntityHandle Entity, const TCHAR* InPrefix)
+void FMassDebugger::OutputEntityDescription(FOutputDevice& Ar, const FMassEntityManager& EntityManager, const FMassEntityHandle Entity, const TCHAR* InPrefix)
 {
-	if (!EntitySubsystem.IsEntityActive(Entity))
+	if (!EntityManager.IsEntityActive(Entity))
 	{
-		Ar.Logf(ELogVerbosity::Log, TEXT("Unable to list fragments values for invalid entity in %s"), *GetPathNameSafe(&EntitySubsystem));
+		Ar.Logf(ELogVerbosity::Log, TEXT("Unable to list fragments values for invalid entity in EntityManager owned by %s"), *GetPathNameSafe(EntityManager.GetOwner()));
 	}
 
-	Ar.Logf(ELogVerbosity::Log, TEXT("Listing fragments values for Entity[%s] in %s"), *Entity.DebugGetDescription(), *GetPathNameSafe(&EntitySubsystem));
+	Ar.Logf(ELogVerbosity::Log, TEXT("Listing fragments values for Entity[%s] in EntityManager owned by %s"), *Entity.DebugGetDescription(), *GetPathNameSafe(EntityManager.GetOwner()));
 
-	const UMassEntitySubsystem::FEntityData& EntityData = EntitySubsystem.Entities[Entity.Index];
+	const FMassEntityManager::FEntityData& EntityData = EntityManager.Entities[Entity.Index];
 	FMassArchetypeData* Archetype = EntityData.CurrentArchetype.Get();
 	if (Archetype == nullptr)
 	{
-		Ar.Logf(ELogVerbosity::Log, TEXT("Unable to list fragments values for invalid entity in %s"), *GetPathNameSafe(&EntitySubsystem));
+		Ar.Logf(ELogVerbosity::Log, TEXT("Unable to list fragments values for invalid entity in EntityManager owned by %s"), *GetPathNameSafe(EntityManager.GetOwner()));
 	}
 	else
 	{
@@ -418,9 +418,9 @@ void FMassDebugger::OutputEntityDescription(FOutputDevice& Ar, const UMassEntity
 	}
 }
 
-void FMassDebugger::SelectEntity(const UMassEntitySubsystem& EntitySubsystem, const FMassEntityHandle EntityHandle)
+void FMassDebugger::SelectEntity(const FMassEntityManager& EntityManager, const FMassEntityHandle EntityHandle)
 {
-	OnEntitySelectedDelegate.Broadcast(EntitySubsystem, EntityHandle);
+	OnEntitySelectedDelegate.Broadcast(EntityManager, EntityHandle);
 }
 
 #endif // WITH_MASSENTITY_DEBUG
\ No newline at end of file
diff --git a/Engine/Plugins/Runtime/MassEntity/Source/MassEntity/Private/MassEntityQuery.cpp b/Engine/Plugins/Runtime/MassEntity/Source/MassEntity/Private/MassEntityQuery.cpp
--- a/Engine/Plugins/Runtime/MassEntity/Source/MassEntity/Private/MassEntityQuery.cpp
+++ b/Engine/Plugins/Runtime/MassEntity/Source/MassEntity/Private/MassEntityQuery.cpp
@@ -68,18 +68,18 @@ void FMassEntityQuery::ReadCommandlineParams()
 	}
 }
 
-void FMassEntityQuery::CacheArchetypes(const UMassEntitySubsystem& InEntitySubsystem)
+void FMassEntityQuery::CacheArchetypes(const FMassEntityManager& InEntityManager)
 {
-	const uint32 InEntitySubsystemHash = PointerHash(&InEntitySubsystem);
+	const uint32 InEntityManagerHash = PointerHash(&InEntityManager);
 
 	// Do an incremental update if the last updated archetype data version is different 
-    bool bUpdateArchetypes = InEntitySubsystem.GetArchetypeDataVersion() != LastUpdatedArchetypeDataVersion;
+    bool bUpdateArchetypes = InEntityManager.GetArchetypeDataVersion() != LastUpdatedArchetypeDataVersion;
 
 	// Force a full update if the entity system changed or if the requirements changed
-	if (EntitySubsystemHash != InEntitySubsystemHash || IncrementalChangesCount)
+	if (EntitySubsystemHash != InEntityManagerHash || IncrementalChangesCount)
 	{
 		bUpdateArchetypes = true;
-		EntitySubsystemHash = InEntitySubsystemHash;
+		EntitySubsystemHash = InEntityManagerHash;
 		ValidArchetypes.Reset();
 		LastUpdatedArchetypeDataVersion = 0;
 		ArchetypeFragmentMapping.Reset();
@@ -94,7 +94,7 @@ void FMassEntityQuery::CacheArchetypes(const UMassEntitySubsystem& InEntitySubsy
 			else
 			{
 				bUpdateArchetypes = false;
-				UE_VLOG_UELOG(&InEntitySubsystem, LogMass, Error, TEXT("FMassEntityQuery::CacheArchetypes: requirements not valid: %s"), *FMassDebugger::GetRequirementsDescription(*this));
+				UE_VLOG_UELOG(InEntityManager.GetOwner(), LogMass, Error, TEXT("FMassEntityQuery::CacheArchetypes: requirements not valid: %s"), *FMassDebugger::GetRequirementsDescription(*this));
 			}
 		}
 	}
@@ -103,8 +103,8 @@ void FMassEntityQuery::CacheArchetypes(const UMassEntitySubsystem& InEntitySubsy
 	if (bUpdateArchetypes)
 	{
 		TArray<FMassArchetypeHandle> NewValidArchetypes;
-		InEntitySubsystem.GetValidArchetypes(*this, NewValidArchetypes, LastUpdatedArchetypeDataVersion);
-		LastUpdatedArchetypeDataVersion = InEntitySubsystem.GetArchetypeDataVersion();
+		InEntityManager.GetValidArchetypes(*this, NewValidArchetypes, LastUpdatedArchetypeDataVersion);
+		LastUpdatedArchetypeDataVersion = InEntityManager.GetArchetypeDataVersion();
 		if (NewValidArchetypes.Num())
 		{
 			const int32 FirstNewArchetype = ValidArchetypes.Num();
@@ -134,15 +134,15 @@ void FMassEntityQuery::CacheArchetypes(const UMassEntitySubsystem& InEntitySubsy
 	}
 }
 
-void FMassEntityQuery::ForEachEntityChunk(const FMassArchetypeEntityCollection& Collection, UMassEntitySubsystem& EntitySubsystem, FMassExecutionContext& ExecutionContext, const FMassExecuteFunction& ExecuteFunction)
+void FMassEntityQuery::ForEachEntityChunk(const FMassArchetypeEntityCollection& Collection, FMassEntityManager& EntityManager, FMassExecutionContext& ExecutionContext, const FMassExecuteFunction& ExecuteFunction)
 {
 	// mz@todo I don't like that we're copying data here.
 	ExecutionContext.SetEntityCollection(Collection);
-	ForEachEntityChunk(EntitySubsystem, ExecutionContext, ExecuteFunction);
+	ForEachEntityChunk(EntityManager, ExecutionContext, ExecuteFunction);
 	ExecutionContext.ClearEntityCollection();
 }
 
-void FMassEntityQuery::ForEachEntityChunk(UMassEntitySubsystem& EntitySubsystem, FMassExecutionContext& ExecutionContext, const FMassExecuteFunction& ExecuteFunction)
+void FMassEntityQuery::ForEachEntityChunk(FMassEntityManager& EntityManager, FMassExecutionContext& ExecutionContext, const FMassExecuteFunction& ExecuteFunction)
 {
 #if WITH_MASSENTITY_DEBUG
 	int32 NumEntitiesToProcess = 0;
@@ -150,7 +150,7 @@ void FMassEntityQuery::ForEachEntityChunk(UMassEntitySubsystem& EntitySubsystem,
 	checkf(ExecutionContext.ExecutionType == ExpectedContextType && (ExpectedContextType == EMassExecutionContextType::Local || bRegistered)
 		, TEXT("ExecutionContextType mismatch, make sure all the queries run as part of processor execution are registered with some processor with a FMassEntityQuery::RegisterWithProcessor call"));
 
-	EntitySubsystem.GetRequirementAccessDetector().RequireAccess(*this);
+	EntityManager.GetRequirementAccessDetector().RequireAccess(*this);
 #endif
 
 	struct FScopedSubsystemRequirementsRestore
@@ -183,11 +183,11 @@ void FMassEntityQuery::ForEachEntityChunk(UMassEntitySubsystem& EntitySubsystem,
 		// verify the archetype matches requirements
 		if (DoesArchetypeMatchRequirements(ArchetypeHandle) == false)
 		{
-			UE_VLOG_UELOG(&EntitySubsystem, LogMass, Log, TEXT("Attempted to execute FMassEntityQuery with an incompatible Archetype: %s")
+			UE_VLOG_UELOG(EntityManager.GetOwner(), LogMass, Log, TEXT("Attempted to execute FMassEntityQuery with an incompatible Archetype: %s")
 				, *FMassDebugger::GetArchetypeRequirementCompatibilityDescription(*this, ArchetypeHandle));
 
 #if WITH_MASSENTITY_DEBUG
-			EntitySubsystem.GetRequirementAccessDetector().ReleaseAccess(*this);
+			EntityManager.GetRequirementAccessDetector().ReleaseAccess(*this);
 #endif // WITH_MASSENTITY_DEBUG
 			return;
 		}
@@ -203,7 +203,7 @@ void FMassEntityQuery::ForEachEntityChunk(UMassEntitySubsystem& EntitySubsystem,
 	}
 	else
 	{
-		CacheArchetypes(EntitySubsystem);
+		CacheArchetypes(EntityManager);
 		// it's important to set requirements after caching archetypes due to that call potentially sorting the requirements and the order is relevant here.
 		ExecutionContext.SetFragmentRequirements(*this);
 
@@ -224,16 +224,16 @@ void FMassEntityQuery::ForEachEntityChunk(UMassEntitySubsystem& EntitySubsystem,
 	UE_CLOG(!ExecutionContext.DebugGetExecutionDesc().IsEmpty(), LogMass, VeryVerbose,
 		TEXT("%s: %d entities sent for processing"), *ExecutionContext.DebugGetExecutionDesc(), NumEntitiesToProcess);
 
-	EntitySubsystem.GetRequirementAccessDetector().ReleaseAccess(*this);
+	EntityManager.GetRequirementAccessDetector().ReleaseAccess(*this);
 #endif
 
 	ExecutionContext.ClearExecutionData();
-	ExecutionContext.FlushDeferred(EntitySubsystem);
+	ExecutionContext.FlushDeferred(EntityManager);
 }
 
-int32 FMassEntityQuery::GetNumMatchingEntities(UMassEntitySubsystem& InEntitySubsystem)
+int32 FMassEntityQuery::GetNumMatchingEntities(FMassEntityManager& InEntityManager)
 {
-	CacheArchetypes(InEntitySubsystem);
+	CacheArchetypes(InEntityManager);
 	int32 TotalEntities = 0;
 	for (FMassArchetypeHandle& ArchetypeHandle : ValidArchetypes)
 	{
@@ -245,9 +245,9 @@ int32 FMassEntityQuery::GetNumMatchingEntities(UMassEntitySubsystem& InEntitySub
 	return TotalEntities;
 }
 
-bool FMassEntityQuery::HasMatchingEntities(UMassEntitySubsystem& InEntitySubsystem)
+bool FMassEntityQuery::HasMatchingEntities(FMassEntityManager& InEntityManager)
 {
-	CacheArchetypes(InEntitySubsystem);
+	CacheArchetypes(InEntityManager);
 
 	for (FMassArchetypeHandle& ArchetypeHandle : ValidArchetypes)
 	{
diff --git a/Engine/Plugins/Runtime/MassEntity/Source/MassEntity/Private/MassEntitySubsystem.cpp b/Engine/Plugins/Runtime/MassEntity/Source/MassEntity/Private/MassEntitySubsystem.cpp
--- a/Engine/Plugins/Runtime/MassEntity/Source/MassEntity/Private/MassEntitySubsystem.cpp
+++ b/Engine/Plugins/Runtime/MassEntity/Source/MassEntity/Private/MassEntitySubsystem.cpp
@@ -9,7 +9,7 @@
 #include "VisualLogger/VisualLogger.h"
 #include "MassExecutionContext.h"
 
-const FMassEntityHandle UMassEntitySubsystem::InvalidEntity;
+const FMassEntityHandle FMassEntityManager::InvalidEntity;
 
 namespace UE::Mass::Private
 {
@@ -36,17 +36,16 @@ namespace UE::Mass::Private
 }
 
 //////////////////////////////////////////////////////////////////////
-// UMassEntitySubsystem
+// FMassEntityManager
 
-UMassEntitySubsystem::UMassEntitySubsystem()
+FMassEntityManager::FMassEntityManager(UObject* InOwner)
 	: ObserverManager(*this)
+	, Owner(InOwner)
 {
 }
 
-void UMassEntitySubsystem::GetResourceSizeEx(FResourceSizeEx& CumulativeResourceSize)
+void FMassEntityManager::GetResourceSizeEx(FResourceSizeEx& CumulativeResourceSize)
 {
-	Super::GetResourceSizeEx(CumulativeResourceSize);
-
 	const SIZE_T MyExtraSize = Entities.GetAllocatedSize() + 
 		EntityFreeIndexList.GetAllocatedSize() +
 		(DeferredCommandBuffer != nullptr ? DeferredCommandBuffer->GetAllocatedSize() : 0) +
@@ -63,7 +62,23 @@ void UMassEntitySubsystem::GetResourceSizeEx(FResourceSizeEx& CumulativeResource
 	}
 }
 
-void UMassEntitySubsystem::Initialize(FSubsystemCollectionBase& Collection)
+void FMassEntityManager::AddReferencedObjects(FReferenceCollector& Collector)
+{
+	for (FConstSharedStruct& Struct : ConstSharedFragments)
+	{
+		Struct.AddStructReferencedObjects(Collector);
+	}
+
+	for (FSharedStruct& Struct : SharedFragments)
+	{
+		Struct.AddStructReferencedObjects(Collector);
+	}
+
+	const class UScriptStruct* ScriptStruct = FMassObserverManager::StaticStruct();
+	Collector.AddReferencedObjects(ScriptStruct, &ObserverManager);
+}
+
+void FMassEntityManager::Initialize()
 {
 	// Index 0 is reserved so we can treat that index as an invalid entity handle
 	Entities.Add();
@@ -113,20 +128,20 @@ void UMassEntitySubsystem::Initialize(FSubsystemCollectionBase& Collection)
 #endif // WITH_MASSENTITY_DEBUG
 }
 
-void UMassEntitySubsystem::PostInitialize()
+void FMassEntityManager::PostInitialize()
 {
 	// this needs to be done after all the subsystems have been initialized since some processors might want to access
 	// them during processors' initialization
 	ObserverManager.Initialize();
 }
 
-void UMassEntitySubsystem::Deinitialize()
+void FMassEntityManager::Deinitialize()
 {
 	// closing down so no point in actually flushing commands, but need to clean them up to avoid warnings on destruction
 	DeferredCommandBuffer->CleanUp();
 }
 
-FMassArchetypeHandle UMassEntitySubsystem::CreateArchetype(TConstArrayView<const UScriptStruct*> FragmentsAndTagsList, const FName ArchetypeDebugName)
+FMassArchetypeHandle FMassEntityManager::CreateArchetype(TConstArrayView<const UScriptStruct*> FragmentsAndTagsList, const FName ArchetypeDebugName)
 {
 	FMassChunkFragmentBitSet ChunkFragments;
 	FMassTagBitSet Tags;
@@ -158,7 +173,7 @@ FMassArchetypeHandle UMassEntitySubsystem::CreateArchetype(TConstArrayView<const
 	return CreateArchetype(Composition, ArchetypeDebugName);
 }
 
-FMassArchetypeHandle UMassEntitySubsystem::CreateArchetype(const TSharedPtr<FMassArchetypeData>& SourceArchetype, const FMassFragmentBitSet& AddedFragments, const FName ArchetypeDebugName)
+FMassArchetypeHandle FMassEntityManager::CreateArchetype(const TSharedPtr<FMassArchetypeData>& SourceArchetype, const FMassFragmentBitSet& AddedFragments, const FName ArchetypeDebugName)
 {
 	check(SourceArchetype.IsValid());
 	checkf(AddedFragments.IsEmpty() == false, TEXT("%s Adding an empty fragment list to an archetype is not supported."), ANSI_TO_TCHAR(__FUNCTION__));
@@ -167,7 +182,7 @@ FMassArchetypeHandle UMassEntitySubsystem::CreateArchetype(const TSharedPtr<FMas
 	return CreateArchetype(Composition, ArchetypeDebugName);
 }
 
-FMassArchetypeHandle UMassEntitySubsystem::CreateArchetype(const FMassArchetypeCompositionDescriptor& Composition, const FName ArchetypeDebugName)
+FMassArchetypeHandle FMassEntityManager::CreateArchetype(const FMassArchetypeCompositionDescriptor& Composition, const FName ArchetypeDebugName)
 {
 	const uint32 TypeHash = Composition.CalculateHash();
 
@@ -212,7 +227,7 @@ FMassArchetypeHandle UMassEntitySubsystem::CreateArchetype(const FMassArchetypeC
 	return FMassArchetypeHelper::ArchetypeHandleFromData(ArchetypeDataPtr);
 }
 
-FMassArchetypeHandle UMassEntitySubsystem::InternalCreateSimilarArchetype(const TSharedPtr<FMassArchetypeData>& SourceArchetype, const FMassTagBitSet& OverrideTags)
+FMassArchetypeHandle FMassEntityManager::InternalCreateSimilarArchetype(const TSharedPtr<FMassArchetypeData>& SourceArchetype, const FMassTagBitSet& OverrideTags)
 {
 	checkSlow(SourceArchetype.IsValid());
 	const FMassArchetypeData& SourceArchetypeRef = *SourceArchetype.Get();
@@ -220,7 +235,7 @@ FMassArchetypeHandle UMassEntitySubsystem::InternalCreateSimilarArchetype(const
 	return InternalCreateSimilarArchetype(SourceArchetypeRef, MoveTemp(NewComposition));
 }
 
-FMassArchetypeHandle UMassEntitySubsystem::InternalCreateSimilarArchetype(const TSharedPtr<FMassArchetypeData>& SourceArchetype, const FMassFragmentBitSet& OverrideFragments)
+FMassArchetypeHandle FMassEntityManager::InternalCreateSimilarArchetype(const TSharedPtr<FMassArchetypeData>& SourceArchetype, const FMassFragmentBitSet& OverrideFragments)
 {
 	checkSlow(SourceArchetype.IsValid());
 	const FMassArchetypeData& SourceArchetypeRef = *SourceArchetype.Get();
@@ -228,7 +243,7 @@ FMassArchetypeHandle UMassEntitySubsystem::InternalCreateSimilarArchetype(const
 	return InternalCreateSimilarArchetype(SourceArchetypeRef, MoveTemp(NewComposition));
 }
 
-FMassArchetypeHandle UMassEntitySubsystem::InternalCreateSimilarArchetype(const FMassArchetypeData& SourceArchetypeRef, FMassArchetypeCompositionDescriptor&& NewComposition)
+FMassArchetypeHandle FMassEntityManager::InternalCreateSimilarArchetype(const FMassArchetypeData& SourceArchetypeRef, FMassArchetypeCompositionDescriptor&& NewComposition)
 {
 	const uint32 TypeHash = NewComposition.CalculateHash();
 
@@ -266,7 +281,7 @@ FMassArchetypeHandle UMassEntitySubsystem::InternalCreateSimilarArchetype(const
 	return FMassArchetypeHelper::ArchetypeHandleFromData(ArchetypeDataPtr);
 }
 
-FMassArchetypeHandle UMassEntitySubsystem::GetArchetypeForEntity(FMassEntityHandle Entity) const
+FMassArchetypeHandle FMassEntityManager::GetArchetypeForEntity(FMassEntityHandle Entity) const
 {
 	if (IsEntityValid(Entity))
 	{
@@ -275,19 +290,19 @@ FMassArchetypeHandle UMassEntitySubsystem::GetArchetypeForEntity(FMassEntityHand
 	return FMassArchetypeHandle();
 }
 
-FMassArchetypeHandle UMassEntitySubsystem::GetArchetypeForEntityUnsafe(FMassEntityHandle Entity) const
+FMassArchetypeHandle FMassEntityManager::GetArchetypeForEntityUnsafe(FMassEntityHandle Entity) const
 {
 	check(Entities.IsValidIndex(Entity.Index));
 	return FMassArchetypeHelper::ArchetypeHandleFromData(Entities[Entity.Index].CurrentArchetype);
 }
 
-void UMassEntitySubsystem::ForEachArchetypeFragmentType(const FMassArchetypeHandle& ArchetypeHandle, TFunction< void(const UScriptStruct* /*FragmentType*/)> Function)
+void FMassEntityManager::ForEachArchetypeFragmentType(const FMassArchetypeHandle& ArchetypeHandle, TFunction< void(const UScriptStruct* /*FragmentType*/)> Function)
 {
 	const FMassArchetypeData& ArchetypeData = FMassArchetypeHelper::ArchetypeDataFromHandleChecked(ArchetypeHandle);
 	ArchetypeData.ForEachFragmentType(Function);
 }
 
-void UMassEntitySubsystem::DoEntityCompaction(const double TimeAllowed)
+void FMassEntityManager::DoEntityCompaction(const double TimeAllowed)
 {
 	const double TimeAllowedEnd = FPlatformTime::Seconds() + TimeAllowed;
 
@@ -311,7 +326,7 @@ void UMassEntitySubsystem::DoEntityCompaction(const double TimeAllowed)
 	}
 }
 
-FMassEntityHandle UMassEntitySubsystem::CreateEntity(const FMassArchetypeHandle& ArchetypeHandle, const FMassArchetypeSharedFragmentValues& SharedFragmentValues)
+FMassEntityHandle FMassEntityManager::CreateEntity(const FMassArchetypeHandle& ArchetypeHandle, const FMassArchetypeSharedFragmentValues& SharedFragmentValues)
 {
 	check(ArchetypeHandle.IsValid());
 
@@ -320,7 +335,7 @@ FMassEntityHandle UMassEntitySubsystem::CreateEntity(const FMassArchetypeHandle&
 	return Entity;
 }
 
-FMassEntityHandle UMassEntitySubsystem::CreateEntity(TConstArrayView<FInstancedStruct> FragmentInstanceList, const FMassArchetypeSharedFragmentValues& SharedFragmentValues, const FName ArchetypeDebugName)
+FMassEntityHandle FMassEntityManager::CreateEntity(TConstArrayView<FInstancedStruct> FragmentInstanceList, const FMassArchetypeSharedFragmentValues& SharedFragmentValues, const FName ArchetypeDebugName)
 {
 	check(FragmentInstanceList.Num() > 0);
 
@@ -337,7 +352,7 @@ FMassEntityHandle UMassEntitySubsystem::CreateEntity(TConstArrayView<FInstancedS
 	return Entity;
 }
 
-FMassEntityHandle UMassEntitySubsystem::ReserveEntity()
+FMassEntityHandle FMassEntityManager::ReserveEntity()
 {
 	// @todo: Need to add thread safety to the reservation of an entity
 	FMassEntityHandle Result;
@@ -348,14 +363,14 @@ FMassEntityHandle UMassEntitySubsystem::ReserveEntity()
 	return Result;
 }
 
-void UMassEntitySubsystem::ReleaseReservedEntity(FMassEntityHandle Entity)
+void FMassEntityManager::ReleaseReservedEntity(FMassEntityHandle Entity)
 {
 	checkf(!IsEntityBuilt(Entity), TEXT("Entity is already built, use DestroyEntity() instead"));
 
 	InternalReleaseEntity(Entity);
 }
 
-void UMassEntitySubsystem::BuildEntity(FMassEntityHandle Entity, const FMassArchetypeHandle& ArchetypeHandle, const FMassArchetypeSharedFragmentValues& SharedFragmentValues)
+void FMassEntityManager::BuildEntity(FMassEntityHandle Entity, const FMassArchetypeHandle& ArchetypeHandle, const FMassArchetypeSharedFragmentValues& SharedFragmentValues)
 {
 	checkf(!IsEntityBuilt(Entity), TEXT("Expecting an entity that is not already built"));
 	check(ArchetypeHandle.IsValid());
@@ -363,7 +378,7 @@ void UMassEntitySubsystem::BuildEntity(FMassEntityHandle Entity, const FMassArch
 	InternalBuildEntity(Entity, ArchetypeHandle, SharedFragmentValues);
 }
 
-void UMassEntitySubsystem::BuildEntity(FMassEntityHandle Entity, TConstArrayView<FInstancedStruct> FragmentInstanceList, const FMassArchetypeSharedFragmentValues& SharedFragmentValues)
+void FMassEntityManager::BuildEntity(FMassEntityHandle Entity, TConstArrayView<FInstancedStruct> FragmentInstanceList, const FMassArchetypeSharedFragmentValues& SharedFragmentValues)
 {
 	check(FragmentInstanceList.Num() > 0);
 	checkf(!IsEntityBuilt(Entity), TEXT("Expecting an entity that is not already built"));
@@ -388,7 +403,7 @@ void UMassEntitySubsystem::BuildEntity(FMassEntityHandle Entity, TConstArrayView
 	EntityData.CurrentArchetype->SetFragmentsData(Entity, FragmentInstanceList);
 }
 
-void UMassEntitySubsystem::BatchBuildEntities(const FMassArchetypeEntityCollectionWithPayload& EncodedEntitiesWithPayload, const FMassFragmentBitSet& FragmentsAffected, const FMassArchetypeSharedFragmentValues& SharedFragmentValues, const FName ArchetypeDebugName)
+void FMassEntityManager::BatchBuildEntities(const FMassArchetypeEntityCollectionWithPayload& EncodedEntitiesWithPayload, const FMassFragmentBitSet& FragmentsAffected, const FMassArchetypeSharedFragmentValues& SharedFragmentValues, const FName ArchetypeDebugName)
 {
 	check(SharedFragmentValues.IsSorted());
 
@@ -405,7 +420,7 @@ void UMassEntitySubsystem::BatchBuildEntities(const FMassArchetypeEntityCollecti
 	BatchBuildEntities(EncodedEntitiesWithPayload, MoveTemp(Composition), SharedFragmentValues, ArchetypeDebugName);
 }
 
-void UMassEntitySubsystem::BatchBuildEntities(const FMassArchetypeEntityCollectionWithPayload& EncodedEntitiesWithPayload, FMassArchetypeCompositionDescriptor&& Composition, const FMassArchetypeSharedFragmentValues& SharedFragmentValues, const FName ArchetypeDebugName)
+void FMassEntityManager::BatchBuildEntities(const FMassArchetypeEntityCollectionWithPayload& EncodedEntitiesWithPayload, FMassArchetypeCompositionDescriptor&& Composition, const FMassArchetypeSharedFragmentValues& SharedFragmentValues, const FName ArchetypeDebugName)
 {
 	TRACE_CPUPROFILER_EVENT_SCOPE(Mass_BatchBuildEntities);
 
@@ -423,7 +438,7 @@ void UMassEntitySubsystem::BatchBuildEntities(const FMassArchetypeEntityCollecti
 
 	// since the handles encoded via FMassArchetypeEntityCollectionWithPayload miss the SerialNumber we need to update it
 	// before passing over the the new archetype. Thankfully we need to iterate over all the entity handles anyway
-	// to update the system's information on these entities (stored in UMassEntitySubsystem::Entities)
+	// to update the manager's information on these entities (stored in FMassEntityManager::Entities)
 	for (FMassEntityHandle& Entity : EntityHandles)
 	{
 		check(Entities.IsValidIndex(Entity.Index));
@@ -452,7 +467,7 @@ void UMassEntitySubsystem::BatchBuildEntities(const FMassArchetypeEntityCollecti
 	}
 }
 
-TSharedRef<UMassEntitySubsystem::FEntityCreationContext> UMassEntitySubsystem::BatchCreateEntities(const FMassArchetypeHandle& ArchetypeHandle, const FMassArchetypeSharedFragmentValues& SharedFragmentValues, const int32 Count, TArray<FMassEntityHandle>& OutEntities)
+TSharedRef<FMassEntityManager::FEntityCreationContext> FMassEntityManager::BatchCreateEntities(const FMassArchetypeHandle& ArchetypeHandle, const FMassArchetypeSharedFragmentValues& SharedFragmentValues, const int32 Count, TArray<FMassEntityHandle>& OutEntities)
 {
 	TRACE_CPUPROFILER_EVENT_SCOPE(Mass_BatchCreateEntities);
 
@@ -479,7 +494,7 @@ TSharedRef<UMassEntitySubsystem::FEntityCreationContext> UMassEntitySubsystem::B
 	FEntityCreationContext* CreationContext = new FEntityCreationContext(Count);
 	// @todo this could probably be optimized since one would assume we're adding elements to OutEntities in order.
 	// Then again, if that's the case, the sorting will be almost instant
-	new (&CreationContext->EntityCollection)FMassArchetypeEntityCollection(ArchetypeHandle, MakeArrayView(&OutEntities[OutEntities.Num() - Count], Count), FMassArchetypeEntityCollection::NoDuplicates);
+	new (&CreationContext->EntityCollection) FMassArchetypeEntityCollection(ArchetypeHandle, MakeArrayView(&OutEntities[OutEntities.Num() - Count], Count), FMassArchetypeEntityCollection::NoDuplicates);
 	if (ObserverManager.HasObserversForBitSet(ArchetypeData.GetCompositionDescriptor().Fragments, EMassObservedOperation::Add))
 	{
 		CreationContext->OnSpawningFinished = [this](FEntityCreationContext& Context){
@@ -490,7 +505,7 @@ TSharedRef<UMassEntitySubsystem::FEntityCreationContext> UMassEntitySubsystem::B
 	return MakeShareable(CreationContext);
 }
 
-void UMassEntitySubsystem::DestroyEntity(FMassEntityHandle Entity)
+void FMassEntityManager::DestroyEntity(FMassEntityHandle Entity)
 {
 	checkf(IsProcessing() == false, TEXT("Synchronous API function %s called during mass processing. Use asynchronous API instead."), ANSI_TO_TCHAR(__FUNCTION__));
 	
@@ -509,7 +524,7 @@ void UMassEntitySubsystem::DestroyEntity(FMassEntityHandle Entity)
 	}
 }
 
-void UMassEntitySubsystem::BatchDestroyEntities(TConstArrayView<FMassEntityHandle> InEntities)
+void FMassEntityManager::BatchDestroyEntities(TConstArrayView<FMassEntityHandle> InEntities)
 {
 	TRACE_CPUPROFILER_EVENT_SCOPE(Mass_BatchDestroyEntities);
 
@@ -540,7 +555,7 @@ void UMassEntitySubsystem::BatchDestroyEntities(TConstArrayView<FMassEntityHandl
 	}
 }
 
-void UMassEntitySubsystem::BatchDestroyEntityChunks(const FMassArchetypeEntityCollection& EntityCollection)
+void FMassEntityManager::BatchDestroyEntityChunks(const FMassArchetypeEntityCollection& EntityCollection)
 {
 	TRACE_CPUPROFILER_EVENT_SCOPE(Mass_BatchDestroyEntityChunks);
 
@@ -548,7 +563,7 @@ void UMassEntitySubsystem::BatchDestroyEntityChunks(const FMassArchetypeEntityCo
 
 	TArray<FMassEntityHandle> EntitiesRemoved;
 	// note that it's important to place the context instance in the same scope as the loop below that updates 
-	// UMassEntitySubsystem.EntityData, otherwise, if there are commands flushed as part of FMassProcessingContext's 
+	// FMassEntityManager.EntityData, otherwise, if there are commands flushed as part of FMassProcessingContext's 
 	// destruction the commands will work on outdated information (which might result in crashes).
 	FMassProcessingContext ProcessingContext(*this, /*TimeDelta=*/0.0f);
 
@@ -582,7 +597,7 @@ void UMassEntitySubsystem::BatchDestroyEntityChunks(const FMassArchetypeEntityCo
 	}
 }
 
-void UMassEntitySubsystem::AddFragmentToEntity(FMassEntityHandle Entity, const UScriptStruct* FragmentType)
+void FMassEntityManager::AddFragmentToEntity(FMassEntityHandle Entity, const UScriptStruct* FragmentType)
 {
 	checkf(FragmentType, TEXT("Null fragment type passed in to %s"), ANSI_TO_TCHAR(__FUNCTION__));
 	checkf(IsProcessing() == false, TEXT("Synchronous API function %s called during mass processing. Use asynchronous API instead."), ANSI_TO_TCHAR(__FUNCTION__));
@@ -592,14 +607,14 @@ void UMassEntitySubsystem::AddFragmentToEntity(FMassEntityHandle Entity, const U
 	InternalAddFragmentListToEntityChecked(Entity, FMassFragmentBitSet(*FragmentType));
 }
 
-void UMassEntitySubsystem::AddFragmentListToEntity(FMassEntityHandle Entity, TConstArrayView<const UScriptStruct*> FragmentList)
+void FMassEntityManager::AddFragmentListToEntity(FMassEntityHandle Entity, TConstArrayView<const UScriptStruct*> FragmentList)
 {
 	CheckIfEntityIsActive(Entity);
 
 	InternalAddFragmentListToEntityChecked(Entity, FMassFragmentBitSet(FragmentList));
 }
 
-void UMassEntitySubsystem::AddCompositionToEntity_GetDelta(FMassEntityHandle Entity, FMassArchetypeCompositionDescriptor& InDescriptor)
+void FMassEntityManager::AddCompositionToEntity_GetDelta(FMassEntityHandle Entity, FMassArchetypeCompositionDescriptor& InDescriptor)
 {
 	CheckIfEntityIsActive(Entity);
 
@@ -633,7 +648,7 @@ void UMassEntitySubsystem::AddCompositionToEntity_GetDelta(FMassEntityHandle Ent
 	}
 }
 
-void UMassEntitySubsystem::RemoveCompositionFromEntity(FMassEntityHandle Entity, const FMassArchetypeCompositionDescriptor& InDescriptor)
+void FMassEntityManager::RemoveCompositionFromEntity(FMassEntityHandle Entity, const FMassArchetypeCompositionDescriptor& InDescriptor)
 {
 	CheckIfEntityIsActive(Entity);
 
@@ -669,27 +684,27 @@ void UMassEntitySubsystem::RemoveCompositionFromEntity(FMassEntityHandle Entity,
 	}
 }
 
-const FMassArchetypeCompositionDescriptor& UMassEntitySubsystem::GetArchetypeComposition(const FMassArchetypeHandle& ArchetypeHandle) const
+const FMassArchetypeCompositionDescriptor& FMassEntityManager::GetArchetypeComposition(const FMassArchetypeHandle& ArchetypeHandle) const
 {
 	const FMassArchetypeData& ArchetypeData = FMassArchetypeHelper::ArchetypeDataFromHandleChecked(ArchetypeHandle);
 	return ArchetypeData.GetCompositionDescriptor();
 }
 
-void UMassEntitySubsystem::InternalBuildEntity(FMassEntityHandle Entity, const FMassArchetypeHandle& ArchetypeHandle, const FMassArchetypeSharedFragmentValues& SharedFragmentValues)
+void FMassEntityManager::InternalBuildEntity(FMassEntityHandle Entity, const FMassArchetypeHandle& ArchetypeHandle, const FMassArchetypeSharedFragmentValues& SharedFragmentValues)
 {
 	FEntityData& EntityData = Entities[Entity.Index];
 	EntityData.CurrentArchetype = ArchetypeHandle.DataPtr;
 	EntityData.CurrentArchetype->AddEntity(Entity, SharedFragmentValues);
 }
 
-void UMassEntitySubsystem::InternalReleaseEntity(FMassEntityHandle Entity)
+void FMassEntityManager::InternalReleaseEntity(FMassEntityHandle Entity)
 {
 	FEntityData& EntityData = Entities[Entity.Index];
 	EntityData.Reset();
 	EntityFreeIndexList.Add(Entity.Index);
 }
 
-void UMassEntitySubsystem::InternalAddFragmentListToEntityChecked(FMassEntityHandle Entity, const FMassFragmentBitSet& InFragments)
+void FMassEntityManager::InternalAddFragmentListToEntityChecked(FMassEntityHandle Entity, const FMassFragmentBitSet& InFragments)
 {
 	const FEntityData& EntityData = Entities[Entity.Index];
 	FMassArchetypeData* OldArchetype = EntityData.CurrentArchetype.Get();
@@ -706,7 +721,7 @@ void UMassEntitySubsystem::InternalAddFragmentListToEntityChecked(FMassEntityHan
 	}
 }
 
-void UMassEntitySubsystem::InternalAddFragmentListToEntity(FMassEntityHandle Entity, const FMassFragmentBitSet& InFragments)
+void FMassEntityManager::InternalAddFragmentListToEntity(FMassEntityHandle Entity, const FMassFragmentBitSet& InFragments)
 {
 	checkf(InFragments.IsEmpty() == false, TEXT("%s is intended for internal calls with non empty NewFragments parameter"), ANSI_TO_TCHAR(__FUNCTION__));
 	check(Entities.IsValidIndex(Entity.Index));
@@ -727,7 +742,7 @@ void UMassEntitySubsystem::InternalAddFragmentListToEntity(FMassEntityHandle Ent
 	}
 }
 
-void UMassEntitySubsystem::AddFragmentInstanceListToEntity(FMassEntityHandle Entity, TConstArrayView<FInstancedStruct> FragmentInstanceList)
+void FMassEntityManager::AddFragmentInstanceListToEntity(FMassEntityHandle Entity, TConstArrayView<FInstancedStruct> FragmentInstanceList)
 {
 	checkf(IsProcessing() == false, TEXT("Synchronous API function %s called during mass processing. Use asynchronous API instead."), ANSI_TO_TCHAR(__FUNCTION__));
 
@@ -740,12 +755,12 @@ void UMassEntitySubsystem::AddFragmentInstanceListToEntity(FMassEntityHandle Ent
 	EntityData.CurrentArchetype->SetFragmentsData(Entity, FragmentInstanceList);
 }
 
-void UMassEntitySubsystem::RemoveFragmentFromEntity(FMassEntityHandle Entity, const UScriptStruct* FragmentType)
+void FMassEntityManager::RemoveFragmentFromEntity(FMassEntityHandle Entity, const UScriptStruct* FragmentType)
 {
 	RemoveFragmentListFromEntity(Entity, MakeArrayView(&FragmentType, 1));
 }
 
-void UMassEntitySubsystem::RemoveFragmentListFromEntity(FMassEntityHandle Entity, TConstArrayView<const UScriptStruct*> FragmentList)
+void FMassEntityManager::RemoveFragmentListFromEntity(FMassEntityHandle Entity, TConstArrayView<const UScriptStruct*> FragmentList)
 {
 	checkf(IsProcessing() == false, TEXT("Synchronous API function %s called during mass processing. Use asynchronous API instead."), ANSI_TO_TCHAR(__FUNCTION__));
 
@@ -771,7 +786,7 @@ void UMassEntitySubsystem::RemoveFragmentListFromEntity(FMassEntityHandle Entity
 	}
 }
 
-void UMassEntitySubsystem::SwapTagsForEntity(FMassEntityHandle Entity, const UScriptStruct* OldTagType, const UScriptStruct* NewTagType)
+void FMassEntityManager::SwapTagsForEntity(FMassEntityHandle Entity, const UScriptStruct* OldTagType, const UScriptStruct* NewTagType)
 {
 	checkf(IsProcessing() == false, TEXT("Synchronous API function %s called during mass processing. Use asynchronous API instead."), ANSI_TO_TCHAR(__FUNCTION__));
 
@@ -799,7 +814,7 @@ void UMassEntitySubsystem::SwapTagsForEntity(FMassEntityHandle Entity, const USc
 	}
 }
 
-void UMassEntitySubsystem::AddTagToEntity(FMassEntityHandle Entity, const UScriptStruct* TagType)
+void FMassEntityManager::AddTagToEntity(FMassEntityHandle Entity, const UScriptStruct* TagType)
 {
 	checkf((TagType != nullptr) && TagType->IsChildOf(FMassTag::StaticStruct()), TEXT("%s works only with tags while '%s' is not one."), ANSI_TO_TCHAR(__FUNCTION__), *GetPathNameSafe(TagType));
 
@@ -823,7 +838,7 @@ void UMassEntitySubsystem::AddTagToEntity(FMassEntityHandle Entity, const UScrip
 	}
 }
 	
-void UMassEntitySubsystem::RemoveTagFromEntity(FMassEntityHandle Entity, const UScriptStruct* TagType)
+void FMassEntityManager::RemoveTagFromEntity(FMassEntityHandle Entity, const UScriptStruct* TagType)
 {
 	checkf((TagType != nullptr) && TagType->IsChildOf(FMassTag::StaticStruct()), TEXT("%s works only with tags while '%s' is not one."), ANSI_TO_TCHAR(__FUNCTION__), *GetPathNameSafe(TagType));
 
@@ -847,7 +862,7 @@ void UMassEntitySubsystem::RemoveTagFromEntity(FMassEntityHandle Entity, const U
 	}
 }
 
-void UMassEntitySubsystem::BatchChangeTagsForEntities(TConstArrayView<FMassArchetypeEntityCollection> EntityCollections, const FMassTagBitSet& TagsToAdd, const FMassTagBitSet& TagsToRemove)
+void FMassEntityManager::BatchChangeTagsForEntities(TConstArrayView<FMassArchetypeEntityCollection> EntityCollections, const FMassTagBitSet& TagsToAdd, const FMassTagBitSet& TagsToRemove)
 {
 	TRACE_CPUPROFILER_EVENT_SCOPE(Mass_BatchChangeTagsForEntities);
 
@@ -897,7 +912,7 @@ void UMassEntitySubsystem::BatchChangeTagsForEntities(TConstArrayView<FMassArche
 	}
 }
 
-void UMassEntitySubsystem::BatchChangeFragmentCompositionForEntities(TConstArrayView<FMassArchetypeEntityCollection> EntityCollections, const FMassFragmentBitSet& FragmentsToAdd, const FMassFragmentBitSet& FragmentsToRemove)
+void FMassEntityManager::BatchChangeFragmentCompositionForEntities(TConstArrayView<FMassArchetypeEntityCollection> EntityCollections, const FMassFragmentBitSet& FragmentsToAdd, const FMassFragmentBitSet& FragmentsToRemove)
 {
 	TRACE_CPUPROFILER_EVENT_SCOPE(Mass_BatchChangeFragmentCompositionForEntities);
 
@@ -950,7 +965,7 @@ void UMassEntitySubsystem::BatchChangeFragmentCompositionForEntities(TConstArray
 	}
 }
 
-void UMassEntitySubsystem::BatchAddFragmentInstancesForEntities(TConstArrayView<FMassArchetypeEntityCollectionWithPayload> EntityCollections, const FMassFragmentBitSet& FragmentsAffected)
+void FMassEntityManager::BatchAddFragmentInstancesForEntities(TConstArrayView<FMassArchetypeEntityCollectionWithPayload> EntityCollections, const FMassFragmentBitSet& FragmentsAffected)
 {
 	TRACE_CPUPROFILER_EVENT_SCOPE(Mass_BatchAddFragmentInstancesForEntities);
 
@@ -1027,7 +1042,7 @@ void UMassEntitySubsystem::BatchAddFragmentInstancesForEntities(TConstArrayView<
 	}
 }
 
-void UMassEntitySubsystem::MoveEntityToAnotherArchetype(FMassEntityHandle Entity, FMassArchetypeHandle NewArchetypeHandle)
+void FMassEntityManager::MoveEntityToAnotherArchetype(FMassEntityHandle Entity, FMassArchetypeHandle NewArchetypeHandle)
 {
 	CheckIfEntityIsActive(Entity);
 
@@ -1039,7 +1054,7 @@ void UMassEntitySubsystem::MoveEntityToAnotherArchetype(FMassEntityHandle Entity
 	EntityData.CurrentArchetype = NewArchetypeHandle.DataPtr;
 }
 
-void UMassEntitySubsystem::SetEntityFragmentsValues(FMassEntityHandle Entity, TArrayView<const FInstancedStruct> FragmentInstanceList)
+void FMassEntityManager::SetEntityFragmentsValues(FMassEntityHandle Entity, TArrayView<const FInstancedStruct> FragmentInstanceList)
 {
 	CheckIfEntityIsActive(Entity);
 
@@ -1047,7 +1062,7 @@ void UMassEntitySubsystem::SetEntityFragmentsValues(FMassEntityHandle Entity, TA
 	EntityData.CurrentArchetype->SetFragmentsData(Entity, FragmentInstanceList);
 }
 
-void UMassEntitySubsystem::BatchSetEntityFragmentsValues(const FMassArchetypeEntityCollection& SparseEntities, TArrayView<const FInstancedStruct> FragmentInstanceList)
+void FMassEntityManager::BatchSetEntityFragmentsValues(const FMassArchetypeEntityCollection& SparseEntities, TArrayView<const FInstancedStruct> FragmentInstanceList)
 {
 	TRACE_CPUPROFILER_EVENT_SCOPE(Mass_BatchSetEntityFragmentsValues);
 
@@ -1060,7 +1075,7 @@ void UMassEntitySubsystem::BatchSetEntityFragmentsValues(const FMassArchetypeEnt
 	}
 }
 
-void* UMassEntitySubsystem::InternalGetFragmentDataChecked(FMassEntityHandle Entity, const UScriptStruct* FragmentType) const
+void* FMassEntityManager::InternalGetFragmentDataChecked(FMassEntityHandle Entity, const UScriptStruct* FragmentType) const
 {
 	CheckIfEntityIsActive(Entity);
 
@@ -1069,7 +1084,7 @@ void* UMassEntitySubsystem::InternalGetFragmentDataChecked(FMassEntityHandle Ent
 	return EntityData.CurrentArchetype->GetFragmentDataForEntityChecked(FragmentType, Entity.Index);
 }
 
-void* UMassEntitySubsystem::InternalGetFragmentDataPtr(FMassEntityHandle Entity, const UScriptStruct* FragmentType) const
+void* FMassEntityManager::InternalGetFragmentDataPtr(FMassEntityHandle Entity, const UScriptStruct* FragmentType) const
 {
 	CheckIfEntityIsActive(Entity);
 	checkf((FragmentType != nullptr) && FragmentType->IsChildOf(FMassFragment::StaticStruct()), TEXT("InternalGetFragmentData called with an invalid fragment type '%s'"), *GetPathNameSafe(FragmentType));
@@ -1077,29 +1092,29 @@ void* UMassEntitySubsystem::InternalGetFragmentDataPtr(FMassEntityHandle Entity,
 	return EntityData.CurrentArchetype->GetFragmentDataForEntity(FragmentType, Entity.Index);
 }
 
-bool UMassEntitySubsystem::IsEntityValid(FMassEntityHandle Entity) const
+bool FMassEntityManager::IsEntityValid(FMassEntityHandle Entity) const
 {
 	return (Entity.Index > 0) && Entities.IsValidIndex(Entity.Index) && (Entities[Entity.Index].SerialNumber == Entity.SerialNumber);
 }
 
-bool UMassEntitySubsystem::IsEntityBuilt(FMassEntityHandle Entity) const
+bool FMassEntityManager::IsEntityBuilt(FMassEntityHandle Entity) const
 {
 	CheckIfEntityIsValid(Entity);
 	return Entities[Entity.Index].CurrentArchetype.IsValid();
 }
 
-void UMassEntitySubsystem::CheckIfEntityIsValid(FMassEntityHandle Entity) const
+void FMassEntityManager::CheckIfEntityIsValid(FMassEntityHandle Entity) const
 {
 	checkf(IsEntityValid(Entity), TEXT("Invalid entity (ID: %d, SN:%d, %s)"), Entity.Index, Entity.SerialNumber,
 		   (Entity.Index == 0) ? TEXT("was never initialized") : TEXT("already destroyed"));
 }
 
-void UMassEntitySubsystem::CheckIfEntityIsActive(FMassEntityHandle Entity) const
+void FMassEntityManager::CheckIfEntityIsActive(FMassEntityHandle Entity) const
 {
 	checkf(IsEntityBuilt(Entity), TEXT("Entity not yet created(ID: %d, SN:%d)"));
 }
 
-void UMassEntitySubsystem::GetValidArchetypes(const FMassEntityQuery& Query, TArray<FMassArchetypeHandle>& OutValidArchetypes, const uint32 FromArchetypeDataVersion) const
+void FMassEntityManager::GetValidArchetypes(const FMassEntityQuery& Query, TArray<FMassArchetypeHandle>& OutValidArchetypes, const uint32 FromArchetypeDataVersion) const
 {
 	//@TODO: Not optimized yet, but we call this rarely now, so not a big deal.
 
@@ -1261,14 +1276,14 @@ void UMassEntitySubsystem::GetValidArchetypes(const FMassEntityQuery& Query, TAr
 	}
 }
 
-FMassExecutionContext UMassEntitySubsystem::CreateExecutionContext(const float DeltaSeconds) const
+FMassExecutionContext FMassEntityManager::CreateExecutionContext(const float DeltaSeconds) const
 {
 	FMassExecutionContext ExecutionContext(DeltaSeconds);
 	ExecutionContext.SetDeferredCommandBuffer(DeferredCommandBuffer);
 	return MoveTemp(ExecutionContext);
 }
 
-void UMassEntitySubsystem::FlushCommands(const TSharedPtr<FMassCommandBuffer>& InCommandBuffer)
+void FMassEntityManager::FlushCommands(const TSharedPtr<FMassCommandBuffer>& InCommandBuffer)
 {
 	constexpr int MaxIterations = 3;
 
@@ -1293,16 +1308,16 @@ void UMassEntitySubsystem::FlushCommands(const TSharedPtr<FMassCommandBuffer>& I
 			CurrentCommandBuffer = FlushedCommandBufferQueue.Dequeue();
 		}
 		ensure(IterationsCounter >= MaxIterations || CurrentCommandBuffer.IsSet() == false);
-		UE_CVLOG_UELOG(IterationsCounter >= MaxIterations, this, LogMass, Error, TEXT("Reached loop count limit while flushing commands"));
+		UE_CVLOG_UELOG(IterationsCounter >= MaxIterations, GetOwner(), LogMass, Error, TEXT("Reached loop count limit while flushing commands"));
 
 		bCommandBufferFlushingInProgress = false;
 	}
 }
 
 #if WITH_MASSENTITY_DEBUG
-void UMassEntitySubsystem::DebugPrintArchetypes(FOutputDevice& Ar, const bool bIncludeEmpty) const
+void FMassEntityManager::DebugPrintArchetypes(FOutputDevice& Ar, const bool bIncludeEmpty) const
 {
-	Ar.Logf(ELogVerbosity::Log, TEXT("Listing archetypes contained in %s"), *GetPathNameSafe(this));
+	Ar.Logf(ELogVerbosity::Log, TEXT("Listing archetypes contained in EntityManager owned by %s"), *GetPathNameSafe(GetOwner()));
 
 	int32 NumBuckets = 0;
 	int32 NumArchetypes = 0;
@@ -1327,7 +1342,7 @@ void UMassEntitySubsystem::DebugPrintArchetypes(FOutputDevice& Ar, const bool bI
 		NumArchetypes, NumBuckets, LongestArchetypeBucket);
 }
 
-void UMassEntitySubsystem::DebugGetArchetypesStringDetails(FOutputDevice& Ar, const bool bIncludeEmpty)
+void FMassEntityManager::DebugGetArchetypesStringDetails(FOutputDevice& Ar, const bool bIncludeEmpty) const
 {
 #if WITH_MASSENTITY_DEBUG
 	Ar.SetAutoEmitLineTerminator(true);
@@ -1346,7 +1361,7 @@ void UMassEntitySubsystem::DebugGetArchetypesStringDetails(FOutputDevice& Ar, co
 #endif // WITH_MASSENTITY_DEBUG
 }
 
-void UMassEntitySubsystem::DebugGetArchetypeFragmentTypes(const FMassArchetypeHandle& Archetype, TArray<const UScriptStruct*>& InOutFragmentList) const
+void FMassEntityManager::DebugGetArchetypeFragmentTypes(const FMassArchetypeHandle& Archetype, TArray<const UScriptStruct*>& InOutFragmentList) const
 {
 	if (Archetype.IsValid())
 	{
@@ -1355,17 +1370,17 @@ void UMassEntitySubsystem::DebugGetArchetypeFragmentTypes(const FMassArchetypeHa
 	}
 }
 
-int32 UMassEntitySubsystem::DebugGetArchetypeEntitiesCount(const FMassArchetypeHandle& Archetype) const
+int32 FMassEntityManager::DebugGetArchetypeEntitiesCount(const FMassArchetypeHandle& Archetype) const
 {
 	return Archetype.IsValid() ? FMassArchetypeHelper::ArchetypeDataFromHandleChecked(Archetype).GetNumEntities() : 0;
 }
 
-int32 UMassEntitySubsystem::DebugGetArchetypeEntitiesCountPerChunk(const FMassArchetypeHandle& Archetype) const
+int32 FMassEntityManager::DebugGetArchetypeEntitiesCountPerChunk(const FMassArchetypeHandle& Archetype) const
 {
 	return Archetype.IsValid() ? FMassArchetypeHelper::ArchetypeDataFromHandleChecked(Archetype).GetNumEntitiesPerChunk() : 0;
 }
 
-void UMassEntitySubsystem::DebugRemoveAllEntities()
+void FMassEntityManager::DebugRemoveAllEntities()
 {
 	for (int EntityIndex = NumReservedEntities; EntityIndex < Entities.Num(); ++EntityIndex)
 	{
@@ -1386,7 +1401,7 @@ void UMassEntitySubsystem::DebugRemoveAllEntities()
 	}
 }
 
-void UMassEntitySubsystem::DebugGetArchetypeStrings(const FMassArchetypeHandle& Archetype, TArray<FName>& OutFragmentNames, TArray<FName>& OutTagNames)
+void FMassEntityManager::DebugGetArchetypeStrings(const FMassArchetypeHandle& Archetype, TArray<FName>& OutFragmentNames, TArray<FName>& OutTagNames)
 {
 	if (Archetype.IsValid() == false)
 	{
@@ -1414,9 +1429,9 @@ FAutoConsoleCommandWithWorldArgsAndOutputDevice GPrintArchetypesCmd(
 	FConsoleCommandWithWorldArgsAndOutputDeviceDelegate::CreateStatic(
 		[](const TArray<FString>& Params, UWorld* World, FOutputDevice& Ar)
 {
-	if (const UMassEntitySubsystem* EntitySystem = World ? World->GetSubsystem<UMassEntitySubsystem>() : nullptr)
+	if (const UMassEntitySubsystem* EntitySubsystem = World ? World->GetSubsystem<UMassEntitySubsystem>() : nullptr)
 	{
-		EntitySystem->DebugPrintArchetypes(Ar);
+		EntitySubsystem->GetEntityManager().DebugPrintArchetypes(Ar);
 	}
 	else
 	{
@@ -1424,3 +1439,46 @@ FAutoConsoleCommandWithWorldArgsAndOutputDevice GPrintArchetypesCmd(
 	}
 }));
 #endif // WITH_MASSENTITY_DEBUG
+
+
+//////////////////////////////////////////////////////////////////////
+// UMassEntitySubsystem
+
+UMassEntitySubsystem::UMassEntitySubsystem()
+	: EntityManager(MakeShareable(new FMassEntityManager(this)))
+{
+	
+}
+
+void UMassEntitySubsystem::GetResourceSizeEx(FResourceSizeEx& CumulativeResourceSize)
+{
+	Super::GetResourceSizeEx(CumulativeResourceSize);
+	EntityManager->GetResourceSizeEx(CumulativeResourceSize);
+}
+
+void UMassEntitySubsystem::Initialize(FSubsystemCollectionBase& Collection)
+{
+	EntityManager->Initialize();
+}
+
+void UMassEntitySubsystem::PostInitialize()
+{
+	// this needs to be done after all the subsystems have been initialized since some processors might want to access
+	// them during processors' initialization
+	EntityManager->PostInitialize();
+}
+
+void UMassEntitySubsystem::Deinitialize()
+{
+	EntityManager->Deinitialize();
+}
+
+void UMassEntitySubsystem::AddReferencedObjects(UObject* InThis, FReferenceCollector& Collector)
+{
+	Super::AddReferencedObjects(InThis, Collector);
+
+	if (UMassEntitySubsystem* Instance = Cast<UMassEntitySubsystem>(InThis))
+	{
+		Instance->EntityManager->AddReferencedObjects(Collector);
+	}
+}
diff --git a/Engine/Plugins/Runtime/MassEntity/Source/MassEntity/Private/MassEntityUtils.cpp b/Engine/Plugins/Runtime/MassEntity/Source/MassEntity/Private/MassEntityUtils.cpp
--- a/Engine/Plugins/Runtime/MassEntity/Source/MassEntity/Private/MassEntityUtils.cpp
+++ b/Engine/Plugins/Runtime/MassEntity/Source/MassEntity/Private/MassEntityUtils.cpp
@@ -34,7 +34,7 @@ EProcessorExecutionFlags GetProcessorExecutionFlagsForWold(const UWorld& World)
 	return ExecutionFlags;
 }
 
-void CreateEntityCollections(const UMassEntitySubsystem& EntitySystem, const TConstArrayView<FMassEntityHandle> Entities
+void CreateEntityCollections(const FMassEntityManager& EntityManager, const TConstArrayView<FMassEntityHandle> Entities
 	, const FMassArchetypeEntityCollection::EDuplicatesHandling DuplicatesHandling, TArray<FMassArchetypeEntityCollection>& OutEntityCollections)
 {
 	TRACE_CPUPROFILER_EVENT_SCOPE_STR("Mass_CreateSparseChunks");
@@ -43,9 +43,9 @@ void CreateEntityCollections(const UMassEntitySubsystem& EntitySystem, const TCo
 
 	for (const FMassEntityHandle& Entity : Entities)
 	{
-		if (EntitySystem.IsEntityValid(Entity))
+		if (EntityManager.IsEntityValid(Entity))
 		{
-			FMassArchetypeHandle Archetype = EntitySystem.GetArchetypeForEntityUnsafe(Entity);
+			FMassArchetypeHandle Archetype = EntityManager.GetArchetypeForEntityUnsafe(Entity);
 			TArray<FMassEntityHandle>& PerArchetypeEntities = ArchetypeToEntities.FindOrAdd(Archetype);
 			PerArchetypeEntities.Add(Entity);
 		}
@@ -57,4 +57,18 @@ void CreateEntityCollections(const UMassEntitySubsystem& EntitySystem, const TCo
 	}
 }
 
+FMassEntityManager* GetEntityManager(const UWorld* World)
+{
+	UMassEntitySubsystem* EntityManager = UWorld::GetSubsystem<UMassEntitySubsystem>(World);
+	check(EntityManager);
+	return &EntityManager->GetMutableEntityManager();
+}
+
+FMassEntityManager& GetEntityManagerChecked(const UWorld& World)
+{
+	UMassEntitySubsystem* EntityManager = UWorld::GetSubsystem<UMassEntitySubsystem>(&World);
+	check(EntityManager);
+	return EntityManager->GetMutableEntityManager();
+}
+
 } // namespace UE::Mass::Utils
\ No newline at end of file
diff --git a/Engine/Plugins/Runtime/MassEntity/Source/MassEntity/Private/MassEntityView.cpp b/Engine/Plugins/Runtime/MassEntity/Source/MassEntity/Private/MassEntityView.cpp
--- a/Engine/Plugins/Runtime/MassEntity/Source/MassEntity/Private/MassEntityView.cpp
+++ b/Engine/Plugins/Runtime/MassEntity/Source/MassEntity/Private/MassEntityView.cpp
@@ -14,10 +14,10 @@ FMassEntityView::FMassEntityView(const FMassArchetypeHandle& ArchetypeHandle, FM
 	EntityHandle = Archetype->MakeEntityHandle(Entity);
 }
 
-FMassEntityView::FMassEntityView(const UMassEntitySubsystem& EntitySubsystem, FMassEntityHandle InEntity)
+FMassEntityView::FMassEntityView(const FMassEntityManager& EntityManager, FMassEntityHandle InEntity)
 {
 	Entity = InEntity;
-	const FMassArchetypeHandle ArchetypeHandle = EntitySubsystem.GetArchetypeForEntity(Entity);
+	const FMassArchetypeHandle ArchetypeHandle = EntityManager.GetArchetypeForEntity(Entity);
 	Archetype = &FMassArchetypeHelper::ArchetypeDataFromHandleChecked(ArchetypeHandle);
 	EntityHandle = Archetype->MakeEntityHandle(Entity);
 }
diff --git a/Engine/Plugins/Runtime/MassEntity/Source/MassEntity/Private/MassExecutionContext.cpp b/Engine/Plugins/Runtime/MassEntity/Source/MassEntity/Private/MassExecutionContext.cpp
--- a/Engine/Plugins/Runtime/MassEntity/Source/MassEntity/Private/MassExecutionContext.cpp
+++ b/Engine/Plugins/Runtime/MassEntity/Source/MassEntity/Private/MassExecutionContext.cpp
@@ -4,11 +4,11 @@
 #include "MassArchetypeData.h"
 
 
-void FMassExecutionContext::FlushDeferred(UMassEntitySubsystem& EntitySystem) const
+void FMassExecutionContext::FlushDeferred(FMassEntityManager& EntityManager) const
 {
 	if (bFlushDeferredCommands && DeferredCommandBuffer)
 	{
-		EntitySystem.FlushCommands(DeferredCommandBuffer);
+		EntityManager.FlushCommands(DeferredCommandBuffer);
 	}
 }
 
diff --git a/Engine/Plugins/Runtime/MassEntity/Source/MassEntity/Private/MassExecutor.cpp b/Engine/Plugins/Runtime/MassEntity/Source/MassEntity/Private/MassExecutor.cpp
--- a/Engine/Plugins/Runtime/MassEntity/Source/MassEntity/Private/MassExecutor.cpp
+++ b/Engine/Plugins/Runtime/MassEntity/Source/MassEntity/Private/MassExecutor.cpp
@@ -10,7 +10,7 @@ namespace UE::Mass::Executor
 {
 void Run(FMassRuntimePipeline& RuntimePipeline, FMassProcessingContext& ProcessingContext)
 {
-	if (!ensure(ProcessingContext.EntitySubsystem) || 
+	if (!ensure(ProcessingContext.EntityManager) || 
 		!ensure(ProcessingContext.DeltaSeconds >= 0.f) ||
 		!ensure(RuntimePipeline.Processors.Find(nullptr) == INDEX_NONE))
 	{
@@ -23,7 +23,7 @@ void Run(FMassRuntimePipeline& RuntimePipeline, FMassProcessingContext& Processi
 
 void RunSparse(FMassRuntimePipeline& RuntimePipeline, FMassProcessingContext& ProcessingContext, FMassArchetypeHandle Archetype, TConstArrayView<FMassEntityHandle> Entities)
 {
-	if (!ensure(ProcessingContext.EntitySubsystem) ||
+	if (!ensure(ProcessingContext.EntityManager) ||
 		!ensure(RuntimePipeline.Processors.Find(nullptr) == INDEX_NONE) ||
 		RuntimePipeline.Processors.Num() == 0 ||
 		!ensureMsgf(Archetype.IsValid(), TEXT("The Archetype passed in to UE::Mass::Executor::RunSparse is invalid")))
@@ -39,7 +39,7 @@ void RunSparse(FMassRuntimePipeline& RuntimePipeline, FMassProcessingContext& Pr
 
 void RunSparse(FMassRuntimePipeline& RuntimePipeline, FMassProcessingContext& ProcessingContext, const FMassArchetypeEntityCollection& EntityCollection)
 {
-	if (!ensure(ProcessingContext.EntitySubsystem) ||
+	if (!ensure(ProcessingContext.EntityManager) ||
 		!ensure(RuntimePipeline.Processors.Find(nullptr) == INDEX_NONE) ||
 		RuntimePipeline.Processors.Num() == 0 ||
 		!ensureMsgf(EntityCollection.GetArchetype().IsValid(), TEXT("The Archetype of EntityCollection passed in to UE::Mass::Executor::RunSparse is invalid")))
@@ -54,7 +54,7 @@ void RunSparse(FMassRuntimePipeline& RuntimePipeline, FMassProcessingContext& Pr
 
 void Run(UMassProcessor& Processor, FMassProcessingContext& ProcessingContext)
 {
-	if (!ensure(ProcessingContext.EntitySubsystem) || !ensure(ProcessingContext.DeltaSeconds >= 0.f))
+	if (!ensure(ProcessingContext.EntityManager) || !ensure(ProcessingContext.DeltaSeconds >= 0.f))
 	{
 		return;
 	}
@@ -69,9 +69,9 @@ void RunProcessorsView(TArrayView<UMassProcessor*> Processors, FMassProcessingCo
 {
 	TRACE_CPUPROFILER_EVENT_SCOPE(RunProcessorsView);
 
-	if (ProcessingContext.EntitySubsystem == nullptr)
+	if (!ProcessingContext.EntityManager)
 	{
-		UE_LOG(LogMass, Error, TEXT("%s ProcessingContext.EntitySubsystem is null. Baling out."), ANSI_TO_TCHAR(__FUNCTION__));
+		UE_LOG(LogMass, Error, TEXT("%s ProcessingContext.EntityManager is null. Baling out."), ANSI_TO_TCHAR(__FUNCTION__));
 		return;
 	}
 #if WITH_MASSENTITY_DEBUG
@@ -99,14 +99,15 @@ void RunProcessorsView(TArrayView<UMassProcessor*> Processors, FMassProcessingCo
 	ExecutionContext.SetAuxData(ProcessingContext.AuxData);
 	ExecutionContext.SetExecutionType(EMassExecutionContextType::Processor);
 
+	FMassEntityManager& EntityManager = *ProcessingContext.EntityManager.Get();
 	{
 		TRACE_CPUPROFILER_EVENT_SCOPE_STR("Execute Processors")
 		
-		UMassEntitySubsystem::FScopedProcessing ProcessingScope = ProcessingContext.EntitySubsystem->NewProcessingScope();
+		FMassEntityManager::FScopedProcessing ProcessingScope = EntityManager.NewProcessingScope();
 
 		for (UMassProcessor* Proc : Processors)
 		{
-			Proc->CallExecute(*ProcessingContext.EntitySubsystem, ExecutionContext);
+			Proc->CallExecute(*ProcessingContext.EntityManager, ExecutionContext);
 		}
 	}
 	
@@ -116,31 +117,31 @@ void RunProcessorsView(TArrayView<UMassProcessor*> Processors, FMassProcessingCo
 		
 		ExecutionContext.SetFlushDeferredCommands(true);
 		// append the commands added from other, non-processor sources (like MassAgentSubsystem)
-		ensure(!ProcessingContext.EntitySubsystem->Defer().IsFlushing());
+		ensure(!EntityManager.Defer().IsFlushing());
 		ensure(!ExecutionContext.Defer().IsFlushing());
-		ExecutionContext.Defer().MoveAppend(ProcessingContext.EntitySubsystem->Defer());
-		ExecutionContext.FlushDeferred(*ProcessingContext.EntitySubsystem);
+		ExecutionContext.Defer().MoveAppend(EntityManager.Defer());
+		ExecutionContext.FlushDeferred(*ProcessingContext.EntityManager);
 	}
 	// else make sure we don't just lose the commands. Append to the command buffer requested via
-	// ProcessingContext.CommandBuffer or to the default EntitySubsystem's command buffer.
+	// ProcessingContext.CommandBuffer or to the default EntityManager's command buffer.
 	else if (CommandBuffer != ProcessingContext.CommandBuffer)
 	{
 		if (ProcessingContext.CommandBuffer)
 		{
 			ProcessingContext.CommandBuffer->MoveAppend(*CommandBuffer.Get());
 		}
-		else if (CommandBuffer.Get() != &ProcessingContext.EntitySubsystem->Defer())
+		else if (CommandBuffer.Get() != &EntityManager.Defer())
 		{
-			ProcessingContext.EntitySubsystem->Defer().MoveAppend(*CommandBuffer.Get());
+			EntityManager.Defer().MoveAppend(*CommandBuffer.Get());
 		}
 	}
 }
 
 struct FMassExecutorDoneTask
 {
-	FMassExecutorDoneTask(const FMassExecutionContext& InExecutionContext, UMassEntitySubsystem& InEntitySubsystem, TFunction<void()> InOnDoneNotification, const FString& InDebugName)
+	FMassExecutorDoneTask(const FMassExecutionContext& InExecutionContext, const TSharedPtr<FMassEntityManager>& InEntityManager, TFunction<void()> InOnDoneNotification, const FString& InDebugName)
 		: ExecutionContext(InExecutionContext)
-		, EntitySubsystem(InEntitySubsystem)
+		, EntityManager(InEntityManager)
 		, OnDoneNotification(InOnDoneNotification)
 		, DebugName(InDebugName)
 	{
@@ -156,20 +157,23 @@ struct FMassExecutorDoneTask
 	{
 		TRACE_CPUPROFILER_EVENT_SCOPE_STR("Flush Deferred Commands Parallel");
 
-		if (&ExecutionContext.Defer() != &EntitySubsystem.Defer())
+		check(EntityManager);
+		FMassEntityManager& EntityManagerRef = *EntityManager.Get();
+
+		if (&ExecutionContext.Defer() != &EntityManagerRef.Defer())
 		{
-			ExecutionContext.Defer().MoveAppend(EntitySubsystem.Defer());
+			ExecutionContext.Defer().MoveAppend(EntityManagerRef.Defer());
 		}
 
 		UE_LOG(LogMass, Verbose, TEXT("MassExecutor %s tasks DONE"), *DebugName);
 		ExecutionContext.SetFlushDeferredCommands(true);
-		ExecutionContext.FlushDeferred(EntitySubsystem);
+		ExecutionContext.FlushDeferred(EntityManagerRef);
 
 		OnDoneNotification();
 	}
 private:
 	FMassExecutionContext ExecutionContext;
-	UMassEntitySubsystem& EntitySubsystem;
+	TSharedPtr<FMassEntityManager> EntityManager;
 	TFunction<void()> OnDoneNotification;
 	FString DebugName;
 };
@@ -178,15 +182,15 @@ FGraphEventRef TriggerParallelTasks(UMassProcessor& Processor, FMassProcessingCo
 {
 	TRACE_CPUPROFILER_EVENT_SCOPE(RunProcessorsView);
 
-	if (ProcessingContext.EntitySubsystem == nullptr)
+	if (!ProcessingContext.EntityManager)
 	{
-		UE_LOG(LogMass, Error, TEXT("%s ProcessingContext.EntitySubsystem is null. Baling out."), ANSI_TO_TCHAR(__FUNCTION__));
+		UE_LOG(LogMass, Error, TEXT("%s ProcessingContext.EntityManager is null. Baling out."), ANSI_TO_TCHAR(__FUNCTION__));
 		return FGraphEventRef();
 	}
 
 	TRACE_CPUPROFILER_EVENT_SCOPE_STR("MassExecutor RunParallel")
 
-	// not going through UMassEntitySubsystem::CreateExecutionContext on purpose - we do need a separate command buffer
+	// not going through FMassEntityManager::CreateExecutionContext on purpose - we do need a separate command buffer
 	FMassExecutionContext ExecutionContext(ProcessingContext.DeltaSeconds);
 	TSharedPtr<FMassCommandBuffer> CommandBuffer = ProcessingContext.CommandBuffer
 		? ProcessingContext.CommandBuffer : MakeShareable(new FMassCommandBuffer());
@@ -198,14 +202,14 @@ FGraphEventRef TriggerParallelTasks(UMassProcessor& Processor, FMassProcessingCo
 	FGraphEventRef CompletionEvent;
 	{
 		TRACE_CPUPROFILER_EVENT_SCOPE_STR("Dispatch Processors")
-		CompletionEvent = Processor.DispatchProcessorTasks(*ProcessingContext.EntitySubsystem, ExecutionContext, {});
+		CompletionEvent = Processor.DispatchProcessorTasks(ProcessingContext.EntityManager, ExecutionContext, {});
 	}
 
 	if (CompletionEvent.IsValid())
 	{
 		const FGraphEventArray Prerequisites = { CompletionEvent };
 		CompletionEvent = TGraphTask<FMassExecutorDoneTask>::CreateTask(&Prerequisites)
-			.ConstructAndDispatchWhenReady(ExecutionContext, *ProcessingContext.EntitySubsystem, OnDoneNotification, Processor.GetName());
+			.ConstructAndDispatchWhenReady(ExecutionContext, ProcessingContext.EntityManager, OnDoneNotification, Processor.GetName());
 	}
 
 	return CompletionEvent;
diff --git a/Engine/Plugins/Runtime/MassEntity/Source/MassEntity/Private/MassObserverManager.cpp b/Engine/Plugins/Runtime/MassEntity/Source/MassEntity/Private/MassObserverManager.cpp
--- a/Engine/Plugins/Runtime/MassEntity/Source/MassEntity/Private/MassObserverManager.cpp
+++ b/Engine/Plugins/Runtime/MassEntity/Source/MassEntity/Private/MassObserverManager.cpp
@@ -10,9 +10,11 @@ namespace UE::Mass::ObserverManager::Private
 {
 // a helper function to reduce code duplication in FMassObserverManager::Initialize
 template<typename TBitSet, typename TPointerType>
-void SetUpObservers(UMassEntitySubsystem& EntitySubsystem, const TMap<TPointerType, FMassProcessorClassCollection>& RegisteredObserverTypes, TBitSet& ObservedBitSet, FMassObserversMap& Observers)
+void SetUpObservers(FMassEntityManager& EntityManager, const TMap<TPointerType, FMassProcessorClassCollection>& RegisteredObserverTypes, TBitSet& ObservedBitSet, FMassObserversMap& Observers)
 {
 	ObservedBitSet.Reset();
+	UObject* Owner = EntityManager.GetOwner();
+	check(Owner);
 
 	for (auto It : RegisteredObserverTypes)
 	{
@@ -26,9 +28,9 @@ void SetUpObservers(UMassEntitySubsystem& EntitySubsystem, const TMap<TPointerTy
 
 		for (const TSubclassOf<UMassProcessor>& ProcessorClass : It.Value.ClassCollection)
 		{
-			Pipeline.AppendProcessor(ProcessorClass, EntitySubsystem);
+			Pipeline.AppendProcessor(ProcessorClass, *Owner);
 		}
-		Pipeline.Initialize(EntitySubsystem);
+		Pipeline.Initialize(*Owner);
 	}
 };
 
@@ -38,13 +40,13 @@ void SetUpObservers(UMassEntitySubsystem& EntitySubsystem, const TMap<TPointerTy
 // FMassObserverManager
 //----------------------------------------------------------------------//
 FMassObserverManager::FMassObserverManager()
-	: EntitySubsystem(*GetMutableDefault<UMassEntitySubsystem>())
+	: EntityManager(GetMutableDefault<UMassEntitySubsystem>()->GetMutableEntityManager())
 {
 
 }
 
-FMassObserverManager::FMassObserverManager(UMassEntitySubsystem& Owner)
-	: EntitySubsystem(Owner)
+FMassObserverManager::FMassObserverManager(FMassEntityManager& Owner)
+	: EntityManager(Owner)
 {
 
 }
@@ -57,14 +59,14 @@ void FMassObserverManager::Initialize()
 	using UE::Mass::ObserverManager::Private::SetUpObservers;
 	for (int i = 0; i < (int)EMassObservedOperation::MAX; ++i)
 	{
-		SetUpObservers(EntitySubsystem, *Registry.FragmentObservers[i], ObservedFragments[i], FragmentObservers[i]);
-		SetUpObservers(EntitySubsystem, *Registry.TagObservers[i], ObservedTags[i], TagObservers[i]);
+		SetUpObservers(EntityManager, *Registry.FragmentObservers[i], ObservedFragments[i], FragmentObservers[i]);
+		SetUpObservers(EntityManager, *Registry.TagObservers[i], ObservedTags[i], TagObservers[i]);
 	}
 }
 
 bool FMassObserverManager::OnPostEntitiesCreated(const FMassArchetypeEntityCollection& EntityCollection)
 {
-	FMassProcessingContext ProcessingContext(EntitySubsystem, /*DeltaSeconds=*/0.f);
+	FMassProcessingContext ProcessingContext(EntityManager, /*DeltaSeconds=*/0.f);
 	// requesting not to flush commands since handling creation of new entities can result in multiple collections of
 	// processors being executed and flushing commands between these runs would ruin EntityCollection since entities could
 	// get their composition changed and get moved to new archetypes
@@ -73,7 +75,7 @@ bool FMassObserverManager::OnPostEntitiesCreated(const FMassArchetypeEntityColle
 
 	if (OnPostEntitiesCreated(ProcessingContext, EntityCollection))
 	{
-		EntitySubsystem.FlushCommands(ProcessingContext.CommandBuffer);
+		EntityManager.FlushCommands(ProcessingContext.CommandBuffer);
 		return true;
 	}
 	return false;
@@ -83,8 +85,8 @@ bool FMassObserverManager::OnPostEntitiesCreated(FMassProcessingContext& Process
 {
 	TRACE_CPUPROFILER_EVENT_SCOPE_STR("OnPostEntitiesCreated")
 
-	check(ProcessingContext.EntitySubsystem);
-	const FMassArchetypeCompositionDescriptor& ArchetypeComposition = ProcessingContext.EntitySubsystem->GetArchetypeComposition(EntityCollection.GetArchetype());
+	check(ProcessingContext.EntityManager);
+	const FMassArchetypeCompositionDescriptor& ArchetypeComposition = ProcessingContext.EntityManager->GetArchetypeComposition(EntityCollection.GetArchetype());
 	const FMassFragmentBitSet Overlap = ObservedFragments[(uint8)EMassObservedOperation::Add].GetOverlap(ArchetypeComposition.Fragments);
 
 	if (Overlap.IsEmpty() == false)
@@ -101,13 +103,13 @@ bool FMassObserverManager::OnPostEntitiesCreated(FMassProcessingContext& Process
 
 bool FMassObserverManager::OnPreEntitiesDestroyed(const FMassArchetypeEntityCollection& EntityCollection)
 {
-	FMassProcessingContext ProcessingContext(EntitySubsystem, /*DeltaSeconds=*/0.f);
+	FMassProcessingContext ProcessingContext(EntityManager, /*DeltaSeconds=*/0.f);
 	ProcessingContext.bFlushCommandBuffer = false;
 	ProcessingContext.CommandBuffer = MakeShareable(new FMassCommandBuffer());
 
 	if (OnPreEntitiesDestroyed(ProcessingContext, EntityCollection))
 	{
-		EntitySubsystem.FlushCommands(ProcessingContext.CommandBuffer);
+		EntityManager.FlushCommands(ProcessingContext.CommandBuffer);
 		return true;
 	}
 	return false;
@@ -117,8 +119,8 @@ bool FMassObserverManager::OnPreEntitiesDestroyed(FMassProcessingContext& Proces
 {
 	TRACE_CPUPROFILER_EVENT_SCOPE_STR("OnPreEntitiesDestroyed")
 
-	check(ProcessingContext.EntitySubsystem);
-	const FMassArchetypeCompositionDescriptor& ArchetypeComposition = ProcessingContext.EntitySubsystem->GetArchetypeComposition(EntityCollection.GetArchetype());
+	check(ProcessingContext.EntityManager);
+	const FMassArchetypeCompositionDescriptor& ArchetypeComposition = ProcessingContext.EntityManager->GetArchetypeComposition(EntityCollection.GetArchetype());
 	
 	return OnCompositionChanged(EntityCollection, ArchetypeComposition, EMassObservedOperation::Remove, &ProcessingContext);
 }
@@ -132,7 +134,7 @@ bool FMassObserverManager::OnCompositionChanged(const FMassArchetypeEntityCollec
 
 	if (bHasFragmentsOverlap || bHasTagsOverlap)
 	{
-		FMassProcessingContext LocalContext(EntitySubsystem, /*DeltaSeconds=*/0.f);
+		FMassProcessingContext LocalContext(EntityManager, /*DeltaSeconds=*/0.f);
 		LocalContext.bFlushCommandBuffer = false;
 		FMassProcessingContext* ProcessingContext = InProcessingContext ? InProcessingContext : &LocalContext;
 		TArray<const UScriptStruct*> ObservedTypesOverlap;
@@ -166,9 +168,9 @@ bool FMassObserverManager::OnCompositionChanged(const FMassEntityHandle Entity,
 	if (bHasFragmentsOverlap || bHasTagsOverlap)
 	{
 		TArray<const UScriptStruct*> ObservedTypesOverlap;
-		FMassProcessingContext ProcessingContext(EntitySubsystem, /*DeltaSeconds=*/0.f);
+		FMassProcessingContext ProcessingContext(EntityManager, /*DeltaSeconds=*/0.f);
 		ProcessingContext.bFlushCommandBuffer = false;
-		const FMassArchetypeHandle ArchetypeHandle = EntitySubsystem.GetArchetypeForEntity(Entity);
+		const FMassArchetypeHandle ArchetypeHandle = EntityManager.GetArchetypeForEntity(Entity);
 
 		if (bHasFragmentsOverlap)
 		{
@@ -229,7 +231,7 @@ void FMassObserverManager::HandleSingleEntityImpl(const UScriptStruct& FragmentT
 {
 	TRACE_CPUPROFILER_EVENT_SCOPE(MassObserver_HandleSingleEntityImpl);
 
-	FMassProcessingContext ProcessingContext(EntitySubsystem, /*DeltaSeconds=*/0.f);
+	FMassProcessingContext ProcessingContext(EntityManager, /*DeltaSeconds=*/0.f);
 	ProcessingContext.bFlushCommandBuffer = false;
 	ProcessingContext.AuxData.InitializeAs(&FragmentType);
 	FMassRuntimePipeline& Pipeline = (*HandlersContainer).FindChecked(&FragmentType);
@@ -255,6 +257,9 @@ void FMassObserverManager::AddObserverInstance(const UScriptStruct& FragmentOrTa
 	}
 	Pipeline->AppendProcessor(ObserverProcessor);
 
-	// calling initialize to ensure the given processor is related to the same EntitySubsystem
-	ObserverProcessor.Initialize(EntitySubsystem);
+	// calling initialize to ensure the given processor is related to the same EntityManager
+	if (UObject* Owner = EntityManager.GetOwner())
+	{	
+		ObserverProcessor.Initialize(*Owner);
+	}
 }
diff --git a/Engine/Plugins/Runtime/MassEntity/Source/MassEntity/Private/MassProcessingPhase.cpp b/Engine/Plugins/Runtime/MassEntity/Source/MassEntity/Private/MassProcessingPhase.cpp
--- a/Engine/Plugins/Runtime/MassEntity/Source/MassEntity/Private/MassProcessingPhase.cpp
+++ b/Engine/Plugins/Runtime/MassEntity/Source/MassEntity/Private/MassProcessingPhase.cpp
@@ -53,11 +53,11 @@ void FMassProcessingPhase::ExecuteTick(float DeltaTime, ELevelTick TickType, ENa
 		return;
 	}
 
-	checkf(Manager, TEXT("Manager is null which is not a supported case. Either this FMassProcessingPhase has not been initialized properly or it's been left dangling after the FMassProcessingPhase owner got destroyed."));
+	checkf(PhaseManager, TEXT("Manager is null which is not a supported case. Either this FMassProcessingPhase has not been initialized properly or it's been left dangling after the FMassProcessingPhase owner got destroyed."));
 
 	TRACE_CPUPROFILER_EVENT_SCOPE_STR(*FString::Printf(TEXT("FMassProcessingPhase::ExecuteTick %s"), *UEnum::GetValueAsString(Phase)));
 
-	Manager->OnPhaseStart(*this);
+	PhaseManager->OnPhaseStart(*this);
 	{
 		LLM_SCOPE_BYNAME(TEXT("Mass/PhaseStartDelegate"));
 		OnPhaseStart.Broadcast(DeltaTime);
@@ -65,8 +65,8 @@ void FMassProcessingPhase::ExecuteTick(float DeltaTime, ELevelTick TickType, ENa
 
 	check(PhaseProcessor);
 	
-	UMassEntitySubsystem& EntitySubsystem = Manager->GetEntitySubsystemRef();
-	FMassProcessingContext Context(EntitySubsystem, DeltaTime);
+	FMassEntityManager& EntityManager = PhaseManager->GetEntityManagerRef();
+	FMassProcessingContext Context(EntityManager, DeltaTime);
 
 	bIsDuringMassProcessing = true;
 
@@ -99,7 +99,7 @@ void FMassProcessingPhase::ExecuteTick(float DeltaTime, ELevelTick TickType, ENa
 			LLM_SCOPE_BYNAME(TEXT("Mass/PhaseEndDelegate"));
 			OnPhaseEnd.Broadcast(DeltaTime);
 		}
-		Manager->OnPhaseEnd(*this);
+		PhaseManager->OnPhaseEnd(*this);
 		bIsDuringMassProcessing = false;
 	}
 }
@@ -111,18 +111,18 @@ void FMassProcessingPhase::OnParallelExecutionDone(const float DeltaTime)
 		LLM_SCOPE_BYNAME(TEXT("Mass/PhaseEndDelegate"));
 		OnPhaseEnd.Broadcast(DeltaTime);
 	}
-	check(Manager);
-	Manager->OnPhaseEnd(*this);
+	check(PhaseManager);
+	PhaseManager->OnPhaseEnd(*this);
 }
 
 FString FMassProcessingPhase::DiagnosticMessage()
 {
-	return (Manager ? Manager->GetFullName() : TEXT("NULL-MassProcessingPhaseManager")) + TEXT("[ProcessorTick]");
+	return (PhaseManager ? PhaseManager->GetFullName() : TEXT("NULL-MassProcessingPhaseManager")) + TEXT("[ProcessorTick]");
 }
 
 FName FMassProcessingPhase::DiagnosticContext(bool bDetailed)
 {
-	return Manager ? Manager->GetClass()->GetFName() : TEXT("NULL-MassProcessingPhaseManager");
+	return PhaseManager ? PhaseManager->GetClass()->GetFName() : TEXT("NULL-MassProcessingPhaseManager");
 }
 
 //----------------------------------------------------------------------//
@@ -202,10 +202,11 @@ void UMassProcessingPhaseManager::InitializePhases(UObject& InProcessorOwner)
 
 void UMassProcessingPhaseManager::Start(UWorld& World)
 {
-	EntitySubsystem = UWorld::GetSubsystem<UMassEntitySubsystem>(&World);
+	UMassEntitySubsystem* EntitySubsystem = UWorld::GetSubsystem<UMassEntitySubsystem>(&World);
 
 	if (ensure(EntitySubsystem))
 	{
+		EntityManager = EntitySubsystem->GetMutableEntityManager().AsShared();
 		EnableTickFunctions(World);
 	}
 	else
@@ -214,23 +215,23 @@ void UMassProcessingPhaseManager::Start(UWorld& World)
 	}
 }
 
-void UMassProcessingPhaseManager::Start(UMassEntitySubsystem& InEntitySubsystem)
+void UMassProcessingPhaseManager::Start(const TSharedPtr<FMassEntityManager>& InEntityManager)
 {
-	UWorld* World = InEntitySubsystem.GetWorld();
+	UWorld* World = InEntityManager->GetWorld();
 	check(World);
-	EntitySubsystem = &InEntitySubsystem;
+	EntityManager = InEntityManager;
 	EnableTickFunctions(*World);
 }
 
 void UMassProcessingPhaseManager::EnableTickFunctions(const UWorld& World)
 {
-	check(EntitySubsystem);
+	check(EntityManager);
 
 	const bool bIsGameWorld = World.IsGameWorld();
 
 	for (FMassProcessingPhase& Phase : ProcessingPhases)
 	{
-		Phase.Manager = this;
+		Phase.PhaseManager = this;
 		Phase.RegisterTickFunction(World.PersistentLevel);
 		Phase.SetTickFunctionEnable(true);
 #if WITH_MASSENTITY_DEBUG
@@ -255,7 +256,7 @@ void UMassProcessingPhaseManager::EnableTickFunctions(const UWorld& World)
 
 void UMassProcessingPhaseManager::Stop()
 {
-	EntitySubsystem = nullptr;
+	EntityManager.Reset();
 	
 	for (FMassProcessingPhase& Phase : ProcessingPhases)
 	{
diff --git a/Engine/Plugins/Runtime/MassEntity/Source/MassEntity/Private/MassProcessingTypes.cpp b/Engine/Plugins/Runtime/MassEntity/Source/MassEntity/Private/MassProcessingTypes.cpp
--- a/Engine/Plugins/Runtime/MassEntity/Source/MassEntity/Private/MassProcessingTypes.cpp
+++ b/Engine/Plugins/Runtime/MassEntity/Source/MassEntity/Private/MassProcessingTypes.cpp
@@ -12,8 +12,14 @@ DEFINE_LOG_CATEGORY(LogMass);
 //----------------------------------------------------------------------//
 //  FMassProcessingContext
 //----------------------------------------------------------------------//
-FMassProcessingContext::FMassProcessingContext(UMassEntitySubsystem& InEntitySubsystem, const float InDeltaSeconds)
-	: EntitySubsystem(&InEntitySubsystem), DeltaSeconds(InDeltaSeconds)
+FMassProcessingContext::FMassProcessingContext(FMassEntityManager& InEntityManager, const float InDeltaSeconds)
+	: EntityManager(InEntityManager.AsShared()), DeltaSeconds(InDeltaSeconds)
+{
+
+}
+
+FMassProcessingContext::FMassProcessingContext(TSharedPtr<FMassEntityManager>& InEntityManager, const float InDeltaSeconds)
+	: EntityManager(InEntityManager), DeltaSeconds(InDeltaSeconds)
 {
 
 }
@@ -22,12 +28,12 @@ FMassProcessingContext::~FMassProcessingContext()
 {
 	if (CommandBuffer && CommandBuffer.IsUnique() && CommandBuffer->HasPendingCommands())
 	{
-		UE_CLOG(EntitySubsystem == nullptr, LogMass, Error, TEXT("Unable to auto-flush FMassProcessingContext\'s commands due to missing EntitySubsystem"));
-		if (ensure(EntitySubsystem))
+		UE_CLOG(!EntityManager, LogMass, Error, TEXT("Unable to auto-flush FMassProcessingContext\'s commands due to missing EntityManager"));
+		if (ensure(EntityManager))
 		{
-			UE_VLOG(EntitySubsystem, LogMass, Log, TEXT("Auto-flushing command buffer as part of FMassProcessingContext destruction"));
+			UE_VLOG(EntityManager->GetOwner(), LogMass, Log, TEXT("Auto-flushing command buffer as part of FMassProcessingContext destruction"));
 			checkf(CommandBuffer->IsFlushing() == false, TEXT("A totally unexpected scenario."));
-			EntitySubsystem->FlushCommands(CommandBuffer);
+			EntityManager->FlushCommands(CommandBuffer);
 		}
 	}
 }
diff --git a/Engine/Plugins/Runtime/MassEntity/Source/MassEntity/Private/MassProcessor.cpp b/Engine/Plugins/Runtime/MassEntity/Source/MassEntity/Private/MassProcessor.cpp
--- a/Engine/Plugins/Runtime/MassEntity/Source/MassEntity/Private/MassProcessor.cpp
+++ b/Engine/Plugins/Runtime/MassEntity/Source/MassEntity/Private/MassProcessor.cpp
@@ -34,8 +34,8 @@ namespace UE::Mass::Debug
 class FMassProcessorTask
 {
 public:
-	FMassProcessorTask(UMassEntitySubsystem& InEntitySubsystem, const FMassExecutionContext& InExecutionContext, UMassProcessor& InProc, bool bInManageCommandBuffer = true)
-		: EntitySubsystem(&InEntitySubsystem)
+	FMassProcessorTask(const TSharedPtr<FMassEntityManager>& InEntityManager, const FMassExecutionContext& InExecutionContext, UMassProcessor& InProc, bool bInManageCommandBuffer = true)
+		: EntityManager(InEntityManager)
 		, ExecutionContext(InExecutionContext)
 		, Processor(&InProc)
 		, bManageCommandBuffer(bInManageCommandBuffer)
@@ -56,11 +56,12 @@ public:
 	void DoTask(ENamedThreads::Type CurrentThread, const FGraphEventRef& MyCompletionGraphEvent)
 	{
 		checkf(Processor, TEXT("Expecting a valid processor to execute"));
-		checkf(EntitySubsystem, TEXT("Expecting a valid entity subsystem to execute processor"));
 
 		PROCESSOR_TASK_LOG(TEXT("+--+ Task %s started on %u"), *Processor->GetProcessorName(), FPlatformTLS::GetCurrentThreadId());
 
-		UMassEntitySubsystem::FScopedProcessing ProcessingScope = EntitySubsystem->NewProcessingScope();
+		check(EntityManager);
+		FMassEntityManager& EntityManagerRef = *EntityManager.Get();
+		FMassEntityManager::FScopedProcessing ProcessingScope = EntityManagerRef.NewProcessingScope();
 
 		TRACE_CPUPROFILER_EVENT_SCOPE_STR("Mass Processor Task");
 		
@@ -68,18 +69,18 @@ public:
 		{
 			TSharedPtr<FMassCommandBuffer> MainSharedPtr = ExecutionContext.GetSharedDeferredCommandBuffer();
 			ExecutionContext.SetDeferredCommandBuffer(MakeShareable(new FMassCommandBuffer()));
-			Processor->CallExecute(*EntitySubsystem, ExecutionContext);
+			Processor->CallExecute(EntityManagerRef, ExecutionContext);
 			MainSharedPtr->MoveAppend(ExecutionContext.Defer());
 		}
 		else
 		{
-			Processor->CallExecute(*EntitySubsystem, ExecutionContext);
+			Processor->CallExecute(EntityManagerRef, ExecutionContext);
 		}
 		PROCESSOR_TASK_LOG(TEXT("+--+ Task %s finished"), *Processor->GetProcessorName());
 	}
 
 private:
-	UMassEntitySubsystem* EntitySubsystem = nullptr;
+	TSharedPtr<FMassEntityManager> EntityManager;
 	FMassExecutionContext ExecutionContext;
 	UMassProcessor* Processor = nullptr;
 	/** 
@@ -92,8 +93,8 @@ private:
 class FMassProcessorsTask_GameThread : public FMassProcessorTask
 {
 public:
-	FMassProcessorsTask_GameThread(UMassEntitySubsystem& InEntitySubsystem, const FMassExecutionContext& InExecutionContext, UMassProcessor& InProc)
-		: FMassProcessorTask(InEntitySubsystem, InExecutionContext, InProc)
+	FMassProcessorsTask_GameThread(const TSharedPtr<FMassEntityManager>& InEntityManager, const FMassExecutionContext& InExecutionContext, UMassProcessor& InProc)
+		: FMassProcessorTask(InEntityManager, InExecutionContext, InProc)
 	{}
 
 	static ENamedThreads::Type GetDesiredThread()
@@ -151,14 +152,14 @@ void UMassProcessor::PostInitProperties()
 #endif
 }
 
-void UMassProcessor::CallExecute(UMassEntitySubsystem& EntitySubsystem, FMassExecutionContext& Context)
+void UMassProcessor::CallExecute(FMassEntityManager& EntityManager, FMassExecutionContext& Context)
 {
 	TRACE_CPUPROFILER_EVENT_SCOPE_TEXT(*StatId);
 	LLM_SCOPE_BYNAME(TEXT("Mass/ExecuteProcessor"));
 	// Not using a more specific scope by default (i.e. LLM_SCOPE_BYNAME(*StatId)) since LLM is more strict regarding the provided string (no spaces or '_')
 
 #if WITH_MASSENTITY_DEBUG
-	Context.DebugSetExecutionDesc(FString::Printf(TEXT("%s (%s)"), *GetProcessorName(), *ToString(EntitySubsystem.GetWorld()->GetNetMode())));
+	Context.DebugSetExecutionDesc(FString::Printf(TEXT("%s (%s)"), *GetProcessorName(), EntityManager.GetWorld() ? *ToString(EntityManager.GetWorld()->GetNetMode()) : TEXT("No World")));
 #endif
 	// CacheSubsystemRequirements will return true only if all requirements declared with ProcessorRequirements are met
 	// meaning if it fails there's no point in calling Execute.
@@ -166,7 +167,7 @@ void UMassProcessor::CallExecute(UMassEntitySubsystem& EntitySubsystem, FMassExe
 	// of their queries not having anything to do.
 	if (Context.CacheSubsystemRequirements(GetWorld(), ProcessorRequirements))
 	{
-		Execute(EntitySubsystem, Context);
+		Execute(EntityManager, Context);
 	}
 	else
 	{
@@ -202,16 +203,16 @@ void UMassProcessor::RegisterQuery(FMassEntityQuery& Query)
 	}
 }
 
-FGraphEventRef UMassProcessor::DispatchProcessorTasks(UMassEntitySubsystem& EntitySubsystem, FMassExecutionContext& ExecutionContext, const FGraphEventArray& Prerequisites)
+FGraphEventRef UMassProcessor::DispatchProcessorTasks(const TSharedPtr<FMassEntityManager>& EntityManager, FMassExecutionContext& ExecutionContext, const FGraphEventArray& Prerequisites)
 {
 	FGraphEventRef ReturnVal;
 	if (bRequiresGameThreadExecution)
 	{
-		ReturnVal = TGraphTask<FMassProcessorsTask_GameThread>::CreateTask(&Prerequisites).ConstructAndDispatchWhenReady(EntitySubsystem, ExecutionContext, *this);
+		ReturnVal = TGraphTask<FMassProcessorsTask_GameThread>::CreateTask(&Prerequisites).ConstructAndDispatchWhenReady(EntityManager, ExecutionContext, *this);
 	}
 	else
 	{
-		ReturnVal = TGraphTask<FMassProcessorTask>::CreateTask(&Prerequisites).ConstructAndDispatchWhenReady(EntitySubsystem, ExecutionContext, *this);
+		ReturnVal = TGraphTask<FMassProcessorTask>::CreateTask(&Prerequisites).ConstructAndDispatchWhenReady(EntityManager, ExecutionContext, *this);
 	}	
 	return ReturnVal;
 }
@@ -259,7 +260,7 @@ void UMassCompositeProcessor::ConfigureQueries()
 	// via their individual PostInitProperties call
 }
 
-FGraphEventRef UMassCompositeProcessor::DispatchProcessorTasks(UMassEntitySubsystem& EntitySubsystem, FMassExecutionContext& ExecutionContext, const FGraphEventArray& InPrerequisites)
+FGraphEventRef UMassCompositeProcessor::DispatchProcessorTasks(const TSharedPtr<FMassEntityManager>& EntityManager, FMassExecutionContext& ExecutionContext, const FGraphEventArray& InPrerequisites)
 {
 	FGraphEventArray Events;
 	Events.Reserve(ProcessingFlatGraph.Num());
@@ -275,7 +276,7 @@ FGraphEventRef UMassCompositeProcessor::DispatchProcessorTasks(UMassEntitySubsys
 		// we don't expect any group nodes at this point. If we get any there's a bug in dependencies solving
 		if (ensure(ProcessingNode.Processor))
 		{
-			Events.Add(ProcessingNode.Processor->DispatchProcessorTasks(EntitySubsystem, ExecutionContext, Prerequisites));
+			Events.Add(ProcessingNode.Processor->DispatchProcessorTasks(EntityManager, ExecutionContext, Prerequisites));
 		}
 	}
 
@@ -305,12 +306,12 @@ FGraphEventRef UMassCompositeProcessor::DispatchProcessorTasks(UMassEntitySubsys
 	return CompletionEvent;
 }
 
-void UMassCompositeProcessor::Execute(UMassEntitySubsystem& EntitySubsystem, FMassExecutionContext& Context)
+void UMassCompositeProcessor::Execute(FMassEntityManager& EntityManager, FMassExecutionContext& Context)
 {
 	for (UMassProcessor* Proc : ChildPipeline.Processors)
 	{
 		check(Proc);
-		Proc->CallExecute(EntitySubsystem, Context);
+		Proc->CallExecute(EntityManager, Context);
 	}
 }
 
diff --git a/Engine/Plugins/Runtime/MassEntity/Source/MassEntity/Public/MassArchetypeTypes.h b/Engine/Plugins/Runtime/MassEntity/Source/MassEntity/Public/MassArchetypeTypes.h
--- a/Engine/Plugins/Runtime/MassEntity/Source/MassEntity/Public/MassArchetypeTypes.h
+++ b/Engine/Plugins/Runtime/MassEntity/Source/MassEntity/Public/MassArchetypeTypes.h
@@ -8,7 +8,7 @@
 #include "Containers/UnrealString.h"
 #include "MassEntityTypes.h"
 
-class UMassEntitySubsystem;
+struct FMassEntityManager;
 struct FMassArchetypeData;
 struct FMassExecutionContext;
 struct FMassFragment;
@@ -43,7 +43,7 @@ private:
 	TSharedPtr<FMassArchetypeData> DataPtr;
 
 	friend FMassArchetypeHelper;
-	friend UMassEntitySubsystem;
+	friend FMassEntityManager;
 };
 
 
@@ -136,7 +136,7 @@ struct MASSENTITY_API FMassArchetypeEntityCollectionWithPayload
 	{
 	}
 
-	static void CreateEntityRangesWithPayload(const UMassEntitySubsystem& EntitySystem, const TConstArrayView<FMassEntityHandle> Entities
+	static void CreateEntityRangesWithPayload(const FMassEntityManager& EntitySubsystem, const TConstArrayView<FMassEntityHandle> Entities
 		, const FMassArchetypeEntityCollection::EDuplicatesHandling DuplicatesHandling, FMassGenericPayloadView Payload
 		, TArray<FMassArchetypeEntityCollectionWithPayload>& OutEntityCollections);
 
diff --git a/Engine/Plugins/Runtime/MassEntity/Source/MassEntity/Public/MassCommandBuffer.h b/Engine/Plugins/Runtime/MassEntity/Source/MassEntity/Public/MassCommandBuffer.h
--- a/Engine/Plugins/Runtime/MassEntity/Source/MassEntity/Public/MassCommandBuffer.h
+++ b/Engine/Plugins/Runtime/MassEntity/Source/MassEntity/Public/MassCommandBuffer.h
@@ -121,7 +121,7 @@ public:
 	bool IsFlushing() const { return bIsFlushing; }
 
 private:
-	friend UMassEntitySubsystem;
+	friend FMassEntityManager;
 
 	template<typename T>
 	T& CreateOrAddCommand()
@@ -144,7 +144,7 @@ private:
 		return *NewCommandInstance;
 	}
 
-	void Flush(UMassEntitySubsystem& EntitySystem);
+	void Flush(FMassEntityManager& EntityManager);
 	void CleanUp();
 
 	FCriticalSection AppendingCommandsCS;
diff --git a/Engine/Plugins/Runtime/MassEntity/Source/MassEntity/Public/MassCommands.h b/Engine/Plugins/Runtime/MassEntity/Source/MassEntity/Public/MassCommands.h
--- a/Engine/Plugins/Runtime/MassEntity/Source/MassEntity/Public/MassCommands.h
+++ b/Engine/Plugins/Runtime/MassEntity/Source/MassEntity/Public/MassCommands.h
@@ -31,7 +31,7 @@ struct MASSENTITY_API FMassBatchedCommand
 {
 	virtual ~FMassBatchedCommand() {}
 
-	virtual void Execute(UMassEntitySubsystem& System) const = 0;
+	virtual void Execute(FMassEntityManager& System) const = 0;
 	virtual void Reset()
 	{
 		bHasWork = false;
@@ -117,7 +117,7 @@ struct FMassCommandDestroyEntities : public FMassBatchedEntityCommand
 	}
 
 protected:
-	virtual void Execute(UMassEntitySubsystem& System) const override
+	virtual void Execute(FMassEntityManager& System) const override
 	{
 		TRACE_CPUPROFILER_EVENT_SCOPE(MassCommandDestroyEntities_Execute);
 
@@ -142,7 +142,7 @@ struct FMassCommandAddFragments : public FMassBatchedEntityCommand
 		OperationType = EMassCommandOperationType::Add;
 	}
 protected:
-	virtual void Execute(UMassEntitySubsystem& System) const override
+	virtual void Execute(FMassEntityManager& System) const override
 	{
 		TRACE_CPUPROFILER_EVENT_SCOPE(MassCommandAddFragments_Execute);
 		TArray<FMassArchetypeEntityCollection> EntityCollections;
@@ -161,7 +161,7 @@ struct FMassCommandRemoveFragments : public FMassBatchedEntityCommand
 		OperationType = EMassCommandOperationType::Remove;
 	}
 protected:
-	virtual void Execute(UMassEntitySubsystem& System) const override
+	virtual void Execute(FMassEntityManager& System) const override
 	{
 		TRACE_CPUPROFILER_EVENT_SCOPE(MassCommandRemoveFragments_Execute);
 		TArray<FMassArchetypeEntityCollection> EntityCollections;
@@ -184,7 +184,7 @@ struct FMassCommandChangeTags : public FMassBatchedEntityCommand
 	}
 
 protected:
-	virtual void Execute(UMassEntitySubsystem& System) const override
+	virtual void Execute(FMassEntityManager& System) const override
 	{
 		TRACE_CPUPROFILER_EVENT_SCOPE(MassCommandChangeTags_Execute);
 		TArray<FMassArchetypeEntityCollection> EntityCollections;
@@ -271,7 +271,7 @@ protected:
 		return Super::GetAllocatedSize() + Fragments.GetAllocatedSize() + FragmentsAffected.GetAllocatedSize();
 	}
 
-	virtual void Execute(UMassEntitySubsystem& System) const override
+	virtual void Execute(FMassEntityManager& System) const override
 	{
 		TRACE_CPUPROFILER_EVENT_SCOPE(MassCommandAddFragmentInstances_Execute);
 
@@ -304,7 +304,7 @@ struct FMassCommandBuildEntity : public FMassCommandAddFragmentInstances<TOthers
 protected:
 	using Super = FMassCommandAddFragmentInstances<TOthers...>;
 
-	virtual void Execute(UMassEntitySubsystem& System) const override
+	virtual void Execute(FMassEntityManager& System) const override
 	{
 		TRACE_CPUPROFILER_EVENT_SCOPE(MassCommandBuildEntity_Execute);
 
@@ -371,7 +371,7 @@ protected:
 		return TotalSize;
 	}
 
-	virtual void Execute(UMassEntitySubsystem& System) const 
+	virtual void Execute(FMassEntityManager& System) const 
 	{
 		TRACE_CPUPROFILER_EVENT_SCOPE(MassCommandBuildEntityWithSharedFragments_Execute);
 
@@ -444,7 +444,7 @@ template<EMassCommandOperationType OpType>
 struct FMassDeferredCommand : public FMassBatchedCommand
 {
 	using Super = FMassBatchedCommand;
-	using FExecFunction = TFunction<void(UMassEntitySubsystem& System)>;
+	using FExecFunction = TFunction<void(FMassEntityManager& System)>;
 
 	FMassDeferredCommand()
 	{
@@ -472,7 +472,7 @@ protected:
 		return DeferredFunctions.GetAllocatedSize();
 	}
 
-	virtual void Execute(UMassEntitySubsystem& System) const override
+	virtual void Execute(FMassEntityManager& System) const override
 	{
 		TRACE_CPUPROFILER_EVENT_SCOPE(MassDeferredCommand_Execute);
 
