Separated "entity management" and USubsystem stuff of UMassEntityManager. Now UMassEntityManager hosts a shared ref of newly created FMassEntityManager, and users are free to create their own FMassEnt

diff --git a/Engine/Plugins/Runtime/MassEntity/Source/MassEntity/Public/MassDebugger.h b/Engine/Plugins/Runtime/MassEntity/Source/MassEntity/Public/MassDebugger.h
--- a/Engine/Plugins/Runtime/MassEntity/Source/MassEntity/Public/MassDebugger.h
+++ b/Engine/Plugins/Runtime/MassEntity/Source/MassEntity/Public/MassDebugger.h
@@ -12,7 +12,7 @@
 class FOutputDevice;
 class UMassProcessor;
 struct FMassEntityQuery;
-class UMassEntitySubsystem;
+struct FMassEntityManager;
 struct FMassArchetypeHandle;
 struct FMassFragmentRequirements;
 struct FMassFragmentRequirementDescription;
@@ -41,17 +41,17 @@ MASSENTITY_API extern void DebugOutputDescription(TConstArrayView<UMassProcessor
 
 struct MASSENTITY_API FMassDebugger
 {
-	DECLARE_MULTICAST_DELEGATE_TwoParams(FOnEntitySelected, const UMassEntitySubsystem&, const FMassEntityHandle);
+	DECLARE_MULTICAST_DELEGATE_TwoParams(FOnEntitySelected, const FMassEntityManager&, const FMassEntityHandle);
 
 	static TConstArrayView<FMassEntityQuery*> GetProcessorQueries(const UMassProcessor& Processor);
 	/** fetches all queries registered for given Processor. Note that in order to get up to date information
 	 *  FMassEntityQuery::CacheArchetypes will be called on each query */
-	static TConstArrayView<FMassEntityQuery*> GetUpToDateProcessorQueries(const UMassEntitySubsystem& EntitySubsystem, UMassProcessor& Processor);
+	static TConstArrayView<FMassEntityQuery*> GetUpToDateProcessorQueries(const FMassEntityManager& EntitySubsystem, UMassProcessor& Processor);
 
 	static UE::Mass::Debug::FQueryRequirementsView GetQueryRequirements(const FMassEntityQuery& Query);
 	static void GetQueryExecutionRequirements(const FMassEntityQuery& Query, FMassExecutionRequirements& OutExecutionRequirements);
 
-	static TArray<FMassArchetypeHandle> GetAllArchetypes(const UMassEntitySubsystem& EntitySubsystem);
+	static TArray<FMassArchetypeHandle> GetAllArchetypes(const FMassEntityManager& EntitySubsystem);
 	static const FMassArchetypeCompositionDescriptor& GetArchetypeComposition(const FMassArchetypeHandle& ArchetypeHandle);
 
 	static void GetArchetypeEntityStats(const FMassArchetypeHandle& ArchetypeHandle, int32& OutEntitiesCount, int32& OutNumEntitiesPerChunk, int32& OutChunksCount);
@@ -64,10 +64,10 @@ struct MASSENTITY_API FMassDebugger
 	static FString GetArchetypeRequirementCompatibilityDescription(const FMassFragmentRequirements& Requirements, const FMassArchetypeHandle& ArchetypeHandle);
 
 	static void OutputArchetypeDescription(FOutputDevice& Ar, const FMassArchetypeHandle& Archetype);
-	static void OutputEntityDescription(FOutputDevice& Ar, const UMassEntitySubsystem& EntitySubsystem, const int32 EntityIndex, const TCHAR* InPrefix = TEXT(""));
-	static void OutputEntityDescription(FOutputDevice& Ar, const UMassEntitySubsystem& EntitySubsystem, const FMassEntityHandle Entity, const TCHAR* InPrefix = TEXT(""));
+	static void OutputEntityDescription(FOutputDevice& Ar, const FMassEntityManager& EntityManager, const int32 EntityIndex, const TCHAR* InPrefix = TEXT(""));
+	static void OutputEntityDescription(FOutputDevice& Ar, const FMassEntityManager& EntityManager, const FMassEntityHandle Entity, const TCHAR* InPrefix = TEXT(""));
 
-	static void SelectEntity(const UMassEntitySubsystem& EntitySubsystem, const FMassEntityHandle EntityHandle);
+	static void SelectEntity(const FMassEntityManager& EntityManager, const FMassEntityHandle EntityHandle);
 
 	static FOnEntitySelected OnEntitySelectedDelegate;
 };
@@ -81,4 +81,4 @@ struct MASSENTITY_API FMassDebugger
 	static FString GetArchetypeRequirementCompatibilityDescription(const FMassFragmentRequirements& Requirements, const FMassArchetypeHandle& ArchetypeHandle) { return TEXT("[no debug information]"); }
 };
 
-#endif // WITH_MASSENTITY_DEBUG
\ No newline at end of file
+#endif // WITH_MASSENTITY_DEBUG
diff --git a/Engine/Plugins/Runtime/MassEntity/Source/MassEntity/Public/MassEntityQuery.h b/Engine/Plugins/Runtime/MassEntity/Source/MassEntity/Public/MassEntityQuery.h
--- a/Engine/Plugins/Runtime/MassEntity/Source/MassEntity/Public/MassEntityQuery.h
+++ b/Engine/Plugins/Runtime/MassEntity/Source/MassEntity/Public/MassEntityQuery.h
@@ -36,15 +36,15 @@ public:
 	void RegisterWithProcessor(UMassProcessor& Owner);
 
 	/** Runs ExecuteFunction on all entities matching Requirements */
-	void ForEachEntityChunk(UMassEntitySubsystem& EntitySubsystem, FMassExecutionContext& ExecutionContext, const FMassExecuteFunction& ExecuteFunction);
+	void ForEachEntityChunk(FMassEntityManager& EntityManager, FMassExecutionContext& ExecutionContext, const FMassExecuteFunction& ExecuteFunction);
 	
 	/** Will first verify that the archetype given with Collection matches the query's requirements, and if so will run the other, more generic ForEachEntityChunk implementation */
-	void ForEachEntityChunk(const FMassArchetypeEntityCollection& Collection, UMassEntitySubsystem& EntitySubsystem, FMassExecutionContext& ExecutionContext, const FMassExecuteFunction& ExecuteFunction);
+	void ForEachEntityChunk(const FMassArchetypeEntityCollection& Collection, FMassEntityManager& EntitySubsystem, FMassExecutionContext& ExecutionContext, const FMassExecuteFunction& ExecuteFunction);
 
-	/** Will gather all archetypes from InEntitySubsystem matching this->Requirements.
+	/** Will gather all archetypes from InEntityManager matching this->Requirements.
 	 *  Note that no work will be done if the cached data is up to date (as tracked by EntitySubsystemHash and 
 	 *	ArchetypeDataVersion properties). */
-	void CacheArchetypes(const UMassEntitySubsystem& InEntitySubsystem);
+	void CacheArchetypes(const FMassEntityManager& InEntityManager);
 
 	void Clear()
 	{
@@ -78,14 +78,14 @@ public:
 	 * Note that the function is not const because calling it can result in re-caching of ValidArchetypes 
 	 * @return the number of entities this given query would process if called "now"
 	 */
-	int32 GetNumMatchingEntities(UMassEntitySubsystem& InEntitySubsystem);
+	int32 GetNumMatchingEntities(FMassEntityManager& InEntityManager);
 
 	/**
 	 * Checks if any of ValidArchetypes has any entities.
 	 * Note that the function is not const because calling it can result in re-caching of ValidArchetypes
 	 * @return "true" if any of the ValidArchetypes has any entities, "false" otherwise
 	 */
-	bool HasMatchingEntities(UMassEntitySubsystem& InEntitySubsystem);
+	bool HasMatchingEntities(FMassEntityManager& InEntityManager);
 
 	/** 
 	 * Sets a chunk filter condition that will applied to each chunk of all valid archetypes. Note 
diff --git a/Engine/Plugins/Runtime/MassEntity/Source/MassEntity/Public/MassEntitySubsystem.h b/Engine/Plugins/Runtime/MassEntity/Source/MassEntity/Public/MassEntitySubsystem.h
--- a/Engine/Plugins/Runtime/MassEntity/Source/MassEntity/Public/MassEntitySubsystem.h
+++ b/Engine/Plugins/Runtime/MassEntity/Source/MassEntity/Public/MassEntitySubsystem.h
@@ -16,7 +16,6 @@
 #include "MassEntitySubsystem.generated.h"
 
 
-class UMassEntitySubsystem;
 struct FMassEntityQuery;
 struct FMassExecutionContext;
 struct FMassArchetypeData;
@@ -27,11 +26,8 @@ struct FMassDebugger;
 enum class EMassFragmentAccess : uint8;
 
 //@TODO: Comment this
-UCLASS()
-class MASSENTITY_API UMassEntitySubsystem : public UWorldSubsystem
+struct MASSENTITY_API FMassEntityManager : public TSharedFromThis<FMassEntityManager>
 {
-	GENERATED_BODY()
-
 	friend FMassEntityQuery;
 	friend FMassDebugger;
 
@@ -73,18 +69,16 @@ public:
 
 	const static FMassEntityHandle InvalidEntity;
 
-	UMassEntitySubsystem();
+	explicit FMassEntityManager(UObject* InOwner = nullptr);
+	FMassEntityManager(const FMassEntityManager& Other) = delete;
 
-	//~UObject interface
-	virtual void GetResourceSizeEx(FResourceSizeEx& CumulativeResourceSize) override;
-	//~End of UObject interface
+	void AddReferencedObjects(FReferenceCollector& Collector);
+	void GetResourceSizeEx(FResourceSizeEx& CumulativeResourceSize);
+
+	void Initialize();
+	void PostInitialize();
+	void Deinitialize();
 
-	//~USubsystem interface
-	virtual void Initialize(FSubsystemCollectionBase& Collection) override;
-	virtual void PostInitialize() override;
-	virtual void Deinitialize() override;
-	//~End of USubsystem interface
-	
 	/** 
 	 * A special, relaxed but slower version of CreateArchetype functions that allows FragmentAngTagsList to contain 
 	 * both fragments and tags. 
@@ -157,7 +151,7 @@ public:
 
 		const FMassArchetypeEntityCollection& GetEntityCollection() const { return EntityCollection; }
 	private:
-		friend UMassEntitySubsystem;
+		friend FMassEntityManager;
 		int32 NumberSpawned;
 		FMassArchetypeEntityCollection EntityCollection;
 		TFunction<void(FEntityCreationContext&)> OnSpawningFinished;
@@ -359,9 +353,16 @@ public:
 
 	FMassObserverManager& GetObserverManager() { return ObserverManager; }
 
+	/** 
+	 * Fetches the world associated with the Owner. 
+	 * @note that it's ok for a given EntityManager to not have an owner or the owner not being part of a UWorld, depending on the use case
+	 */
+	UWorld* GetWorld() const { return Owner.IsValid() ? Owner->GetWorld() : nullptr; }
+	UObject* GetOwner() const { return Owner.Get(); }
+
 #if WITH_MASSENTITY_DEBUG
 	void DebugPrintArchetypes(FOutputDevice& Ar, const bool bIncludeEmpty = true) const;
-	void DebugGetArchetypesStringDetails(FOutputDevice& Ar, const bool bIncludeEmpty = true);
+	void DebugGetArchetypesStringDetails(FOutputDevice& Ar, const bool bIncludeEmpty = true) const;
 	void DebugGetArchetypeFragmentTypes(const FMassArchetypeHandle& Archetype, TArray<const UScriptStruct*>& InOutFragmentList) const;
 	int32 DebugGetArchetypeEntitiesCount(const FMassArchetypeHandle& Archetype) const;
 	int32 DebugGetArchetypeEntitiesCountPerChunk(const FMassArchetypeHandle& Archetype) const;
@@ -425,20 +426,52 @@ private:
 	TMap<const UScriptStruct*, TArray<TSharedPtr<FMassArchetypeData>>> FragmentTypeToArchetypeMap;
 
 	// Shared fragments
-	UPROPERTY(Transient)
 	TArray<FConstSharedStruct> ConstSharedFragments;
 	// Hash/Index in array pair
 	TMap<uint32, int32> ConstSharedFragmentsMap;
 
-	UPROPERTY(Transient)
 	TArray<FSharedStruct> SharedFragments;
 	// Hash/Index in array pair
 	TMap<uint32, int32> SharedFragmentsMap;
 
-	UPROPERTY(Transient)
 	FMassObserverManager ObserverManager;
 
 #if WITH_MASSENTITY_DEBUG
 	FMassRequirementAccessDetector RequirementAccessDetector;
 #endif // WITH_MASSENTITY_DEBUG
+
+	TWeakObjectPtr<UObject> Owner;
+};
+
+
+/** 
+ * The sole responsibility of this world subsystem class is to host the default instance of FMassEntityManager
+ * for a given UWorld. All the gameplay-related use cases of Mass (found in MassGameplay and related plugins) 
+ * use this by default. 
+ */
+UCLASS()
+class MASSENTITY_API UMassEntitySubsystem : public UWorldSubsystem
+{
+	GENERATED_BODY()
+
+	//~USubsystem interface
+	virtual void Initialize(FSubsystemCollectionBase& Collection) override;
+	virtual void PostInitialize() override;
+	virtual void Deinitialize() override;
+	//~End of USubsystem interface
+
+public:
+	UMassEntitySubsystem();
+
+	static void AddReferencedObjects(UObject* InThis, FReferenceCollector& Collector);
+
+	//~UObject interface
+	virtual void GetResourceSizeEx(FResourceSizeEx& CumulativeResourceSize) override;
+	//~End of UObject interface
+
+	const FMassEntityManager& GetEntityManager() const { return EntityManager.Get(); }
+	FMassEntityManager& GetMutableEntityManager() { return EntityManager.Get(); }
+
+protected:
+	TSharedRef<FMassEntityManager> EntityManager;
 };
diff --git a/Engine/Plugins/Runtime/MassEntity/Source/MassEntity/Public/MassEntityTypes.h b/Engine/Plugins/Runtime/MassEntity/Source/MassEntity/Public/MassEntityTypes.h
--- a/Engine/Plugins/Runtime/MassEntity/Source/MassEntity/Public/MassEntityTypes.h
+++ b/Engine/Plugins/Runtime/MassEntity/Source/MassEntity/Public/MassEntityTypes.h
@@ -47,7 +47,7 @@ struct FMassSharedFragment
 	FMassSharedFragment() {}
 };
 
-// A handle to a lightweight entity.  An entity is used in conjunction with the UMassEntitySubsystem
+// A handle to a lightweight entity.  An entity is used in conjunction with the FMassEntityManager
 // for the current world and can contain lightweight fragments.
 USTRUCT()
 struct FMassEntityHandle
diff --git a/Engine/Plugins/Runtime/MassEntity/Source/MassEntity/Public/MassEntityUtils.h b/Engine/Plugins/Runtime/MassEntity/Source/MassEntity/Public/MassEntityUtils.h
--- a/Engine/Plugins/Runtime/MassEntity/Source/MassEntity/Public/MassEntityUtils.h
+++ b/Engine/Plugins/Runtime/MassEntity/Source/MassEntity/Public/MassEntityUtils.h
@@ -6,7 +6,7 @@
 #include "MassArchetypeTypes.h"
 
 class UWorld;
-class UMassEntitySubsystem;
+struct FMassEntityManager;
 struct FMassEntityHandle;
 
 namespace UE::Mass::Utils
@@ -19,7 +19,7 @@ MASSENTITY_API extern EProcessorExecutionFlags GetProcessorExecutionFlagsForWold
  * Fills OutEntityCollections with per-archetype FMassArchetypeEntityCollection instances. 
  * @param DuplicatesHandling used to inform the function whether to expect duplicates.
  */
-MASSENTITY_API extern void CreateEntityCollections(const UMassEntitySubsystem& EntitySystem, const TConstArrayView<FMassEntityHandle> Entities
+MASSENTITY_API extern void CreateEntityCollections(const FMassEntityManager& EntitySubsystem, const TConstArrayView<FMassEntityHandle> Entities
 	, const FMassArchetypeEntityCollection::EDuplicatesHandling DuplicatesHandling, TArray<FMassArchetypeEntityCollection>& OutEntityCollections);
 
 /**
@@ -75,6 +75,8 @@ inline void AbstractSort(const int32 NumElements, TPred&& Predicate, TSwap&& Swa
 	}
 }
 
+MASSENTITY_API extern FMassEntityManager* GetEntityManager(const UWorld* World);
+MASSENTITY_API extern FMassEntityManager& GetEntityManagerChecked(const UWorld& World);
 
 } // namespace UE::Mass::Utils
 
diff --git a/Engine/Plugins/Runtime/MassEntity/Source/MassEntity/Public/MassEntityView.h b/Engine/Plugins/Runtime/MassEntity/Source/MassEntity/Public/MassEntityView.h
--- a/Engine/Plugins/Runtime/MassEntity/Source/MassEntity/Public/MassEntityView.h
+++ b/Engine/Plugins/Runtime/MassEntity/Source/MassEntity/Public/MassEntityView.h
@@ -7,7 +7,7 @@
 #include "MassEntityView.generated.h"
 
 
-class UMassEntitySubsystem;
+struct FMassEntityManager;
 struct FMassArchetypeData;
 struct FMassArchetypeHandle;
 
@@ -25,7 +25,7 @@ struct MASSENTITY_API FMassEntityView
 	/** 
 	 *  Resolves Entity against ArchetypeHandle. Note that this approach requires the caller to ensure that Entity
 	 *  indeed belongs to ArchetypeHandle. If not the call will fail a check. As a remedy calling the 
-	 *  UMassEntitySubsystem-flavored constructor is recommended since it will first find the appropriate archetype for
+	 *  FMassEntityManager-flavored constructor is recommended since it will first find the appropriate archetype for
 	 *  Entity. 
 	 */
 	FMassEntityView(const FMassArchetypeHandle& ArchetypeHandle, FMassEntityHandle Entity);
@@ -34,7 +34,7 @@ struct MASSENTITY_API FMassEntityView
 	 *  Finds the archetype Entity belongs to and then resolves against it. The caller is responsible for ensuring
 	 *  that the given Entity is in fact a valid ID tied to any of the archetypes 
 	 */
-	FMassEntityView(const UMassEntitySubsystem& EntitySubsystem, FMassEntityHandle Entity);
+	FMassEntityView(const FMassEntityManager& EntityManager, FMassEntityHandle Entity);
 
 	FMassEntityHandle GetEntity() const	{ return Entity; }
 
diff --git a/Engine/Plugins/Runtime/MassEntity/Source/MassEntity/Public/MassExecutionContext.h b/Engine/Plugins/Runtime/MassEntity/Source/MassEntity/Public/MassExecutionContext.h
--- a/Engine/Plugins/Runtime/MassEntity/Source/MassEntity/Public/MassExecutionContext.h
+++ b/Engine/Plugins/Runtime/MassEntity/Source/MassEntity/Public/MassExecutionContext.h
@@ -333,7 +333,7 @@ public:
 		return FragmentType != nullptr && FragmentType == TFragment::StaticStruct();
 	}
 
-	void FlushDeferred(UMassEntitySubsystem& EntitySystem) const;
+	void FlushDeferred(FMassEntityManager& EntityManager) const;
 
 	void ClearExecutionData();
 	void SetCurrentArchetypesTagBitSet(const FMassTagBitSet& BitSet)
diff --git a/Engine/Plugins/Runtime/MassEntity/Source/MassEntity/Public/MassObserverManager.h b/Engine/Plugins/Runtime/MassEntity/Source/MassEntity/Public/MassObserverManager.h
--- a/Engine/Plugins/Runtime/MassEntity/Source/MassEntity/Public/MassObserverManager.h
+++ b/Engine/Plugins/Runtime/MassEntity/Source/MassEntity/Public/MassObserverManager.h
@@ -6,7 +6,7 @@
 #include "MassObserverManager.generated.h"
 
 
-class UMassEntitySubsystem;
+struct FMassEntityManager;
 struct FMassArchetypeEntityCollection;
 class UMassProcessor;
 struct FMassProcessingContext;
@@ -103,8 +103,8 @@ public:
 	void AddObserverInstance(const UScriptStruct& FragmentOrTagType, const EMassObservedOperation Operation, UMassProcessor& ObserverProcessor);
 
 protected:
-	friend UMassEntitySubsystem;
-	explicit FMassObserverManager(UMassEntitySubsystem& Owner);
+	friend FMassEntityManager;
+	explicit FMassObserverManager(FMassEntityManager& Owner);
 
 	void Initialize();
 	void HandleFragmentsImpl(FMassProcessingContext& ProcessingContext, const FMassArchetypeEntityCollection& EntityCollection
@@ -121,10 +121,10 @@ protected:
 	FMassObserversMap TagObservers[(uint8)EMassObservedOperation::MAX];
 
 	/** 
-	 * The owning EntitySubsystem. No need for it to be a UPROPERTY since by design we don't support creation of 
-	 * FMassObserverManager outside of an UMassEntitySubsystem instance 
+	 * The owning EntityManager. No need for it to be a UPROPERTY since by design we don't support creation of 
+	 * FMassObserverManager outside of an FMassEntityManager instance 
 	 */
-	UMassEntitySubsystem& EntitySubsystem;
+	FMassEntityManager& EntityManager;
 };
 
 template<>
diff --git a/Engine/Plugins/Runtime/MassEntity/Source/MassEntity/Public/MassObserverProcessor.h b/Engine/Plugins/Runtime/MassEntity/Source/MassEntity/Public/MassObserverProcessor.h
--- a/Engine/Plugins/Runtime/MassEntity/Source/MassEntity/Public/MassObserverProcessor.h
+++ b/Engine/Plugins/Runtime/MassEntity/Source/MassEntity/Public/MassObserverProcessor.h
@@ -6,8 +6,6 @@
 #include "MassObserverProcessor.generated.h"
 
 
-class UMassEntitySubsystem;
-
 UCLASS(abstract)
 class MASSENTITY_API UMassObserverProcessor : public UMassProcessor
 {
diff --git a/Engine/Plugins/Runtime/MassEntity/Source/MassEntity/Public/MassProcessingPhase.h b/Engine/Plugins/Runtime/MassEntity/Source/MassEntity/Public/MassProcessingPhase.h
--- a/Engine/Plugins/Runtime/MassEntity/Source/MassEntity/Public/MassProcessingPhase.h
+++ b/Engine/Plugins/Runtime/MassEntity/Source/MassEntity/Public/MassProcessingPhase.h
@@ -13,7 +13,7 @@
 class UMassProcessingPhaseManager;
 class UMassProcessor;
 class UMassCompositeProcessor;
-class UMassEntitySubsystem;
+struct FMassEntityManager;
 struct FMassCommandBuffer;
 
 USTRUCT()
@@ -48,7 +48,7 @@ protected:
 	UMassCompositeProcessor* PhaseProcessor = nullptr;
 
 	UPROPERTY()
-	UMassProcessingPhaseManager* Manager = nullptr;
+	UMassProcessingPhaseManager* PhaseManager = nullptr;
 
 	EMassProcessingPhase Phase = EMassProcessingPhase::MAX;
 	FOnPhaseEvent OnPhaseStart;
@@ -80,7 +80,7 @@ class MASSENTITY_API UMassProcessingPhaseManager : public UObject
 {
 	GENERATED_BODY()
 public:
-	UMassEntitySubsystem& GetEntitySubsystemRef() { check(EntitySubsystem); return *EntitySubsystem; }
+	FMassEntityManager& GetEntityManagerRef() { check(EntityManager); return *EntityManager.Get(); }
 
 	/** Retrieves OnPhaseStart multicast delegate's reference for a given Phase */
 	FMassProcessingPhase::FOnPhaseEvent& GetOnPhaseStart(const EMassProcessingPhase Phase) { return ProcessingPhases[uint8(Phase)].OnPhaseStart; }
@@ -98,9 +98,9 @@ public:
 	 *  at hand it's suggested to use that Start version, otherwise call the World-using one. 
 	 */
 	void Start(UWorld& World);
-	void Start(UMassEntitySubsystem& InEntitySubsystem);
+	void Start(const TSharedPtr<FMassEntityManager>& InEntityManager);
 	void Stop();
-	bool IsRunning() const { return EntitySubsystem != nullptr; }
+	bool IsRunning() const { return EntityManager.IsValid(); }
 
 	/** 
 	 *  returns true when called while any of the ProcessingPhases is actively executing its processors. Used to 
@@ -146,8 +146,7 @@ protected:
 	UPROPERTY(VisibleAnywhere, Category=Mass)
 	FMassProcessingPhase ProcessingPhases[(uint8)EMassProcessingPhase::MAX];
 
-	UPROPERTY()
-	UMassEntitySubsystem* EntitySubsystem;
+	TSharedPtr<FMassEntityManager> EntityManager;
 
 	EMassProcessingPhase CurrentPhase = EMassProcessingPhase::MAX;
 
diff --git a/Engine/Plugins/Runtime/MassEntity/Source/MassEntity/Public/MassProcessingTypes.h b/Engine/Plugins/Runtime/MassEntity/Source/MassEntity/Public/MassProcessingTypes.h
--- a/Engine/Plugins/Runtime/MassEntity/Source/MassEntity/Public/MassProcessingTypes.h
+++ b/Engine/Plugins/Runtime/MassEntity/Source/MassEntity/Public/MassProcessingTypes.h
@@ -12,7 +12,7 @@
 
 #define WITH_MASSENTITY_DEBUG (!(UE_BUILD_SHIPPING || UE_BUILD_SHIPPING_WITH_EDITOR || UE_BUILD_TEST) && WITH_STRUCTUTILS_DEBUG && 1)
 
-class UMassEntitySubsystem;
+struct FMassEntityManager;
 class UMassProcessor;
 class UMassCompositeProcessor;
 struct FMassCommandBuffer;
@@ -39,8 +39,7 @@ struct MASSENTITY_API FMassProcessingContext
 {
 	GENERATED_BODY()
 
-	UPROPERTY()
-	UMassEntitySubsystem* EntitySubsystem = nullptr;
+	TSharedPtr<FMassEntityManager> EntityManager;
 	
 	UPROPERTY()
 	float DeltaSeconds = 0.f;
@@ -58,7 +57,8 @@ struct MASSENTITY_API FMassProcessingContext
 	TSharedPtr<FMassCommandBuffer> CommandBuffer;
 	
 	FMassProcessingContext() = default;
-	FMassProcessingContext(UMassEntitySubsystem& InEntitySubsystem, const float InDeltaSeconds);
+	FMassProcessingContext(FMassEntityManager& InEntityManager, const float InDeltaSeconds);
+	FMassProcessingContext(TSharedPtr<FMassEntityManager>& InEntityManager, const float InDeltaSeconds);
 	~FMassProcessingContext();
 };
 
diff --git a/Engine/Plugins/Runtime/MassEntity/Source/MassEntity/Public/MassProcessor.h b/Engine/Plugins/Runtime/MassEntity/Source/MassEntity/Public/MassProcessor.h
--- a/Engine/Plugins/Runtime/MassEntity/Source/MassEntity/Public/MassProcessor.h
+++ b/Engine/Plugins/Runtime/MassEntity/Source/MassEntity/Public/MassProcessor.h
@@ -49,13 +49,13 @@ public:
 	UMassProcessor(const FObjectInitializer& ObjectInitializer);
 
 	virtual void Initialize(UObject& Owner) {}
-	virtual FGraphEventRef DispatchProcessorTasks(UMassEntitySubsystem& EntitySubsystem, FMassExecutionContext& ExecutionContext, const FGraphEventArray& Prerequisites = FGraphEventArray());
+	virtual FGraphEventRef DispatchProcessorTasks(const TSharedPtr<FMassEntityManager>& EntityManager, FMassExecutionContext& ExecutionContext, const FGraphEventArray& Prerequisites = FGraphEventArray());
 
 	EProcessorExecutionFlags GetExecutionFlags() const { return (EProcessorExecutionFlags)ExecutionFlags; }
 
 	/** Whether this processor should execute according the CurrentExecutionFlags parameters */
 	bool ShouldExecute(const EProcessorExecutionFlags CurrentExecutionFlags) const { return (GetExecutionFlags() & CurrentExecutionFlags) != EProcessorExecutionFlags::None; }
-	void CallExecute(UMassEntitySubsystem& EntitySubsystem, FMassExecutionContext& Context);
+	void CallExecute(FMassEntityManager& EntityManager, FMassExecutionContext& Context);
 
 	bool AllowDuplicates() const { return bAllowDuplicates; }
 
@@ -99,7 +99,7 @@ public:
 protected:
 	virtual void ConfigureQueries() PURE_VIRTUAL(UMassProcessor::ConfigureQueries);
 	virtual void PostInitProperties() override;
-	virtual void Execute(UMassEntitySubsystem& EntitySubsystem, FMassExecutionContext& Context) PURE_VIRTUAL(UMassProcessor::Execute);
+	virtual void Execute(FMassEntityManager& EntityManager, FMassExecutionContext& Context) PURE_VIRTUAL(UMassProcessor::Execute);
 #if WITH_EDITOR
 	virtual void PostEditChangeProperty(FPropertyChangedEvent& PropertyChangedEvent) override;
 #endif // WITH_EDITOR
@@ -183,7 +183,7 @@ public:
 	 *  and if it's missing it will be created and AddGroupedProcessor will be called recursively */
 	void AddGroupedProcessor(FName RequestedGroupName, UMassProcessor& SubProcessor);
 
-	virtual FGraphEventRef DispatchProcessorTasks(UMassEntitySubsystem& EntitySubsystem, FMassExecutionContext& ExecutionContext, const FGraphEventArray& Prerequisites = FGraphEventArray()) override;
+	virtual FGraphEventRef DispatchProcessorTasks(const TSharedPtr<FMassEntityManager>& EntityManager, FMassExecutionContext& ExecutionContext, const FGraphEventArray& Prerequisites = FGraphEventArray()) override;
 
 	bool IsEmpty() const { return ChildPipeline.IsEmpty(); }
 
@@ -191,7 +191,7 @@ public:
 
 protected:
 	virtual void ConfigureQueries() override;
-	virtual void Execute(UMassEntitySubsystem& EntitySubsystem, FMassExecutionContext& Context) override;
+	virtual void Execute(FMassEntityManager& EntityManager, FMassExecutionContext& Context) override;
 
 	/**
 	 *  Adds processors in OrderedProcessors to ChildPipeline and builds flat processing graph that's being used in multithreaded mode.
diff --git a/Engine/Plugins/Runtime/MassEntity/Source/MassEntityDebugger/Private/MassDebuggerModel.cpp b/Engine/Plugins/Runtime/MassEntity/Source/MassEntityDebugger/Private/MassDebuggerModel.cpp
--- a/Engine/Plugins/Runtime/MassEntity/Source/MassEntityDebugger/Private/MassDebuggerModel.cpp
+++ b/Engine/Plugins/Runtime/MassEntity/Source/MassEntityDebugger/Private/MassDebuggerModel.cpp
@@ -58,13 +58,9 @@ FString FMassDebuggerEnvironment::GetDisplayName() const
 	return World.IsValid() ? World->GetPathName() : TEXT("Invalid");
 }
 
-const UMassEntitySubsystem* FMassDebuggerEnvironment::GetEntitySubsystem() const
+const FMassEntityManager* FMassDebuggerEnvironment::GetEntityManager() const
 {
-	if (World.IsValid())
-	{
-		return World->GetSubsystem<UMassEntitySubsystem>();
-	}
-	return nullptr;
+	return UE::Mass::Utils::GetEntityManager(World.Get());
 }
 
 //----------------------------------------------------------------------//
@@ -116,11 +112,11 @@ FMassDebuggerProcessorData::FMassDebuggerProcessorData(const UMassProcessor& InP
 #endif // WITH_MASSENTITY_DEBUG
 }
 
-FMassDebuggerProcessorData::FMassDebuggerProcessorData(const UMassEntitySubsystem& EntitySubsystem, UMassProcessor& InProcessor, const TMap<FMassArchetypeHandle, TSharedPtr<FMassDebuggerArchetypeData>>& InTransientArchetypesMap)
+FMassDebuggerProcessorData::FMassDebuggerProcessorData(const FMassEntityManager& EntityManager, UMassProcessor& InProcessor, const TMap<FMassArchetypeHandle, TSharedPtr<FMassDebuggerArchetypeData>>& InTransientArchetypesMap)
 {
 	SetProcessor(InProcessor);
 #if WITH_MASSENTITY_DEBUG
-	TConstArrayView<FMassEntityQuery*> ProcessorQueries = FMassDebugger::GetUpToDateProcessorQueries(EntitySubsystem, InProcessor);
+	TConstArrayView<FMassEntityQuery*> ProcessorQueries = FMassDebugger::GetUpToDateProcessorQueries(EntityManager, InProcessor);
 
 	ProcessorRequirements = MakeShareable(new FMassDebuggerQueryData(InProcessor.GetProcessorRequirements(), LOCTEXT("MassProcessorRequirementsLabel", "Processor Requirements")));
 
@@ -388,7 +384,7 @@ void FMassDebuggerModel::CacheProcessorsData(const TMap<FMassArchetypeHandle, TS
 	CachedProcessors.Reset();
 
 	UWorld* World = Environment ? Environment->World.Get() : nullptr;
-	const UMassEntitySubsystem* EntitySubsystem = Environment ? Environment->GetEntitySubsystem() : nullptr;
+	const FMassEntityManager* EntitySubsystem = Environment ? Environment->GetEntityManager() : nullptr;
 
 	if (EntitySubsystem)
 	{
@@ -446,7 +442,7 @@ void FMassDebuggerModel::CacheArchetypesData(TMap<FMassArchetypeHandle, TSharedP
 
 	if (Environment)
 	{
-		if (const UMassEntitySubsystem* EntitySubsystem = Environment->GetEntitySubsystem())
+		if (const FMassEntityManager* EntitySubsystem = Environment->GetEntityManager())
 		{
 			StoreArchetypes(*EntitySubsystem, OutTransientArchetypesMap);
 		}
@@ -463,10 +459,10 @@ float FMassDebuggerModel::MinDistanceToSelectedArchetypes(const TSharedPtr<FMass
 	return MinDistance;
 }
 
-void FMassDebuggerModel::StoreArchetypes(const UMassEntitySubsystem& EntitySubsystem, TMap<FMassArchetypeHandle, TSharedPtr<FMassDebuggerArchetypeData>>& OutTransientArchetypesMap)
+void FMassDebuggerModel::StoreArchetypes(const FMassEntityManager& EntityManager, TMap<FMassArchetypeHandle, TSharedPtr<FMassDebuggerArchetypeData>>& OutTransientArchetypesMap)
 {
 #if WITH_MASSENTITY_DEBUG
-	TArray<FMassArchetypeHandle> ArchetypeHandles = FMassDebugger::GetAllArchetypes(EntitySubsystem);
+	TArray<FMassArchetypeHandle> ArchetypeHandles = FMassDebugger::GetAllArchetypes(EntityManager);
 
 	CachedArchetypes.Reset(ArchetypeHandles.Num());
 
diff --git a/Engine/Plugins/Runtime/MassEntity/Source/MassEntityDebugger/Private/MassDebuggerModel.h b/Engine/Plugins/Runtime/MassEntity/Source/MassEntityDebugger/Private/MassDebuggerModel.h
--- a/Engine/Plugins/Runtime/MassEntity/Source/MassEntityDebugger/Private/MassDebuggerModel.h
+++ b/Engine/Plugins/Runtime/MassEntity/Source/MassEntityDebugger/Private/MassDebuggerModel.h
@@ -13,7 +13,7 @@
 class UMassProcessor;
 struct FMassArchetypeHandle;
 struct FMassDebuggerModel;
-class UMassEntitySubsystem;
+struct FMassEntityManager;
 
 enum class EMassDebuggerSelectionMode : uint8
 {
@@ -100,7 +100,7 @@ struct FMassDebuggerArchetypeData
 struct FMassDebuggerProcessorData
 {
 	FMassDebuggerProcessorData(const UMassProcessor& InProcessor);
-	FMassDebuggerProcessorData(const UMassEntitySubsystem& EntitySubsystem, UMassProcessor& InProcessor, const TMap<FMassArchetypeHandle, TSharedPtr<FMassDebuggerArchetypeData>>& InTransientArchetypesMap);
+	FMassDebuggerProcessorData(const FMassEntityManager& EntitySubsystem, UMassProcessor& InProcessor, const TMap<FMassArchetypeHandle, TSharedPtr<FMassDebuggerArchetypeData>>& InTransientArchetypesMap);
 
 private:
 	void SetProcessor(const UMassProcessor& InProcessor);
@@ -146,7 +146,7 @@ struct FMassDebuggerEnvironment
 	bool operator==(const FMassDebuggerEnvironment& Other) const { return World == Other.World; }
 
 	FString GetDisplayName() const;
-	const UMassEntitySubsystem* GetEntitySubsystem() const;
+	const FMassEntityManager* GetEntityManager() const;
 	bool IsWorldValid() const { return World.IsValid(); }
 	
 	TWeakObjectPtr<UWorld> World;
@@ -186,7 +186,7 @@ struct FMassDebuggerModel
 	const TSharedPtr<FMassDebuggerProcessorData>& GetProcessorDataChecked(const UMassProcessor& Processor) const;
 
 protected:
-	void StoreArchetypes(const UMassEntitySubsystem& EntitySubsystem, TMap<FMassArchetypeHandle, TSharedPtr<FMassDebuggerArchetypeData>>& OutTransientArchetypesMap);
+	void StoreArchetypes(const FMassEntityManager& EntitySubsystem, TMap<FMassArchetypeHandle, TSharedPtr<FMassDebuggerArchetypeData>>& OutTransientArchetypesMap);
 
 	void ResetSelectedArchetypes();
 	void ResetSelectedProcessors();
diff --git a/Engine/Plugins/Runtime/MassEntity/Source/MassEntityTestSuite/Private/LWComponentTestFarmPlot.cpp b/Engine/Plugins/Runtime/MassEntity/Source/MassEntityTestSuite/Private/LWComponentTestFarmPlot.cpp
--- a/Engine/Plugins/Runtime/MassEntity/Source/MassEntityTestSuite/Private/LWComponentTestFarmPlot.cpp
+++ b/Engine/Plugins/Runtime/MassEntity/Source/MassEntityTestSuite/Private/LWComponentTestFarmPlot.cpp
@@ -4,6 +4,7 @@
 
 #include "Components/HierarchicalInstancedStaticMeshComponent.h"
 #include "Engine/CollisionProfile.h"
+#include "MassEntityUtils.h"
 
 //@TODO: Can add a ReadyToHarvest tag Fragment on when things are ready to harvest, to stop them ticking and signal that we need to create an icon
 
@@ -20,11 +21,11 @@ void UFragmentUpdateSystem::PostInitProperties()
 //////////////////////////////////////////////////////////////////////////
 // UFarmHarvestTimerSetIcon
 
-void UFarmHarvestTimerSetIcon::Execute(UMassEntitySubsystem& EntitySubsystem, FMassExecutionContext& Context)
+void UFarmHarvestTimerSetIcon::Execute(FMassEntityManager& EntityManager, FMassExecutionContext& Context)
 {
 	QUICK_SCOPE_CYCLE_COUNTER(SET_ICON_SET_ICON_SET_ICON);
 
-	EntityQuery.ForEachEntityChunk(EntitySubsystem, Context, [this](FMassExecutionContext& Context) {
+	EntityQuery.ForEachEntityChunk(EntityManager, Context, [this](FMassExecutionContext& Context) {
 
 		const int32 NumEntities = Context.GetNumEntities();
 		TConstArrayView<FFarmGridCellData> GridCoordList = Context.GetFragmentView<FFarmGridCellData>();
@@ -64,25 +65,25 @@ ALWFragmentTestFarmPlot::ALWFragmentTestFarmPlot()
 	HarvestIconISMC->SetCollisionProfileName(UCollisionProfile::NoCollision_ProfileName);
 }
 
-void ALWFragmentTestFarmPlot::AddItemToGrid(UMassEntitySubsystem* EntitySystem, uint16 X, uint16 Y, FMassArchetypeHandle Archetype, uint16 VisualIndex)
+void ALWFragmentTestFarmPlot::AddItemToGrid(FMassEntityManager& EntityManager, uint16 X, uint16 Y, FMassArchetypeHandle Archetype, uint16 VisualIndex)
 {
-	FMassEntityHandle NewItem = EntitySystem->CreateEntity(Archetype);
+	FMassEntityHandle NewItem = EntityManager.CreateEntity(Archetype);
 	PlantedSquares[X + Y * GridWidth] = NewItem;
 
-	EntitySystem->GetFragmentDataChecked<FFarmWaterFragment>(NewItem).DeltaWaterPerSecond = FMath::FRandRange(-0.01f, -0.001f);
-	EntitySystem->GetFragmentDataChecked<FHarvestTimerFragment>(NewItem).NumSecondsLeft = 5 + (FMath::Rand() % 100);
+	EntityManager.GetFragmentDataChecked<FFarmWaterFragment>(NewItem).DeltaWaterPerSecond = FMath::FRandRange(-0.01f, -0.001f);
+	EntityManager.GetFragmentDataChecked<FHarvestTimerFragment>(NewItem).NumSecondsLeft = 5 + (FMath::Rand() % 100);
 
 	FFarmGridCellData GridCoords;
 	GridCoords.CellX = X;
 	GridCoords.CellY = Y;
-	EntitySystem->GetFragmentDataChecked<FFarmGridCellData>(NewItem) = GridCoords;
+	EntityManager.GetFragmentDataChecked<FFarmGridCellData>(NewItem) = GridCoords;
 
 	const FVector MeshPosition(X*GridCellWidth, Y*GridCellHeight, 0.0f);
 	const FRotator MeshRotation(0.0f, FMath::FRand()*360.0f, 0.0f);
 	const FVector MeshScale(1.0f, 1.0f, 1.0f); //@TODO: plumb in scale param?
 	const FTransform MeshTransform(MeshRotation, MeshPosition, MeshScale);
 
-	FFarmVisualFragment& VisualComp = EntitySystem->GetFragmentDataChecked<FFarmVisualFragment>(NewItem);
+	FFarmVisualFragment& VisualComp = EntityManager.GetFragmentDataChecked<FFarmVisualFragment>(NewItem);
 	VisualComp.VisualType = VisualIndex;
 	VisualComp.InstanceIndex = VisualDataISMCs[VisualComp.VisualType]->AddInstance(MeshTransform);
 }
@@ -91,10 +92,12 @@ void ALWFragmentTestFarmPlot::BeginPlay()
 {
 	Super::BeginPlay();
 
-	UMassEntitySubsystem* EntitySystem = UWorld::GetSubsystem<UMassEntitySubsystem>(GetWorld());
+	UWorld* World = GetWorld();
+	check(World);
+	FMassEntityManager& EntityManager = UE::Mass::Utils::GetEntityManagerChecked(*World);
 
-	FMassArchetypeHandle CropArchetype = EntitySystem->CreateArchetype(TArray<const UScriptStruct*>{ FFarmWaterFragment::StaticStruct(), FFarmCropFragment::StaticStruct(), FHarvestTimerFragment::StaticStruct(), FFarmVisualFragment::StaticStruct(), FFarmGridCellData::StaticStruct() });
-	FMassArchetypeHandle FlowerArchetype = EntitySystem->CreateArchetype(TArray<const UScriptStruct*>{ FFarmWaterFragment::StaticStruct(), FFarmFlowerFragment::StaticStruct(), FHarvestTimerFragment::StaticStruct(), FFarmVisualFragment::StaticStruct(), FFarmGridCellData::StaticStruct() });
+	FMassArchetypeHandle CropArchetype = EntityManager.CreateArchetype(TArray<const UScriptStruct*>{ FFarmWaterFragment::StaticStruct(), FFarmCropFragment::StaticStruct(), FHarvestTimerFragment::StaticStruct(), FFarmVisualFragment::StaticStruct(), FFarmGridCellData::StaticStruct() });
+	FMassArchetypeHandle FlowerArchetype = EntityManager.CreateArchetype(TArray<const UScriptStruct*>{ FFarmWaterFragment::StaticStruct(), FFarmFlowerFragment::StaticStruct(), FHarvestTimerFragment::StaticStruct(), FFarmVisualFragment::StaticStruct(), FFarmGridCellData::StaticStruct() });
 
 	PerFrameSystems.Add(NewObject<UFarmWaterUpdateSystem>(this));
 
@@ -139,7 +142,7 @@ void ALWFragmentTestFarmPlot::BeginPlay()
 			const bool bIsOdd = ((X + Y) & 1) != 0;
 			const uint16 VisualIndex = bIsOdd ? TestDataCropIndicies[FMath::RandRange(0, TestDataCropIndicies.Num() - 1)] : TestDataFlowerIndicies[FMath::RandRange(0, TestDataFlowerIndicies.Num() - 1)];
 
-			AddItemToGrid(EntitySystem, X, Y, bIsOdd ? CropArchetype : FlowerArchetype, VisualIndex);
+			AddItemToGrid(EntityManager, X, Y, bIsOdd ? CropArchetype : FlowerArchetype, VisualIndex);
 		}
 	}
 }
@@ -150,7 +153,9 @@ void ALWFragmentTestFarmPlot::TickActor(float DeltaTime, enum ELevelTick TickTyp
 
 	QUICK_SCOPE_CYCLE_COUNTER(HeyaTick);
 
-	UMassEntitySubsystem* EntitySystem = UWorld::GetSubsystem<UMassEntitySubsystem>(GetWorld());
+	UWorld* World = GetWorld();
+	check(World);
+	FMassEntityManager& EntityManager = UE::Mass::Utils::GetEntityManagerChecked(*World);
 
 	// Run every frame systems
 	{
@@ -161,7 +166,7 @@ void ALWFragmentTestFarmPlot::TickActor(float DeltaTime, enum ELevelTick TickTyp
 		for (UFragmentUpdateSystem* System : PerFrameSystems)
 		{
 			check(System);
-			System->Execute(*EntitySystem, ExecContext);
+			System->Execute(EntityManager, ExecContext);
 		}
 	}
 
@@ -177,11 +182,11 @@ void ALWFragmentTestFarmPlot::TickActor(float DeltaTime, enum ELevelTick TickTyp
 		for (UFragmentUpdateSystem* System : PerSecondSystems)
 		{
 			check(System);
-			System->Execute(*EntitySystem, ExecContext);
+			System->Execute(EntityManager, ExecContext);
 		}
 	}
 
-#if 0
+#if 1
 	// Update visuals
 	{
 		QUICK_SCOPE_CYCLE_COUNTER(FARM_BASE_MESHES);
@@ -191,9 +196,9 @@ void ALWFragmentTestFarmPlot::TickActor(float DeltaTime, enum ELevelTick TickTyp
 			for (int32 X = 0; X < GridWidth; ++X)
 			{
 				FMassEntityHandle GridEntity = PlantedSquares[X + Y * GridWidth];
-				if (EntitySystem->IsValidEntity(GridEntity))
+				if (EntityManager.IsEntityValid(GridEntity))
 				{
-					FFarmVisualFragment& VisualComp = EntitySystem->GetFragmentDataChecked<FFarmVisualFragment>(GridEntity);
+					FFarmVisualFragment& VisualComp = EntityManager.GetFragmentDataChecked<FFarmVisualFragment>(GridEntity);
 					if (VisualComp.InstanceIndex < 0)
 					{
 						const FVector MeshPosition(X*GridCellWidth, Y*GridCellHeight, 0.0f);
@@ -219,7 +224,7 @@ void ALWFragmentTestFarmPlot::TickActor(float DeltaTime, enum ELevelTick TickTyp
 	if (QQQZZZ == 1)
 	{
 		QUICK_SCOPE_CYCLE_COUNTER(FARM_HARVEST_ICONS);
-		EntitySystem->ForEachFragmentChunk<FHarvestTimerFragment>(
+		EntitySubsystem->ForEachFragmentChunk<FHarvestTimerFragment>(
 			[&](TArrayView<FHarvestTimerFragment> Timers)
 		{
 			for (FHarvestTimerFragment& Timer : Timers)
@@ -243,9 +248,9 @@ void ALWFragmentTestFarmPlot::TickActor(float DeltaTime, enum ELevelTick TickTyp
 			for (int32 X = 0; X < GridWidth; ++X)
 			{
 				FMassEntityHandle GridEntity = PlantedSquares[X + Y * GridWidth];
-				if (EntitySystem->IsValidEntity(GridEntity))
+				if (EntitySubsystem->IsValidEntity(GridEntity))
 				{
-					FHarvestTimerFragment& HarvestTimerData = EntitySystem->GetFragmentDataChecked<FHarvestTimerFragment>(GridEntity);
+					FHarvestTimerFragment& HarvestTimerData = EntitySubsystem->GetFragmentDataChecked<FHarvestTimerFragment>(GridEntity);
 					if ((HarvestTimerData.NumSecondsLeft == 0) && (HarvestTimerData.HarvestIconIndex < 0))
 					{
 						const FVector IconPosition(X*GridCellWidth, Y*GridCellHeight, HarvestIconHeight);
diff --git a/Engine/Plugins/Runtime/MassEntity/Source/MassEntityTestSuite/Private/MassChunkCollectionTest.cpp b/Engine/Plugins/Runtime/MassEntity/Source/MassEntityTestSuite/Private/MassChunkCollectionTest.cpp
--- a/Engine/Plugins/Runtime/MassEntity/Source/MassEntityTestSuite/Private/MassChunkCollectionTest.cpp
+++ b/Engine/Plugins/Runtime/MassEntity/Source/MassEntityTestSuite/Private/MassChunkCollectionTest.cpp
@@ -23,7 +23,7 @@ struct FEntityCollectionTestBase : FEntityTestBase
 		FEntityTestBase::SetUp();
 
 		const int32 Count = 100;
-		EntitySubsystem->BatchCreateEntities(FloatsArchetype, Count, Entities);
+		EntityManager->BatchCreateEntities(FloatsArchetype, Count, Entities);
 		return true;
 	}
 
@@ -107,11 +107,11 @@ struct FEntityCollection_CreateCrossChunk : FEntityTestBase
 	{
 #if WITH_MASSENTITY_DEBUG
 		TArray<FMassEntityHandle> Entities;
-		const int32 EntitiesPerChunk = EntitySubsystem->DebugGetArchetypeEntitiesCountPerChunk(FloatsArchetype);
+		const int32 EntitiesPerChunk = EntityManager->DebugGetArchetypeEntitiesCountPerChunk(FloatsArchetype);
 
 		const int32 SpillOver = 10;
 		const int32 Count = EntitiesPerChunk + SpillOver;
-		EntitySubsystem->BatchCreateEntities(FloatsArchetype, Count, Entities);
+		EntityManager->BatchCreateEntities(FloatsArchetype, Count, Entities);
 		
 		TArray<FMassEntityHandle> EntitiesSubCollection;
 		EntitiesSubCollection.Add(Entities[EntitiesPerChunk]);
@@ -162,8 +162,8 @@ struct FEntityCollection_WithPayload : FEntityCollectionTestBase
 		// skipping FEntityCollectionTestBase::SetUp on purpose to manually create entities
 		FEntityTestBase::SetUp();
 
-		const int32 EntitiesPerChunk = EntitySubsystem->DebugGetArchetypeEntitiesCountPerChunk(FloatsArchetype);
-		EntitySubsystem->BatchCreateEntities(FloatsArchetype, EntitiesPerChunk * 2, Entities);
+		const int32 EntitiesPerChunk = EntityManager->DebugGetArchetypeEntitiesCountPerChunk(FloatsArchetype);
+		EntityManager->BatchCreateEntities(FloatsArchetype, EntitiesPerChunk * 2, Entities);
 		return true;
 	}
 
@@ -195,7 +195,7 @@ struct FEntityCollection_WithPayload : FEntityCollectionTestBase
 		// transform typed payload array into generic one for sorting purposes
 		FStructArrayView PaloadView(Payload);
 		TArray<FMassArchetypeEntityCollectionWithPayload> Result;
-		FMassArchetypeEntityCollectionWithPayload::CreateEntityRangesWithPayload(*EntitySubsystem, EntitiesSubSet, FMassArchetypeEntityCollection::FoldDuplicates
+		FMassArchetypeEntityCollectionWithPayload::CreateEntityRangesWithPayload(*EntityManager, EntitiesSubSet, FMassArchetypeEntityCollection::FoldDuplicates
 			, FMassGenericPayloadView(MakeArrayView(&PaloadView, 1)), Result);
 
 		// at this point Payload should be sorted ascending 
diff --git a/Engine/Plugins/Runtime/MassEntity/Source/MassEntityTestSuite/Private/MassCommandsTest.cpp b/Engine/Plugins/Runtime/MassEntity/Source/MassEntityTestSuite/Private/MassCommandsTest.cpp
--- a/Engine/Plugins/Runtime/MassEntity/Source/MassEntityTestSuite/Private/MassCommandsTest.cpp
+++ b/Engine/Plugins/Runtime/MassEntity/Source/MassEntityTestSuite/Private/MassCommandsTest.cpp
@@ -28,22 +28,22 @@ struct FCommands_FragmentInstanceList : FEntityTestBase
 		const int32 Count = 5;
 		TArray<FMassEntityHandle> IntEntities;
 		TArray<FMassEntityHandle> FloatEntities;
-		EntitySubsystem->BatchCreateEntities(IntsArchetype, Count, IntEntities);
-		EntitySubsystem->BatchCreateEntities(FloatsArchetype, Count, FloatEntities);
+		EntityManager->BatchCreateEntities(IntsArchetype, Count, IntEntities);
+		EntityManager->BatchCreateEntities(FloatsArchetype, Count, FloatEntities);
 
 		for (int i = 0; i < Count; ++i)
 		{
-			EntitySubsystem->Defer().PushCommand<FMassCommandAddFragmentInstances>(IntEntities[i], FTestFragment_Int(i), FTestFragment_Float(i));
-			EntitySubsystem->Defer().PushCommand<FMassCommandAddFragmentInstances>(FloatEntities[i], FTestFragment_Int(i), FTestFragment_Float(i));
+			EntityManager->Defer().PushCommand<FMassCommandAddFragmentInstances>(IntEntities[i], FTestFragment_Int(i), FTestFragment_Float(i));
+			EntityManager->Defer().PushCommand<FMassCommandAddFragmentInstances>(FloatEntities[i], FTestFragment_Int(i), FTestFragment_Float(i));
 		}
 
-		EntitySubsystem->FlushCommands();
+		EntityManager->FlushCommands();
 
 		auto TestEntities = [this](const TArray<FMassEntityHandle>& Entities) -> bool {
 			// all entities should have ended up in the same archetype, FloatsIntsArchetype
 			for (int i = 0; i < Entities.Num(); ++i)
 			{
-				AITEST_EQUAL(TEXT("All entities should have ended up in the same archetype"), EntitySubsystem->GetArchetypeForEntity(Entities[i]), FloatsIntsArchetype);
+				AITEST_EQUAL(TEXT("All entities should have ended up in the same archetype"), EntityManager->GetArchetypeForEntity(Entities[i]), FloatsIntsArchetype);
 
 				FMassEntityView View(FloatsIntsArchetype, Entities[i]);
 				AITEST_EQUAL(TEXT("Should have predicted values"), View.GetFragmentData<FTestFragment_Int>().Value, i);
@@ -72,15 +72,15 @@ struct FCommands_FragmentMemoryCleanup : FEntityTestBase
 			FTestFragment_Array::StaticStruct(), 
 			FTestFragment_Int::StaticStruct()
 		};
-		const FMassArchetypeHandle ArrayArchetype = EntitySubsystem->CreateArchetype(MakeArrayView(ArrayFragmentTypes, 1));
-		const FMassArchetypeHandle ArrayIntArchetype = EntitySubsystem->CreateArchetype(MakeArrayView(ArrayFragmentTypes, 2));
-		const int32 EntitiesPerChunk = EntitySubsystem->DebugGetArchetypeEntitiesCountPerChunk(ArrayArchetype);
+		const FMassArchetypeHandle ArrayArchetype = EntityManager->CreateArchetype(MakeArrayView(ArrayFragmentTypes, 1));
+		const FMassArchetypeHandle ArrayIntArchetype = EntityManager->CreateArchetype(MakeArrayView(ArrayFragmentTypes, 2));
+		const int32 EntitiesPerChunk = EntityManager->DebugGetArchetypeEntitiesCountPerChunk(ArrayArchetype);
 		const int32 Count = EntitiesPerChunk * 2.5f;
 		
 		TArray<FMassEntityHandle> Entities;
-		EntitySubsystem->BatchCreateEntities(ArrayArchetype, Count, Entities);
+		EntityManager->BatchCreateEntities(ArrayArchetype, Count, Entities);
 
-		AITEST_EQUAL(TEXT("All entities created should be in ArrayArchetype"), EntitySubsystem->DebugGetArchetypeEntitiesCount(ArrayArchetype), Entities.Num());
+		AITEST_EQUAL(TEXT("All entities created should be in ArrayArchetype"), EntityManager->DebugGetArchetypeEntitiesCount(ArrayArchetype), Entities.Num());
 
 		TArray<int32> EntitiesModified;
 		for (int i = 0; i < Count; ++i)
@@ -89,13 +89,13 @@ struct FCommands_FragmentMemoryCleanup : FEntityTestBase
 			{	
 				FTestFragment_Array A;
 				A.Value.Add(i);
-				EntitySubsystem->Defer().PushCommand<FMassCommandAddFragmentInstances>(Entities[i], A);
-				EntitySubsystem->Defer().AddFragment<FTestFragment_Int>(Entities[i]);
+				EntityManager->Defer().PushCommand<FMassCommandAddFragmentInstances>(Entities[i], A);
+				EntityManager->Defer().AddFragment<FTestFragment_Int>(Entities[i]);
 				EntitiesModified.Add(i);
 			}
 		}
 
-		EntitySubsystem->FlushCommands();
+		EntityManager->FlushCommands();
 
 		for (int32 i : EntitiesModified)
 		{
@@ -114,19 +114,19 @@ struct FCommands_BuildEntitiesWithFragments : FEntityTestBase
 {
 	virtual bool InstantTest() override
 	{
-		const int32 EntitiesPerChunk = EntitySubsystem->DebugGetArchetypeEntitiesCountPerChunk(FloatsIntsArchetype);
+		const int32 EntitiesPerChunk = EntityManager->DebugGetArchetypeEntitiesCountPerChunk(FloatsIntsArchetype);
 		const int32 Count = EntitiesPerChunk * 2.5f;
 
 		TArray<FMassEntityHandle> Entities;
 		for (int i = 0; i < Count; ++i)
 		{
-			Entities.Add(EntitySubsystem->ReserveEntity());
-			EntitySubsystem->Defer().PushCommand<FMassCommandAddFragmentInstances>(Entities.Last(), FTestFragment_Int(i), FTestFragment_Float(i));
+			Entities.Add(EntityManager->ReserveEntity());
+			EntityManager->Defer().PushCommand<FMassCommandAddFragmentInstances>(Entities.Last(), FTestFragment_Int(i), FTestFragment_Float(i));
 		}
 
-		AITEST_EQUAL(TEXT("All entities created should be in ArrayArchetype"), EntitySubsystem->DebugGetArchetypeEntitiesCount(FloatsIntsArchetype), 0);
-		EntitySubsystem->FlushCommands();
-		AITEST_EQUAL(TEXT("All entities created should be in ArrayArchetype"), EntitySubsystem->DebugGetArchetypeEntitiesCount(FloatsIntsArchetype), Entities.Num());
+		AITEST_EQUAL(TEXT("All entities created should be in ArrayArchetype"), EntityManager->DebugGetArchetypeEntitiesCount(FloatsIntsArchetype), 0);
+		EntityManager->FlushCommands();
+		AITEST_EQUAL(TEXT("All entities created should be in ArrayArchetype"), EntityManager->DebugGetArchetypeEntitiesCount(FloatsIntsArchetype), Entities.Num());
 
 		for (int i = 0; i < Entities.Num(); ++i)
 		{
@@ -144,25 +144,25 @@ struct FCommands_BuildEntitiesInHoles : FEntityTestBase
 {
 	virtual bool InstantTest() override
 	{
-		const int32 EntitiesPerChunk = EntitySubsystem->DebugGetArchetypeEntitiesCountPerChunk(FloatsIntsArchetype);
+		const int32 EntitiesPerChunk = EntityManager->DebugGetArchetypeEntitiesCountPerChunk(FloatsIntsArchetype);
 		const int32 Count = int(EntitiesPerChunk * 1.25f) * 2; // making sure it's even
 
 		TArray<FMassEntityHandle> Entities;
-		EntitySubsystem->BatchCreateEntities(FloatsIntsArchetype, Count, Entities);
+		EntityManager->BatchCreateEntities(FloatsIntsArchetype, Count, Entities);
 		FMath::SRandInit(0);
 		Algo::RandomShuffle(Entities);
-		EntitySubsystem->BatchDestroyEntities(MakeArrayView(Entities.GetData(), Entities.Num()/2));
+		EntityManager->BatchDestroyEntities(MakeArrayView(Entities.GetData(), Entities.Num()/2));
 
 		Entities.Reset();
 		for (int i = 0; i < EntitiesPerChunk; ++i)
 		{
-			Entities.Add(EntitySubsystem->ReserveEntity());
-			EntitySubsystem->Defer().PushCommand<FMassCommandAddFragmentInstances>(Entities.Last(), FTestFragment_Int(i), FTestFragment_Float(i));
+			Entities.Add(EntityManager->ReserveEntity());
+			EntityManager->Defer().PushCommand<FMassCommandAddFragmentInstances>(Entities.Last(), FTestFragment_Int(i), FTestFragment_Float(i));
 		}
 
-		AITEST_EQUAL(TEXT("All entities created should be in ArrayArchetype"), EntitySubsystem->DebugGetArchetypeEntitiesCount(FloatsIntsArchetype), Count / 2);
-		EntitySubsystem->FlushCommands();
-		AITEST_EQUAL(TEXT("All entities created should be in ArrayArchetype"), EntitySubsystem->DebugGetArchetypeEntitiesCount(FloatsIntsArchetype), Count / 2 + Entities.Num());
+		AITEST_EQUAL(TEXT("All entities created should be in ArrayArchetype"), EntityManager->DebugGetArchetypeEntitiesCount(FloatsIntsArchetype), Count / 2);
+		EntityManager->FlushCommands();
+		AITEST_EQUAL(TEXT("All entities created should be in ArrayArchetype"), EntityManager->DebugGetArchetypeEntitiesCount(FloatsIntsArchetype), Count / 2 + Entities.Num());
 
 		for (int i = 0; i < Entities.Num(); ++i)
 		{
@@ -180,19 +180,19 @@ struct FCommands_BuildEntitiesWithFragmentInstances : FEntityTestBase
 {
 	virtual bool InstantTest() override
 	{
-		const int32 EntitiesPerChunk = EntitySubsystem->DebugGetArchetypeEntitiesCountPerChunk(FloatsIntsArchetype);
+		const int32 EntitiesPerChunk = EntityManager->DebugGetArchetypeEntitiesCountPerChunk(FloatsIntsArchetype);
 		const int32 Count = EntitiesPerChunk * 2.5f;
 
 		TArray<FMassEntityHandle> Entities;
 		for (int i = 0; i < Count; ++i)
 		{
-			Entities.Add(EntitySubsystem->ReserveEntity());
-			EntitySubsystem->Defer().PushCommand<FMassCommandBuildEntity>(Entities.Last(), FTestFragment_Int(i), FTestFragment_Float(i));
+			Entities.Add(EntityManager->ReserveEntity());
+			EntityManager->Defer().PushCommand<FMassCommandBuildEntity>(Entities.Last(), FTestFragment_Int(i), FTestFragment_Float(i));
 		}
 
-		AITEST_EQUAL(TEXT("All entities created should be in ArrayArchetype"), EntitySubsystem->DebugGetArchetypeEntitiesCount(FloatsIntsArchetype), 0);
-		EntitySubsystem->FlushCommands();
-		AITEST_EQUAL(TEXT("All entities created should be in ArrayArchetype"), EntitySubsystem->DebugGetArchetypeEntitiesCount(FloatsIntsArchetype), Entities.Num());
+		AITEST_EQUAL(TEXT("All entities created should be in ArrayArchetype"), EntityManager->DebugGetArchetypeEntitiesCount(FloatsIntsArchetype), 0);
+		EntityManager->FlushCommands();
+		AITEST_EQUAL(TEXT("All entities created should be in ArrayArchetype"), EntityManager->DebugGetArchetypeEntitiesCount(FloatsIntsArchetype), Entities.Num());
 
 		for (int i = 0; i < Entities.Num(); ++i)
 		{
@@ -214,7 +214,7 @@ struct FCommands_DeferredFunction : FEntityTestBase
 		const int32 Offset = 1000;
 
 		TArray<FMassEntityHandle> Entities;
-		EntitySubsystem->BatchCreateEntities(IntsArchetype, Count, Entities);
+		EntityManager->BatchCreateEntities(IntsArchetype, Count, Entities);
 
 		int i = 0;
 		for (FMassEntityHandle Entity : Entities)
@@ -222,14 +222,14 @@ struct FCommands_DeferredFunction : FEntityTestBase
 			FMassEntityView View(IntsArchetype, Entity);
 			View.GetFragmentData<FTestFragment_Int>().Value = Offset + i++;
 
-			EntitySubsystem->Defer().PushCommand<FMassDeferredSetCommand>([Entity, Archetype = IntsArchetype, Offset](UMassEntitySubsystem& System)
+			EntityManager->Defer().PushCommand<FMassDeferredSetCommand>([Entity, Archetype = IntsArchetype, Offset](FMassEntityManager& System)
 				{
 					FMassEntityView View(Archetype, Entity);
 					View.GetFragmentData<FTestFragment_Int>().Value -= Offset;
 				});
 		}
 
-		EntitySubsystem->FlushCommands();
+		EntityManager->FlushCommands();
 
 		for (i = 0; i < Entities.Num(); ++i)
 		{
diff --git a/Engine/Plugins/Runtime/MassEntity/Source/MassEntityTestSuite/Private/MassCompositeProcessorTest.cpp b/Engine/Plugins/Runtime/MassEntity/Source/MassEntityTestSuite/Private/MassCompositeProcessorTest.cpp
--- a/Engine/Plugins/Runtime/MassEntity/Source/MassEntityTestSuite/Private/MassCompositeProcessorTest.cpp
+++ b/Engine/Plugins/Runtime/MassEntity/Source/MassEntityTestSuite/Private/MassCompositeProcessorTest.cpp
@@ -20,11 +20,11 @@ struct FCompositeProcessorTest_Empty : FEntityTestBase
 {
 	virtual bool InstantTest() override
 	{
-		CA_ASSUME(EntitySubsystem);
+		CA_ASSUME(EntityManager);
 
-		UMassCompositeProcessor* CompositeProcessor = NewObject<UMassCompositeProcessor>(EntitySubsystem);
+		UMassCompositeProcessor* CompositeProcessor = NewObject<UMassCompositeProcessor>();
 		check(CompositeProcessor);
-		FMassProcessingContext ProcessingContext(*EntitySubsystem, /*DeltaSeconds=*/0.f);
+		FMassProcessingContext ProcessingContext(EntityManager, /*DeltaSeconds=*/0.f);
 		// it should just run, no warnings
 		UE::Mass::Executor::Run(*CompositeProcessor, ProcessingContext);
 		return true;
@@ -37,9 +37,9 @@ struct FCompositeProcessorTest_MultipleSubProcessors : FEntityTestBase
 {
 	virtual bool InstantTest() override
 	{
-		CA_ASSUME(EntitySubsystem);
+		CA_ASSUME(EntityManager);
 
-		UMassCompositeProcessor* CompositeProcessor = NewObject<UMassCompositeProcessor>(EntitySubsystem);
+		UMassCompositeProcessor* CompositeProcessor = NewObject<UMassCompositeProcessor>();
 		check(CompositeProcessor);
 
 		int Result = 0;
@@ -47,8 +47,8 @@ struct FCompositeProcessorTest_MultipleSubProcessors : FEntityTestBase
 			TArray<UMassProcessor*> Processors;
 			for (int i = 0; i < 3; ++i)
 			{
-				UMassTestProcessorBase* Processor = NewObject<UMassTestProcessorBase>(EntitySubsystem);
-				Processor->ExecutionFunction = [Processor, &Result, i](UMassEntitySubsystem& InEntitySubsystem, FMassExecutionContext& Context) {
+				UMassTestProcessorBase* Processor = NewObject<UMassTestProcessorBase>();
+				Processor->ExecutionFunction = [Processor, &Result, i](FMassEntityManager& InEntitySubsystem, FMassExecutionContext& Context) {
 						check(Processor);
 						Result += (int)FMath::Pow(10.f, float(i));
 					};
@@ -58,7 +58,7 @@ struct FCompositeProcessorTest_MultipleSubProcessors : FEntityTestBase
 			CompositeProcessor->SetChildProcessors(MoveTemp(Processors));
 		}
 
-		FMassProcessingContext ProcessingContext(*EntitySubsystem, /*DeltaSeconds=*/0.f);
+		FMassProcessingContext ProcessingContext(EntityManager, /*DeltaSeconds=*/0.f);
 		UE::Mass::Executor::Run(*CompositeProcessor, ProcessingContext);
 		AITEST_EQUAL("All of the child processors should get run", Result, 111);
 		return true;
diff --git a/Engine/Plugins/Runtime/MassEntity/Source/MassEntityTestSuite/Private/MassEntityTest.cpp b/Engine/Plugins/Runtime/MassEntity/Source/MassEntityTestSuite/Private/MassEntityTest.cpp
--- a/Engine/Plugins/Runtime/MassEntity/Source/MassEntityTestSuite/Private/MassEntityTest.cpp
+++ b/Engine/Plugins/Runtime/MassEntity/Source/MassEntityTestSuite/Private/MassEntityTest.cpp
@@ -25,17 +25,17 @@ struct FEntityTest_ArchetypeCreation : FEntityTestBase
 		AITEST_TRUE("Ints archetype should have been created", IntsArchetype.IsValid());
 
 		TArray<const UScriptStruct*> FragmentsList;
-		EntitySubsystem->DebugGetArchetypeFragmentTypes(FloatsArchetype, FragmentsList);
+		EntityManager->DebugGetArchetypeFragmentTypes(FloatsArchetype, FragmentsList);
 		AITEST_EQUAL("Floats archetype should contain just a single fragment", FragmentsList.Num(), 1);
 		AITEST_EQUAL("Floats archetype\'s lone fragment should be of Float fragment type", FragmentsList[0], FTestFragment_Float::StaticStruct());
 
 		FragmentsList.Reset();
-		EntitySubsystem->DebugGetArchetypeFragmentTypes(IntsArchetype, FragmentsList);
+		EntityManager->DebugGetArchetypeFragmentTypes(IntsArchetype, FragmentsList);
 		AITEST_EQUAL("Ints archetype should contain just a single fragment", FragmentsList.Num(), 1);
 		AITEST_EQUAL("Ints archetype\'s lone fragment should be of Ints fragment type", FragmentsList[0], FTestFragment_Int::StaticStruct());
 
 		FragmentsList.Reset();
-		EntitySubsystem->DebugGetArchetypeFragmentTypes(FloatsIntsArchetype, FragmentsList);
+		EntityManager->DebugGetArchetypeFragmentTypes(FloatsIntsArchetype, FragmentsList);
 		AITEST_EQUAL("FloatsInts archetype should contain exactly two fragments", FragmentsList.Num(), 2);
 		AITEST_TRUE("FloatsInts archetype\'s should contain both expected fragment types", FragmentsList.Find(FTestFragment_Int::StaticStruct()) != INDEX_NONE && FragmentsList.Find(FTestFragment_Float::StaticStruct()) != INDEX_NONE);
 		
@@ -51,8 +51,8 @@ struct FEntityTest_ArchetypeEquivalence : FEntityTestBase
 	{
 		TArray<const UScriptStruct*> FragmentsA = { FTestFragment_Float::StaticStruct(), FTestFragment_Int::StaticStruct() };
 		TArray<const UScriptStruct*> FragmentsB = { FTestFragment_Int::StaticStruct(), FTestFragment_Float::StaticStruct() };
-		const FMassArchetypeHandle ArchetypeA = EntitySubsystem->CreateArchetype(FragmentsA);
-		const FMassArchetypeHandle ArchetypeB = EntitySubsystem->CreateArchetype(FragmentsB);
+		const FMassArchetypeHandle ArchetypeA = EntityManager->CreateArchetype(FragmentsA);
+		const FMassArchetypeHandle ArchetypeB = EntityManager->CreateArchetype(FragmentsB);
 
 		AITEST_EQUAL("Archetype creation is expected to be independent of fragments ordering", ArchetypeA, ArchetypeB);
 
@@ -66,7 +66,7 @@ struct FEntityTest_MultipleEntitiesCreation : FEntityTestBase
 {
 	virtual bool InstantTest() override
 	{
-		CA_ASSUME(EntitySubsystem);
+		CA_ASSUME(EntityManager);
 		int32 Counts[] = { 10, 100, 1000 };
 		int32 TotalCreatedCount = 0;
 		FMassArchetypeHandle Archetypes[] = { FloatsArchetype, IntsArchetype, FloatsIntsArchetype };
@@ -75,14 +75,14 @@ struct FEntityTest_MultipleEntitiesCreation : FEntityTestBase
 		{
 			for (int i = 0; i < Counts[ArchetypeIndex]; ++i)
 			{
-				EntitySubsystem->CreateEntity(Archetypes[ArchetypeIndex]);
+				EntityManager->CreateEntity(Archetypes[ArchetypeIndex]);
 			}
 			TotalCreatedCount += Counts[ArchetypeIndex];
 		}
-		AITEST_EQUAL("The total number of entities must match the number created", EntitySubsystem->DebugGetEntityCount(), TotalCreatedCount);
-		AITEST_EQUAL("10 entities of FloatsArchetype should have been created", EntitySubsystem->DebugGetArchetypeEntitiesCount(FloatsArchetype), 10);
-		AITEST_EQUAL("100 entities of IntsArchetype should have been created", EntitySubsystem->DebugGetArchetypeEntitiesCount(IntsArchetype), 100);
-		AITEST_EQUAL("1000 entities of FloatsIntsArchetype should have been created", EntitySubsystem->DebugGetArchetypeEntitiesCount(FloatsIntsArchetype), 1000);
+		AITEST_EQUAL("The total number of entities must match the number created", EntityManager->DebugGetEntityCount(), TotalCreatedCount);
+		AITEST_EQUAL("10 entities of FloatsArchetype should have been created", EntityManager->DebugGetArchetypeEntitiesCount(FloatsArchetype), 10);
+		AITEST_EQUAL("100 entities of IntsArchetype should have been created", EntityManager->DebugGetArchetypeEntitiesCount(IntsArchetype), 100);
+		AITEST_EQUAL("1000 entities of FloatsIntsArchetype should have been created", EntityManager->DebugGetArchetypeEntitiesCount(FloatsIntsArchetype), 1000);
 
 		return true;
 	}
@@ -94,12 +94,12 @@ struct FEntityTest_EntityBatchCreation : FEntityTestBase
 {
 	virtual bool InstantTest() override
 	{
-		CA_ASSUME(EntitySubsystem);
+		CA_ASSUME(EntityManager);
 		const int32 Count = 123;
 		TArray<FMassEntityHandle> Entities;
-		EntitySubsystem->BatchCreateEntities(FloatsIntsArchetype, Count, Entities);
+		EntityManager->BatchCreateEntities(FloatsIntsArchetype, Count, Entities);
 		AITEST_EQUAL("Batch creation should create the expected number of entities", Entities.Num(), Count);
-		AITEST_EQUAL("The total number of entities present must match the number requested", EntitySubsystem->DebugGetEntityCount(), Count);
+		AITEST_EQUAL("The total number of entities present must match the number requested", EntityManager->DebugGetEntityCount(), Count);
 
 		return true;
 	}
@@ -110,11 +110,11 @@ struct FEntityTest_BatchCreatingSingleEntity : FEntityTestBase
 {
 	virtual bool InstantTest() override
 	{
-		CA_ASSUME(EntitySubsystem);
+		CA_ASSUME(EntityManager);
 		TArray<FMassEntityHandle> Entities;
-		EntitySubsystem->BatchCreateEntities(FloatsIntsArchetype, /*Count=*/1, Entities);
+		EntityManager->BatchCreateEntities(FloatsIntsArchetype, /*Count=*/1, Entities);
 		AITEST_EQUAL("Batch creation should have created a single entity", Entities.Num(), 1);
-		AITEST_EQUAL("The total number of entities present must match the number created by batch creation", EntitySubsystem->DebugGetEntityCount(), 1);
+		AITEST_EQUAL("The total number of entities present must match the number created by batch creation", EntityManager->DebugGetEntityCount(), 1);
 
 		return true;
 	}
@@ -126,12 +126,12 @@ struct FEntityTest_EntityCreation : FEntityTestBase
 {
 	virtual bool InstantTest() override
 	{
-		CA_ASSUME(EntitySubsystem);		
-		const FMassEntityHandle Entity = EntitySubsystem->CreateEntity(FloatsArchetype);
-		AITEST_EQUAL("There should be one entity across the whole system", EntitySubsystem->DebugGetEntityCount(), 1);
-		AITEST_EQUAL("Entity\'s archetype should be the Float one", EntitySubsystem->GetArchetypeForEntity(Entity), FloatsArchetype);
-		AITEST_EQUAL("The created entity should have been added to the Floats archetype", EntitySubsystem->DebugGetArchetypeEntitiesCount(FloatsArchetype), 1);
-		AITEST_EQUAL("Other archetypes should not get affected", EntitySubsystem->DebugGetArchetypeEntitiesCount(IntsArchetype) + EntitySubsystem->DebugGetArchetypeEntitiesCount(FloatsIntsArchetype), 0);
+		CA_ASSUME(EntityManager);		
+		const FMassEntityHandle Entity = EntityManager->CreateEntity(FloatsArchetype);
+		AITEST_EQUAL("There should be one entity across the whole system", EntityManager->DebugGetEntityCount(), 1);
+		AITEST_EQUAL("Entity\'s archetype should be the Float one", EntityManager->GetArchetypeForEntity(Entity), FloatsArchetype);
+		AITEST_EQUAL("The created entity should have been added to the Floats archetype", EntityManager->DebugGetArchetypeEntitiesCount(FloatsArchetype), 1);
+		AITEST_EQUAL("Other archetypes should not get affected", EntityManager->DebugGetArchetypeEntitiesCount(IntsArchetype) + EntityManager->DebugGetArchetypeEntitiesCount(FloatsIntsArchetype), 0);
 
 		return true;
 	}
@@ -143,12 +143,12 @@ struct FEntityTest_EntityCreationFromInstances : FEntityTestBase
 {
 	virtual bool InstantTest() override
 	{
-		CA_ASSUME(EntitySubsystem);
-		const FMassEntityHandle Entity = EntitySubsystem->CreateEntity(MakeArrayView(&InstanceInt, 1));
-		AITEST_EQUAL("There should be one entity across the whole system", EntitySubsystem->DebugGetEntityCount(), 1);
-		AITEST_EQUAL("Entity\'s archetype should be the Ints one", EntitySubsystem->GetArchetypeForEntity(Entity), IntsArchetype);
-		AITEST_EQUAL("The created entity should have been added to the Ints archetype", EntitySubsystem->DebugGetArchetypeEntitiesCount(IntsArchetype), 1);
-		AITEST_EQUAL("The entity should have the new component with the correct value set", EntitySubsystem->GetFragmentDataChecked<FTestFragment_Int>(Entity).Value, TestIntValue);
+		CA_ASSUME(EntityManager);
+		const FMassEntityHandle Entity = EntityManager->CreateEntity(MakeArrayView(&InstanceInt, 1));
+		AITEST_EQUAL("There should be one entity across the whole system", EntityManager->DebugGetEntityCount(), 1);
+		AITEST_EQUAL("Entity\'s archetype should be the Ints one", EntityManager->GetArchetypeForEntity(Entity), IntsArchetype);
+		AITEST_EQUAL("The created entity should have been added to the Ints archetype", EntityManager->DebugGetArchetypeEntitiesCount(IntsArchetype), 1);
+		AITEST_EQUAL("The entity should have the new component with the correct value set", EntityManager->GetFragmentDataChecked<FTestFragment_Int>(Entity).Value, TestIntValue);
 
 		return true;
 	}
@@ -176,15 +176,15 @@ struct FEntityTest_AddingFragmentType : FEntityTestBase
 {
 	virtual bool InstantTest() override
 	{
-		CA_ASSUME(EntitySubsystem);
-		const FMassEntityHandle Entity = EntitySubsystem->CreateEntity(FloatsArchetype);
-		EntitySubsystem->AddFragmentToEntity(Entity, FTestFragment_Int::StaticStruct());
-		AITEST_EQUAL("There should be one entity across the whole system", EntitySubsystem->DebugGetEntityCount(), 1);
-		AITEST_EQUAL("The original archetype should now have no entities", EntitySubsystem->DebugGetArchetypeEntitiesCount(FloatsArchetype), 0);
-		AITEST_EQUAL("The destination archetype should now store a single entity", EntitySubsystem->DebugGetArchetypeEntitiesCount(FloatsIntsArchetype), 1);
-		AITEST_EQUAL("The remaining archetype should not be affected", EntitySubsystem->DebugGetArchetypeEntitiesCount(IntsArchetype), 0);
+		CA_ASSUME(EntityManager);
+		const FMassEntityHandle Entity = EntityManager->CreateEntity(FloatsArchetype);
+		EntityManager->AddFragmentToEntity(Entity, FTestFragment_Int::StaticStruct());
+		AITEST_EQUAL("There should be one entity across the whole system", EntityManager->DebugGetEntityCount(), 1);
+		AITEST_EQUAL("The original archetype should now have no entities", EntityManager->DebugGetArchetypeEntitiesCount(FloatsArchetype), 0);
+		AITEST_EQUAL("The destination archetype should now store a single entity", EntityManager->DebugGetArchetypeEntitiesCount(FloatsIntsArchetype), 1);
+		AITEST_EQUAL("The remaining archetype should not be affected", EntityManager->DebugGetArchetypeEntitiesCount(IntsArchetype), 0);
 		// this test was originally failing due to FEntityData.CurrentArchetype not getting updated during entity moving between archetypes
-		AITEST_EQUAL("The entity should get associated with the new archetype", EntitySubsystem->GetArchetypeForEntity(Entity), FloatsIntsArchetype);
+		AITEST_EQUAL("The entity should get associated with the new archetype", EntityManager->GetArchetypeForEntity(Entity), FloatsIntsArchetype);
 
 		return true;
 	}
@@ -195,16 +195,16 @@ struct FEntityTest_AddingFragmentInstance : FEntityTestBase
 {
 	virtual bool InstantTest() override
 	{
-		CA_ASSUME(EntitySubsystem);
-		const FMassEntityHandle Entity = EntitySubsystem->CreateEntity(FloatsArchetype);
-		EntitySubsystem->AddFragmentInstanceListToEntity(Entity, MakeArrayView(&InstanceInt, 1));
-		AITEST_EQUAL("There should be one entity across the whole system", EntitySubsystem->DebugGetEntityCount(), 1);
-		AITEST_EQUAL("The original archetype should now have no entities", EntitySubsystem->DebugGetArchetypeEntitiesCount(FloatsArchetype), 0);
-		AITEST_EQUAL("The destination archetype should now store a single entity", EntitySubsystem->DebugGetArchetypeEntitiesCount(FloatsIntsArchetype), 1);
-		AITEST_EQUAL("The archetype containing just the new fragment should not be affected", EntitySubsystem->DebugGetArchetypeEntitiesCount(IntsArchetype), 0);
+		CA_ASSUME(EntityManager);
+		const FMassEntityHandle Entity = EntityManager->CreateEntity(FloatsArchetype);
+		EntityManager->AddFragmentInstanceListToEntity(Entity, MakeArrayView(&InstanceInt, 1));
+		AITEST_EQUAL("There should be one entity across the whole system", EntityManager->DebugGetEntityCount(), 1);
+		AITEST_EQUAL("The original archetype should now have no entities", EntityManager->DebugGetArchetypeEntitiesCount(FloatsArchetype), 0);
+		AITEST_EQUAL("The destination archetype should now store a single entity", EntityManager->DebugGetArchetypeEntitiesCount(FloatsIntsArchetype), 1);
+		AITEST_EQUAL("The archetype containing just the new fragment should not be affected", EntityManager->DebugGetArchetypeEntitiesCount(IntsArchetype), 0);
 		// this test was originally failing due to FEntityData.CurrentArchetype not getting updated during entity moving between archetypes
-		AITEST_EQUAL("The entity should get associated with the new archetype", EntitySubsystem->GetArchetypeForEntity(Entity), FloatsIntsArchetype);
-		AITEST_EQUAL("The entity should have the new component with the correct value set", EntitySubsystem->GetFragmentDataChecked<FTestFragment_Int>(Entity).Value, TestIntValue);
+		AITEST_EQUAL("The entity should get associated with the new archetype", EntityManager->GetArchetypeForEntity(Entity), FloatsIntsArchetype);
+		AITEST_EQUAL("The entity should have the new component with the correct value set", EntityManager->GetFragmentDataChecked<FTestFragment_Int>(Entity).Value, TestIntValue);
 
 		return true;
 	}
@@ -216,15 +216,15 @@ struct FEntityTest_RemovingFragment : FEntityTestBase
 {
 	virtual bool InstantTest() override
 	{
-		CA_ASSUME(EntitySubsystem);
-		const FMassEntityHandle Entity = EntitySubsystem->CreateEntity(FloatsIntsArchetype);
-		EntitySubsystem->RemoveFragmentFromEntity(Entity, FTestFragment_Float::StaticStruct());
-		AITEST_EQUAL("There should be just one entity across the whole system", EntitySubsystem->DebugGetEntityCount(), 1);
-		AITEST_EQUAL("The original archetype should now have no entities", EntitySubsystem->DebugGetArchetypeEntitiesCount(FloatsIntsArchetype), 0);
-		AITEST_EQUAL("The destination archetype should now store a single entity", EntitySubsystem->DebugGetArchetypeEntitiesCount(IntsArchetype), 1);
-		AITEST_EQUAL("The remaining archetype should not be affected", EntitySubsystem->DebugGetArchetypeEntitiesCount(FloatsArchetype), 0);
+		CA_ASSUME(EntityManager);
+		const FMassEntityHandle Entity = EntityManager->CreateEntity(FloatsIntsArchetype);
+		EntityManager->RemoveFragmentFromEntity(Entity, FTestFragment_Float::StaticStruct());
+		AITEST_EQUAL("There should be just one entity across the whole system", EntityManager->DebugGetEntityCount(), 1);
+		AITEST_EQUAL("The original archetype should now have no entities", EntityManager->DebugGetArchetypeEntitiesCount(FloatsIntsArchetype), 0);
+		AITEST_EQUAL("The destination archetype should now store a single entity", EntityManager->DebugGetArchetypeEntitiesCount(IntsArchetype), 1);
+		AITEST_EQUAL("The remaining archetype should not be affected", EntityManager->DebugGetArchetypeEntitiesCount(FloatsArchetype), 0);
 		// this test was originally failing due to FEntityData.CurrentArchetype not getting updated during entity moving between archetypes
-		AITEST_EQUAL("The entity should get associated with the new archetype", EntitySubsystem->GetArchetypeForEntity(Entity), IntsArchetype);
+		AITEST_EQUAL("The entity should get associated with the new archetype", EntityManager->GetArchetypeForEntity(Entity), IntsArchetype);
 
 		return true;
 	}
@@ -235,13 +235,13 @@ struct FEntityTest_RemovingLastFragment : FEntityTestBase
 {
 	virtual bool InstantTest() override
 	{
-		CA_ASSUME(EntitySubsystem);
-		const FMassEntityHandle Entity = EntitySubsystem->CreateEntity(FloatsArchetype);
-		EntitySubsystem->RemoveFragmentFromEntity(Entity, FTestFragment_Float::StaticStruct());
-		AITEST_EQUAL("There should be one entity across the whole system", EntitySubsystem->DebugGetEntityCount(), 1);
-		AITEST_EQUAL("The original archetype should now have no entities", EntitySubsystem->DebugGetArchetypeEntitiesCount(FloatsArchetype), 0);
+		CA_ASSUME(EntityManager);
+		const FMassEntityHandle Entity = EntityManager->CreateEntity(FloatsArchetype);
+		EntityManager->RemoveFragmentFromEntity(Entity, FTestFragment_Float::StaticStruct());
+		AITEST_EQUAL("There should be one entity across the whole system", EntityManager->DebugGetEntityCount(), 1);
+		AITEST_EQUAL("The original archetype should now have no entities", EntityManager->DebugGetArchetypeEntitiesCount(FloatsArchetype), 0);
 		// this test was originally failing due to FEntityData.CurrentArchetype not getting updated during entity moving between archetypes
-		AITEST_EQUAL("The entity should not get associated to any archetype", EntitySubsystem->GetArchetypeForEntity(Entity), EmptyArchetype);
+		AITEST_EQUAL("The entity should not get associated to any archetype", EntityManager->GetArchetypeForEntity(Entity), EmptyArchetype);
 
 		return true;
 	}
@@ -252,12 +252,12 @@ struct FEntityTest_DestroyEntity : FEntityTestBase
 {
 	virtual bool InstantTest() override
 	{
-		CA_ASSUME(EntitySubsystem);
-		const FMassEntityHandle Entity = EntitySubsystem->CreateEntity(FloatsArchetype);
-		AITEST_EQUAL("The entity should get associated to the right archetype", EntitySubsystem->GetArchetypeForEntity(Entity), FloatsArchetype);
-		EntitySubsystem->DestroyEntity(Entity);
-		AITEST_EQUAL("There should not be any entity across the whole system", EntitySubsystem->DebugGetEntityCount(), 0);
-		AITEST_EQUAL("The original archetype should now have no entities", EntitySubsystem->DebugGetArchetypeEntitiesCount(FloatsArchetype), 0);
+		CA_ASSUME(EntityManager);
+		const FMassEntityHandle Entity = EntityManager->CreateEntity(FloatsArchetype);
+		AITEST_EQUAL("The entity should get associated to the right archetype", EntityManager->GetArchetypeForEntity(Entity), FloatsArchetype);
+		EntityManager->DestroyEntity(Entity);
+		AITEST_EQUAL("There should not be any entity across the whole system", EntityManager->DebugGetEntityCount(), 0);
+		AITEST_EQUAL("The original archetype should now have no entities", EntityManager->DebugGetArchetypeEntitiesCount(FloatsArchetype), 0);
 		return true;
 	}
 };
@@ -267,19 +267,19 @@ struct FEntityTest_EntityReservationAndBuilding : FEntityTestBase
 {
 	virtual bool InstantTest() override
 	{
-		CA_ASSUME(EntitySubsystem);
-		const FMassEntityHandle ReservedEntity = EntitySubsystem->ReserveEntity();
-		AITEST_TRUE("The reserved entity should be a valid entity", EntitySubsystem->IsEntityValid(ReservedEntity));
-		AITEST_FALSE("The reserved entity should not be a valid entity", EntitySubsystem->IsEntityBuilt(ReservedEntity));
-		EntitySubsystem->BuildEntity(ReservedEntity, FloatsArchetype);
-		AITEST_TRUE("The reserved entity should be a valid entity", EntitySubsystem->IsEntityBuilt(ReservedEntity));
-		AITEST_EQUAL("There should be one entity across the whole system", EntitySubsystem->DebugGetEntityCount(), 1);
-		AITEST_EQUAL("Entity\'s archetype should be the Float one", EntitySubsystem->GetArchetypeForEntity(ReservedEntity), FloatsArchetype);
-		AITEST_EQUAL("The created entity should have been added to the Floats archetype", EntitySubsystem->DebugGetArchetypeEntitiesCount(FloatsArchetype), 1);
-		AITEST_EQUAL("Other archetypes should not get affected", EntitySubsystem->DebugGetArchetypeEntitiesCount(IntsArchetype) + EntitySubsystem->DebugGetArchetypeEntitiesCount(FloatsIntsArchetype), 0);
-		EntitySubsystem->DestroyEntity(ReservedEntity);
-		AITEST_EQUAL("There should not be any entity across the whole system", EntitySubsystem->DebugGetEntityCount(), 0);
-		AITEST_EQUAL("The original archetype should now have no entities", EntitySubsystem->DebugGetArchetypeEntitiesCount(FloatsArchetype), 0);
+		CA_ASSUME(EntityManager);
+		const FMassEntityHandle ReservedEntity = EntityManager->ReserveEntity();
+		AITEST_TRUE("The reserved entity should be a valid entity", EntityManager->IsEntityValid(ReservedEntity));
+		AITEST_FALSE("The reserved entity should not be a valid entity", EntityManager->IsEntityBuilt(ReservedEntity));
+		EntityManager->BuildEntity(ReservedEntity, FloatsArchetype);
+		AITEST_TRUE("The reserved entity should be a valid entity", EntityManager->IsEntityBuilt(ReservedEntity));
+		AITEST_EQUAL("There should be one entity across the whole system", EntityManager->DebugGetEntityCount(), 1);
+		AITEST_EQUAL("Entity\'s archetype should be the Float one", EntityManager->GetArchetypeForEntity(ReservedEntity), FloatsArchetype);
+		AITEST_EQUAL("The created entity should have been added to the Floats archetype", EntityManager->DebugGetArchetypeEntitiesCount(FloatsArchetype), 1);
+		AITEST_EQUAL("Other archetypes should not get affected", EntityManager->DebugGetArchetypeEntitiesCount(IntsArchetype) + EntityManager->DebugGetArchetypeEntitiesCount(FloatsIntsArchetype), 0);
+		EntityManager->DestroyEntity(ReservedEntity);
+		AITEST_EQUAL("There should not be any entity across the whole system", EntityManager->DebugGetEntityCount(), 0);
+		AITEST_EQUAL("The original archetype should now have no entities", EntityManager->DebugGetArchetypeEntitiesCount(FloatsArchetype), 0);
 
 		return true;
 	}
@@ -290,19 +290,19 @@ struct FEntityTest_EntityReservationAndBuildingFromInstances : FEntityTestBase
 {
 	virtual bool InstantTest() override
 	{
-		CA_ASSUME(EntitySubsystem);
-		const FMassEntityHandle ReservedEntity = EntitySubsystem->ReserveEntity();
-		AITEST_TRUE("The reserved entity should be a valid entity", EntitySubsystem->IsEntityValid(ReservedEntity));
-		AITEST_FALSE("The reserved entity should not be a valid entity", EntitySubsystem->IsEntityBuilt(ReservedEntity));
-		EntitySubsystem->BuildEntity(ReservedEntity, MakeArrayView(&InstanceInt, 1));
-		AITEST_TRUE("The reserved entity should not be a valid entity", EntitySubsystem->IsEntityBuilt(ReservedEntity));
-		AITEST_EQUAL("There should be one entity across the whole system", EntitySubsystem->DebugGetEntityCount(), 1);
-		AITEST_EQUAL("Entity\'s archetype should be the Ints one", EntitySubsystem->GetArchetypeForEntity(ReservedEntity), IntsArchetype);
-		AITEST_EQUAL("The created entity should have been added to the Ints archetype", EntitySubsystem->DebugGetArchetypeEntitiesCount(IntsArchetype), 1);
-		AITEST_EQUAL("The entity should have the new component with the correct value set", EntitySubsystem->GetFragmentDataChecked<FTestFragment_Int>(ReservedEntity).Value, TestIntValue);
-		EntitySubsystem->DestroyEntity(ReservedEntity);
-		AITEST_EQUAL("There should not be any entity across the whole system", EntitySubsystem->DebugGetEntityCount(), 0);
-		AITEST_EQUAL("The original archetype should now have no entities", EntitySubsystem->DebugGetArchetypeEntitiesCount(FloatsArchetype), 0);
+		CA_ASSUME(EntityManager);
+		const FMassEntityHandle ReservedEntity = EntityManager->ReserveEntity();
+		AITEST_TRUE("The reserved entity should be a valid entity", EntityManager->IsEntityValid(ReservedEntity));
+		AITEST_FALSE("The reserved entity should not be a valid entity", EntityManager->IsEntityBuilt(ReservedEntity));
+		EntityManager->BuildEntity(ReservedEntity, MakeArrayView(&InstanceInt, 1));
+		AITEST_TRUE("The reserved entity should not be a valid entity", EntityManager->IsEntityBuilt(ReservedEntity));
+		AITEST_EQUAL("There should be one entity across the whole system", EntityManager->DebugGetEntityCount(), 1);
+		AITEST_EQUAL("Entity\'s archetype should be the Ints one", EntityManager->GetArchetypeForEntity(ReservedEntity), IntsArchetype);
+		AITEST_EQUAL("The created entity should have been added to the Ints archetype", EntityManager->DebugGetArchetypeEntitiesCount(IntsArchetype), 1);
+		AITEST_EQUAL("The entity should have the new component with the correct value set", EntityManager->GetFragmentDataChecked<FTestFragment_Int>(ReservedEntity).Value, TestIntValue);
+		EntityManager->DestroyEntity(ReservedEntity);
+		AITEST_EQUAL("There should not be any entity across the whole system", EntityManager->DebugGetEntityCount(), 0);
+		AITEST_EQUAL("The original archetype should now have no entities", EntityManager->DebugGetArchetypeEntitiesCount(FloatsArchetype), 0);
 
 		return true;
 	}
@@ -313,14 +313,14 @@ struct FEntityTest_ReleaseEntity : FEntityTestBase
 {
 	virtual bool InstantTest() override
 	{
-		CA_ASSUME(EntitySubsystem);
-		const FMassEntityHandle ReservedEntity = EntitySubsystem->ReserveEntity();
-		AITEST_TRUE("The reserved entity should be a valid entity", EntitySubsystem->IsEntityValid(ReservedEntity));
-		AITEST_FALSE("The reserved entity should not be a valid entity", EntitySubsystem->IsEntityBuilt(ReservedEntity));
-		AITEST_EQUAL("There should only be one entity across the whole system", EntitySubsystem->DebugGetEntityCount(), 1);
-		AITEST_EQUAL("The entity should not get associated to any archetype", EntitySubsystem->GetArchetypeForEntity(ReservedEntity), FMassArchetypeHandle());
-		EntitySubsystem->ReleaseReservedEntity(ReservedEntity);
-		AITEST_EQUAL("There should not be any entity across the whole system", EntitySubsystem->DebugGetEntityCount(), 0);
+		CA_ASSUME(EntityManager);
+		const FMassEntityHandle ReservedEntity = EntityManager->ReserveEntity();
+		AITEST_TRUE("The reserved entity should be a valid entity", EntityManager->IsEntityValid(ReservedEntity));
+		AITEST_FALSE("The reserved entity should not be a valid entity", EntityManager->IsEntityBuilt(ReservedEntity));
+		AITEST_EQUAL("There should only be one entity across the whole system", EntityManager->DebugGetEntityCount(), 1);
+		AITEST_EQUAL("The entity should not get associated to any archetype", EntityManager->GetArchetypeForEntity(ReservedEntity), FMassArchetypeHandle());
+		EntityManager->ReleaseReservedEntity(ReservedEntity);
+		AITEST_EQUAL("There should not be any entity across the whole system", EntityManager->DebugGetEntityCount(), 0);
 		return true;
 	}
 };
@@ -330,26 +330,26 @@ struct FEntityTest_ReserveAPreviouslyBuiltEntity : FEntityTestBase
 {
 	virtual bool InstantTest() override
 	{
-		CA_ASSUME(EntitySubsystem);
+		CA_ASSUME(EntityManager);
 		{
-			const FMassEntityHandle Entity = EntitySubsystem->CreateEntity(IntsArchetype);
-			AITEST_EQUAL("The entity should get associated to the right archetype", EntitySubsystem->GetArchetypeForEntity(Entity), IntsArchetype);
-			EntitySubsystem->DestroyEntity(Entity);
-			AITEST_EQUAL("There should not be any entity across the whole system", EntitySubsystem->DebugGetEntityCount(), 0);
-			AITEST_EQUAL("The original archetype should now have no entities", EntitySubsystem->DebugGetArchetypeEntitiesCount(IntsArchetype), 0);
+			const FMassEntityHandle Entity = EntityManager->CreateEntity(IntsArchetype);
+			AITEST_EQUAL("The entity should get associated to the right archetype", EntityManager->GetArchetypeForEntity(Entity), IntsArchetype);
+			EntityManager->DestroyEntity(Entity);
+			AITEST_EQUAL("There should not be any entity across the whole system", EntityManager->DebugGetEntityCount(), 0);
+			AITEST_EQUAL("The original archetype should now have no entities", EntityManager->DebugGetArchetypeEntitiesCount(IntsArchetype), 0);
 		}
 		
-		const FMassEntityHandle ReservedEntity = EntitySubsystem->ReserveEntity();
-		AITEST_TRUE("The reserved entity should be a valid entity", EntitySubsystem->IsEntityValid(ReservedEntity));
-		AITEST_FALSE("The reserved entity should not be a valid entity", EntitySubsystem->IsEntityBuilt(ReservedEntity));
-		AITEST_EQUAL("There should only be one entity across the whole system", EntitySubsystem->DebugGetEntityCount(), 1);
-		AITEST_EQUAL("The entity should not get associated to any archetype", EntitySubsystem->GetArchetypeForEntity(ReservedEntity), FMassArchetypeHandle());
-		EntitySubsystem->BuildEntity(ReservedEntity, FloatsArchetype);
-		AITEST_TRUE("The reserved entity should not be a valid entity", EntitySubsystem->IsEntityBuilt(ReservedEntity));
-		AITEST_EQUAL("The entity should get associated to the right archetype", EntitySubsystem->GetArchetypeForEntity(ReservedEntity), FloatsArchetype);
-		EntitySubsystem->DestroyEntity(ReservedEntity);
-		AITEST_EQUAL("There should not be any entity across the whole system", EntitySubsystem->DebugGetEntityCount(), 0);
-		AITEST_EQUAL("The original archetype should now have no entities", EntitySubsystem->DebugGetArchetypeEntitiesCount(FloatsArchetype), 0);
+		const FMassEntityHandle ReservedEntity = EntityManager->ReserveEntity();
+		AITEST_TRUE("The reserved entity should be a valid entity", EntityManager->IsEntityValid(ReservedEntity));
+		AITEST_FALSE("The reserved entity should not be a valid entity", EntityManager->IsEntityBuilt(ReservedEntity));
+		AITEST_EQUAL("There should only be one entity across the whole system", EntityManager->DebugGetEntityCount(), 1);
+		AITEST_EQUAL("The entity should not get associated to any archetype", EntityManager->GetArchetypeForEntity(ReservedEntity), FMassArchetypeHandle());
+		EntityManager->BuildEntity(ReservedEntity, FloatsArchetype);
+		AITEST_TRUE("The reserved entity should not be a valid entity", EntityManager->IsEntityBuilt(ReservedEntity));
+		AITEST_EQUAL("The entity should get associated to the right archetype", EntityManager->GetArchetypeForEntity(ReservedEntity), FloatsArchetype);
+		EntityManager->DestroyEntity(ReservedEntity);
+		AITEST_EQUAL("There should not be any entity across the whole system", EntityManager->DebugGetEntityCount(), 0);
+		AITEST_EQUAL("The original archetype should now have no entities", EntityManager->DebugGetArchetypeEntitiesCount(FloatsArchetype), 0);
 		return true;
 	}
 };
diff --git a/Engine/Plugins/Runtime/MassEntity/Source/MassEntityTestSuite/Private/MassExecutionTest.cpp b/Engine/Plugins/Runtime/MassEntity/Source/MassEntityTestSuite/Private/MassExecutionTest.cpp
--- a/Engine/Plugins/Runtime/MassEntity/Source/MassEntityTestSuite/Private/MassExecutionTest.cpp
+++ b/Engine/Plugins/Runtime/MassEntity/Source/MassEntityTestSuite/Private/MassExecutionTest.cpp
@@ -22,7 +22,7 @@ struct FExecution_Setup : FExecutionTestBase
 	virtual bool InstantTest() override
 	{
 		AITEST_NOT_NULL("World needs to exist for the test to be performed", World);
-		AITEST_NOT_NULL("EntitySubsystem needs to be created for the test to be performed", EntitySubsystem);
+		AITEST_NOT_NULL("EntitySubsystem needs to be created for the test to be performed", EntityManager.Get());
 
 		return true;
 	}
@@ -34,9 +34,9 @@ struct FExecution_EmptyArray : FExecutionTestBase
 {
 	virtual bool InstantTest() override
 	{
-		CA_ASSUME(EntitySubsystem); // if EntitySubsystem null InstantTest won't be called at all
+		CA_ASSUME(EntityManager); // if EntitySubsystem null InstantTest won't be called at all
 		const float DeltaSeconds = 0.f;
-		FMassProcessingContext ProcessingContext(*EntitySubsystem, DeltaSeconds);
+		FMassProcessingContext ProcessingContext(*EntityManager, DeltaSeconds);
 		// no test performed, let's just see if it results in errors/warnings
 		UE::Mass::Executor::RunProcessorsView(TArrayView<UMassProcessor*>(), ProcessingContext);
 		return true;
@@ -49,9 +49,9 @@ struct FExecution_EmptyPipeline : FExecutionTestBase
 {
 	virtual bool InstantTest() override
 	{
-		CA_ASSUME(EntitySubsystem); // if EntitySubsystem null InstantTest won't be called at all
+		CA_ASSUME(EntityManager); // if EntitySubsystem null InstantTest won't be called at all
 		const float DeltaSeconds = 0.f;
-		FMassProcessingContext ProcessingContext(*EntitySubsystem, DeltaSeconds);
+		FMassProcessingContext ProcessingContext(*EntityManager, DeltaSeconds);
 		FMassRuntimePipeline Pipeline;
 		// no test performed, let's just see if it results in errors/warnings
 		UE::Mass::Executor::Run(Pipeline, ProcessingContext);
@@ -65,13 +65,13 @@ struct FExecution_InvalidProcessingContext : FExecutionTestBase
 {
 	virtual bool InstantTest() override
 	{
-		CA_ASSUME(EntitySubsystem);
+		CA_ASSUME(EntityManager);
 		const float DeltaSeconds = 0.f;
 		FMassProcessingContext ProcessingContext;
 		// test assumption
-		AITEST_NULL("FMassProcessingContext\'s default constructor is expected to set FMassProcessingContext.EntitySubsystem to null", ProcessingContext.EntitySubsystem);
+		AITEST_NULL("FMassProcessingContext\'s default constructor is expected to set FMassProcessingContext.EntitySubsystem to null", ProcessingContext.EntityManager.Get());
 		
-		GetTestRunner().AddExpectedError(TEXT("ProcessingContext.EntitySubsystem is null"), EAutomationExpectedErrorFlags::Contains, 1);		
+		GetTestRunner().AddExpectedError(TEXT("ProcessingContext.EntityManager is null"), EAutomationExpectedErrorFlags::Contains, 1);		
 		// note that using RunProcessorsView is to bypass reasonable tests UE::Mass::Executor::Run(Pipeline,...) does that are 
 		// reported via ensures which are not handled by the automation framework
 		UE::Mass::Executor::RunProcessorsView(TArrayView<UMassProcessor*>(), ProcessingContext);
@@ -86,9 +86,9 @@ struct FExecution_SingleNullProcessor : FExecutionTestBase
 {
 	virtual bool InstantTest() override
 	{
-		CA_ASSUME(EntitySubsystem);
+		CA_ASSUME(EntityManager);
 		const float DeltaSeconds = 0.f;
-		FMassProcessingContext ProcessingContext(*EntitySubsystem, DeltaSeconds);
+		FMassProcessingContext ProcessingContext(EntityManager, DeltaSeconds);
 		TArray<UMassProcessor*> Processors;
 		Processors.Add(nullptr);
 		
@@ -107,10 +107,10 @@ struct FExecution_SingleValidProcessor : FExecutionTestBase
 {
 	virtual bool InstantTest() override
 	{
-		CA_ASSUME(EntitySubsystem);
+		CA_ASSUME(EntityManager);
 		const float DeltaSeconds = 0.f;
-		FMassProcessingContext ProcessingContext(*EntitySubsystem, DeltaSeconds);
-		UMassTestProcessorBase* Processor = NewObject<UMassTestProcessorBase>(EntitySubsystem);
+		FMassProcessingContext ProcessingContext(EntityManager, DeltaSeconds);
+		UMassTestProcessorBase* Processor = NewObject<UMassTestProcessorBase>();
 		check(Processor);
 
 		// nothing should break. The actual result of processing is getting tested in MassProcessorTests.cpp
@@ -125,9 +125,9 @@ struct FExecution_MultipleNullProcessors : FExecutionTestBase
 {
 	virtual bool InstantTest() override
 	{
-		CA_ASSUME(EntitySubsystem);
+		CA_ASSUME(EntityManager);
 		const float DeltaSeconds = 0.f;
-		FMassProcessingContext ProcessingContext(*EntitySubsystem, DeltaSeconds);
+		FMassProcessingContext ProcessingContext(EntityManager, DeltaSeconds);
 		TArray<UMassProcessor*> Processors;
 		Processors.Add(nullptr);
 		Processors.Add(nullptr);
@@ -148,10 +148,10 @@ struct FExecution_Sparse : FEntityTestBase
 {
 	virtual bool InstantTest() override
 	{
-		CA_ASSUME(EntitySubsystem);
+		CA_ASSUME(EntityManager);
 		const float DeltaSeconds = 0.f;
-		FMassProcessingContext ProcessingContext(*EntitySubsystem, DeltaSeconds);
-		UMassTestProcessorBase* Processor = NewObject<UMassTestProcessorBase>(EntitySubsystem);
+		FMassProcessingContext ProcessingContext(*EntityManager, DeltaSeconds);
+		UMassTestProcessorBase* Processor = NewObject<UMassTestProcessorBase>();
 		check(Processor);
 
 		FMassRuntimePipeline Pipeline;
diff --git a/Engine/Plugins/Runtime/MassEntity/Source/MassEntityTestSuite/Private/MassMultiThreadTest.cpp b/Engine/Plugins/Runtime/MassEntity/Source/MassEntityTestSuite/Private/MassMultiThreadTest.cpp
--- a/Engine/Plugins/Runtime/MassEntity/Source/MassEntityTestSuite/Private/MassMultiThreadTest.cpp
+++ b/Engine/Plugins/Runtime/MassEntity/Source/MassEntityTestSuite/Private/MassMultiThreadTest.cpp
@@ -41,7 +41,7 @@ struct FMTTestBase : FEntityTestBase
 			CompositeProcessor->SetGroupName(TEXT("Test"));
 			CompositeProcessor->SetProcessors(MakeArrayView<UMassProcessor*>((UMassProcessor**)Processors.GetData(), Processors.Num()));
 
-			FMassProcessingContext Context(*EntitySubsystem, /*DeltaTime=*/0);
+			FMassProcessingContext Context(*EntityManager, /*DeltaTime=*/0);
 			FinishEvent = UE::Mass::Executor::TriggerParallelTasks(*CompositeProcessor, Context, []() {});
 		}
 
@@ -70,14 +70,14 @@ struct FMTTrivial : FMTTestBase
 			return false;
 		}
 
-		EntitySubsystem->BatchCreateEntities(IntsArchetype, NumToCreate, Entities);
+		EntityManager->BatchCreateEntities(IntsArchetype, NumToCreate, Entities);
 
 		Processors.Reset();
 		{
 			UMassTestProcessorBase* Proc = Processors.Add_GetRef(NewObject<UMassTestProcessor_A>());
 
 			Proc->TestGetQuery().AddRequirement<FTestFragment_Int>(EMassFragmentAccess::ReadOnly);
-			Proc->ExecutionFunction = [this, Proc](UMassEntitySubsystem& InEntitySubsystem, FMassExecutionContext& Context)
+			Proc->ExecutionFunction = [this, Proc](FMassEntityManager& InEntitySubsystem, FMassExecutionContext& Context)
 			{
 				Proc->TestGetQuery().ForEachEntityChunk(InEntitySubsystem, Context, [this](FMassExecutionContext& Context)
 					{
@@ -110,7 +110,7 @@ struct FMTBasic : FMTTestBase
 			return false;
 		}
 
-		EntitySubsystem->BatchCreateEntities(FloatsIntsArchetype, NumToCreate, Entities);
+		EntityManager->BatchCreateEntities(FloatsIntsArchetype, NumToCreate, Entities);
 
 		Processors.Reset();
 		{
@@ -118,7 +118,7 @@ struct FMTBasic : FMTTestBase
 			Proc->GetMutableExecutionOrder().ExecuteAfter.Add(GetProcessorName<UMassTestProcessor_B>());
 			Proc->TestGetQuery().AddRequirement<FTestFragment_Float>(EMassFragmentAccess::ReadOnly);
 			Proc->TestGetQuery().AddRequirement<FTestFragment_Int>(EMassFragmentAccess::ReadWrite);
-			Proc->ExecutionFunction = [this, Proc](UMassEntitySubsystem& InEntitySubsystem, FMassExecutionContext& Context)
+			Proc->ExecutionFunction = [this, Proc](FMassEntityManager& InEntitySubsystem, FMassExecutionContext& Context)
 			{
 				Proc->TestGetQuery().ForEachEntityChunk(InEntitySubsystem, Context, [this](FMassExecutionContext& Context)
 					{
@@ -136,7 +136,7 @@ struct FMTBasic : FMTTestBase
 			Proc->GetMutableExecutionOrder().ExecuteAfter.Add(GetProcessorName<UMassTestProcessor_A>());
 			Proc->TestGetQuery().AddRequirement<FTestFragment_Int>(EMassFragmentAccess::ReadOnly);
 			Proc->TestGetQuery().AddRequirement<FTestFragment_Float>(EMassFragmentAccess::ReadWrite);
-			Proc->ExecutionFunction = [this, Proc](UMassEntitySubsystem& InEntitySubsystem, FMassExecutionContext& Context)
+			Proc->ExecutionFunction = [this, Proc](FMassEntityManager& InEntitySubsystem, FMassExecutionContext& Context)
 			{
 				Proc->TestGetQuery().ForEachEntityChunk(InEntitySubsystem, Context, [this](FMassExecutionContext& Context)
 					{
@@ -152,7 +152,7 @@ struct FMTBasic : FMTTestBase
 		{
 			UMassTestProcessorBase* Proc = Processors.Add_GetRef(NewObject<UMassTestProcessor_A>());
 			Proc->TestGetQuery().AddRequirement<FTestFragment_Int>(EMassFragmentAccess::ReadWrite);
-			Proc->ExecutionFunction = [this, Proc](UMassEntitySubsystem& InEntitySubsystem, FMassExecutionContext& Context)
+			Proc->ExecutionFunction = [this, Proc](FMassEntityManager& InEntitySubsystem, FMassExecutionContext& Context)
 			{
 				int Index = 0;
 				Proc->TestGetQuery().ForEachEntityChunk(InEntitySubsystem, Context, [this, &Index](FMassExecutionContext& Context)
