Separated "entity management" and USubsystem stuff of UMassEntityManager. Now UMassEntityManager hosts a shared ref of newly created FMassEntityManager, and users are free to create their own FMassEnt

diff --git a/Engine/Plugins/Runtime/MassEntity/Source/MassEntityTestSuite/Private/MassObserverTest.cpp b/Engine/Plugins/Runtime/MassEntity/Source/MassEntityTestSuite/Private/MassObserverTest.cpp
--- a/Engine/Plugins/Runtime/MassEntity/Source/MassEntityTestSuite/Private/MassObserverTest.cpp
+++ b/Engine/Plugins/Runtime/MassEntity/Source/MassEntityTestSuite/Private/MassObserverTest.cpp
@@ -37,11 +37,10 @@ struct FTagBaseOperation : FEntityTestBase
 	{
 		if (FEntityTestBase::SetUp())
 		{
-			//TArray<FMassEntityHandle>* ContainerPtr = &AffectedEntities;
-			TagObserver = NewObject<UMassTestProcessorBase>(EntitySubsystem);
+			TagObserver = NewObject<UMassTestProcessorBase>();
 			TagObserver->TestGetQuery().AddRequirement<FTestFragment_Int>(EMassFragmentAccess::ReadOnly);
 			TagObserver->TestGetQuery().AddTagRequirement<FTestTag_A>(EMassFragmentPresence::All);
-			TagObserver->ExecutionFunction = [this](UMassEntitySubsystem& InEntitySubsystem, FMassExecutionContext& Context)
+			TagObserver->ExecutionFunction = [this](FMassEntityManager& InEntitySubsystem, FMassExecutionContext& Context)
 			{
 				TagObserver->TestGetQuery().ForEachEntityChunk(InEntitySubsystem, Context, [this](FMassExecutionContext& Context)
 					{
@@ -56,15 +55,15 @@ struct FTagBaseOperation : FEntityTestBase
 
 	virtual bool InstantTest() override
 	{
-		FMassObserverManager& ObserverManager = EntitySubsystem->GetObserverManager();
+		FMassObserverManager& ObserverManager = EntityManager->GetObserverManager();
 		ObserverManager.AddObserverInstance(*FTestTag_A::StaticStruct(), OperationObserved, *TagObserver);
 
-		EntitySubsystem->BatchCreateEntities(IntsArchetype, 3, EntitiesInt);
-		EntitySubsystem->BatchCreateEntities(FloatsIntsArchetype, 3, EntitiesIntsFloat);
+		EntityManager->BatchCreateEntities(IntsArchetype, 3, EntitiesInt);
+		EntityManager->BatchCreateEntities(FloatsIntsArchetype, 3, EntitiesIntsFloat);
 
 		if (PerformOperation())
 		{
-			EntitySubsystem->FlushCommands();
+			EntityManager->FlushCommands();
 			AITEST_EQUAL(TEXT("The tag observer is expected to be run for predicted number of entities"), AffectedEntities.Num(), ExpectedEntities.Num());
 			
 			ExpectedEntities.Sort(EntityIndexSorted);
@@ -89,7 +88,7 @@ struct FTag_SingleEntitySingleArchetypeAdd : FTagBaseOperation
 	virtual bool PerformOperation() override 
 	{
 		ExpectedEntities = { EntitiesInt[1] };
-		EntitySubsystem->Defer().AddTag<FTestTag_A>(EntitiesInt[1]);
+		EntityManager->Defer().AddTag<FTestTag_A>(EntitiesInt[1]);
 		return true;
 	}
 };
@@ -102,11 +101,11 @@ struct FTag_SingleEntitySingleArchetypeRemove : FTagBaseOperation
 	{
 		ExpectedEntities = { EntitiesInt[1] };
 
-		EntitySubsystem->Defer().AddTag<FTestTag_A>(EntitiesInt[1]);
-		EntitySubsystem->FlushCommands();
+		EntityManager->Defer().AddTag<FTestTag_A>(EntitiesInt[1]);
+		EntityManager->FlushCommands();
 		// since we're only observing tag removal we don't expect AffectedEntities to contain any data at this point
 		AITEST_EQUAL(TEXT("Tag addition is not being observed and is not expected to produce results yet"), AffectedEntities.Num(), 0);
-		EntitySubsystem->Defer().RemoveTag<FTestTag_A>(EntitiesInt[1]);
+		EntityManager->Defer().RemoveTag<FTestTag_A>(EntitiesInt[1]);
 		return true;
 	}
 };
@@ -118,11 +117,11 @@ struct FTag_SingleEntitySingleArchetypeDestroy : FTagBaseOperation
 	virtual bool PerformOperation() override
 	{
 		ExpectedEntities = { EntitiesInt[1] };
-		EntitySubsystem->Defer().AddTag<FTestTag_A>(EntitiesInt[1]); 
-		EntitySubsystem->FlushCommands();
+		EntityManager->Defer().AddTag<FTestTag_A>(EntitiesInt[1]); 
+		EntityManager->FlushCommands();
 		// since we're only observing tag removal we don't expect AffectedEntities to contain any data at this point
 		AITEST_EQUAL(TEXT("Tag addition is not being observed and is not expected to produce results yet"), AffectedEntities.Num(), 0);
-		EntitySubsystem->Defer().DestroyEntity(EntitiesInt[1]);
+		EntityManager->Defer().DestroyEntity(EntitiesInt[1]);
 		return true;
 	}
 };
@@ -137,7 +136,7 @@ struct FTag_MultipleArchetypeAdd : FTagBaseOperation
 		ExpectedEntities = { EntitiesInt[0], EntitiesInt[2], EntitiesIntsFloat[1] };
 		for (const FMassEntityHandle& ModifiedEntity : ExpectedEntities)
 		{
-			EntitySubsystem->Defer().AddTag<FTestTag_A>(ModifiedEntity);
+			EntityManager->Defer().AddTag<FTestTag_A>(ModifiedEntity);
 		}
 		return true;
 	}
@@ -153,14 +152,14 @@ struct FTag_MultipleArchetypeRemove : FTagBaseOperation
 		ExpectedEntities = { EntitiesInt[0], EntitiesInt[2], EntitiesIntsFloat[1] };
 		for (const FMassEntityHandle& ModifiedEntity : ExpectedEntities)
 		{
-			EntitySubsystem->Defer().AddTag<FTestTag_A>(ModifiedEntity);
+			EntityManager->Defer().AddTag<FTestTag_A>(ModifiedEntity);
 		}
-		EntitySubsystem->FlushCommands();
+		EntityManager->FlushCommands();
 		// since we're only observing tag removal we don't expect AffectedEntities to contain any data at this point
 		AITEST_EQUAL(TEXT("Tag addition is not being observed and is not expected to produce results yet"), AffectedEntities.Num(), 0);
 		for (const FMassEntityHandle& ModifiedEntity : ExpectedEntities)
 		{
-			EntitySubsystem->Defer().RemoveTag<FTestTag_A>(ModifiedEntity);
+			EntityManager->Defer().RemoveTag<FTestTag_A>(ModifiedEntity);
 		}
 		return true;
 	}
@@ -176,14 +175,14 @@ struct FTag_MultipleArchetypeDestroy : FTagBaseOperation
 		ExpectedEntities = { EntitiesInt[0], EntitiesInt[2], EntitiesIntsFloat[1] };
 		for (const FMassEntityHandle& ModifiedEntity : ExpectedEntities)
 		{
-			EntitySubsystem->Defer().AddTag<FTestTag_A>(ModifiedEntity);
+			EntityManager->Defer().AddTag<FTestTag_A>(ModifiedEntity);
 		}
-		EntitySubsystem->FlushCommands();
+		EntityManager->FlushCommands();
 		// since we're only observing tag removal we don't expect AffectedEntities to contain any data at this point
 		AITEST_EQUAL(TEXT("Tag addition is not being observed and is not expected to produce results yet"), AffectedEntities.Num(), 0);
 		for (const FMassEntityHandle& ModifiedEntity : ExpectedEntities)
 		{
-			EntitySubsystem->Defer().DestroyEntity(ModifiedEntity);
+			EntityManager->Defer().DestroyEntity(ModifiedEntity);
 		}
 		return true;
 	}
@@ -199,14 +198,14 @@ struct FTag_MultipleArchetypeSwap : FTagBaseOperation
 		ExpectedEntities = { EntitiesIntsFloat[1], EntitiesInt[0], EntitiesInt[2] };
 		for (const FMassEntityHandle& ModifiedEntity : ExpectedEntities)
 		{
-			EntitySubsystem->Defer().AddTag<FTestTag_A>(ModifiedEntity);
+			EntityManager->Defer().AddTag<FTestTag_A>(ModifiedEntity);
 		}
-		EntitySubsystem->FlushCommands();
+		EntityManager->FlushCommands();
 		// since we're only observing tag removal we don't expect AffectedEntities to contain any data at this point
 		AITEST_EQUAL(TEXT("Tag addition is not being observed and is not expected to produce results yet"), AffectedEntities.Num(), 0);
 		for (const FMassEntityHandle& ModifiedEntity : ExpectedEntities)
 		{
-			EntitySubsystem->Defer().SwapTags<FTestTag_A, FTestTag_B>(ModifiedEntity);
+			EntityManager->Defer().SwapTags<FTestTag_A, FTestTag_B>(ModifiedEntity);
 		}
 		return true;
 	}
diff --git a/Engine/Plugins/Runtime/MassEntity/Source/MassEntityTestSuite/Private/MassProcessorTest.cpp b/Engine/Plugins/Runtime/MassEntity/Source/MassEntityTestSuite/Private/MassProcessorTest.cpp
--- a/Engine/Plugins/Runtime/MassEntity/Source/MassEntityTestSuite/Private/MassProcessorTest.cpp
+++ b/Engine/Plugins/Runtime/MassEntity/Source/MassEntityTestSuite/Private/MassProcessorTest.cpp
@@ -17,11 +17,11 @@ namespace FMassProcessorTest
 {
 
 template<typename TProcessor>
-int32 SimpleProcessorRun(UMassEntitySubsystem& EntitySubsystem)
+int32 SimpleProcessorRun(FMassEntityManager& EntityManager)
 {
 	int32 EntityProcessedCount = 0;
-	TProcessor* Processor = NewObject<TProcessor>(&EntitySubsystem);
-	Processor->ExecutionFunction = [Processor, &EntityProcessedCount](UMassEntitySubsystem& InEntitySubsystem, FMassExecutionContext& Context) {
+	TProcessor* Processor = NewObject<TProcessor>();
+	Processor->ExecutionFunction = [Processor, &EntityProcessedCount](FMassEntityManager& InEntitySubsystem, FMassExecutionContext& Context) {
 		check(Processor);
 		Processor->TestGetQuery().ForEachEntityChunk(InEntitySubsystem, Context, [Processor, &EntityProcessedCount](FMassExecutionContext& Context)
 			{
@@ -29,7 +29,7 @@ int32 SimpleProcessorRun(UMassEntitySubsystem& EntitySubsystem)
 			});
 	};
 
-	FMassProcessingContext ProcessingContext(EntitySubsystem, /*DeltaSeconds=*/0.f);
+	FMassProcessingContext ProcessingContext(EntityManager, /*DeltaSeconds=*/0.f);
 	UE::Mass::Executor::Run(*Processor, ProcessingContext);
 
 	return EntityProcessedCount;
@@ -39,9 +39,9 @@ struct FProcessorTest_NoEntities : FEntityTestBase
 {
 	virtual bool InstantTest() override
 	{
-		CA_ASSUME(EntitySubsystem);
+		CA_ASSUME(EntityManager);
 		
-		int32 EntityProcessedCount = SimpleProcessorRun<UMassTestProcessor_Floats>(*EntitySubsystem);
+		int32 EntityProcessedCount = SimpleProcessorRun<UMassTestProcessor_Floats>(*EntityManager);
 		AITEST_EQUAL("No entities have been created yet so the processor shouldn't do any work.", EntityProcessedCount, 0);
 				
 		return true;
@@ -54,11 +54,11 @@ struct FProcessorTest_NoMatchingEntities : FEntityTestBase
 {
 	virtual bool InstantTest() override
 	{
-		CA_ASSUME(EntitySubsystem);
+		CA_ASSUME(EntityManager);
 
 		TArray<FMassEntityHandle> EntitiesCreated;
-		EntitySubsystem->BatchCreateEntities(IntsArchetype, 100, EntitiesCreated);
-		int32 EntityProcessedCount = SimpleProcessorRun<UMassTestProcessor_Floats>(*EntitySubsystem);
+		EntityManager->BatchCreateEntities(IntsArchetype, 100, EntitiesCreated);
+		int32 EntityProcessedCount = SimpleProcessorRun<UMassTestProcessor_Floats>(*EntityManager);
 		AITEST_EQUAL("No matching entities have been created yet so the processor shouldn't do any work.", EntityProcessedCount, 0);
 
 		return true;
@@ -70,12 +70,12 @@ struct FProcessorTest_OneMatchingArchetype : FEntityTestBase
 {
 	virtual bool InstantTest() override
 	{
-		CA_ASSUME(EntitySubsystem);
+		CA_ASSUME(EntityManager);
 
 		const int32 EntitiesToCreate = 100;
 		TArray<FMassEntityHandle> EntitiesCreated;
-		EntitySubsystem->BatchCreateEntities(FloatsArchetype, EntitiesToCreate, EntitiesCreated);
-		int32 EntityProcessedCount = SimpleProcessorRun<UMassTestProcessor_Floats>(*EntitySubsystem);
+		EntityManager->BatchCreateEntities(FloatsArchetype, EntitiesToCreate, EntitiesCreated);
+		int32 EntityProcessedCount = SimpleProcessorRun<UMassTestProcessor_Floats>(*EntityManager);
 		AITEST_EQUAL("No matching entities have been created yet so the processor shouldn't do any work.", EntityProcessedCount, EntitiesToCreate);
 
 		return true;
@@ -88,15 +88,15 @@ struct FProcessorTest_MultipleMatchingArchetype : FEntityTestBase
 {
 	virtual bool InstantTest() override
 	{
-		CA_ASSUME(EntitySubsystem);
+		CA_ASSUME(EntityManager);
 
 		const int32 EntitiesToCreate = 100;
 		TArray<FMassEntityHandle> EntitiesCreated;
-		EntitySubsystem->BatchCreateEntities(FloatsArchetype, EntitiesToCreate, EntitiesCreated);
-		EntitySubsystem->BatchCreateEntities(FloatsIntsArchetype, EntitiesToCreate, EntitiesCreated);
-		EntitySubsystem->BatchCreateEntities(IntsArchetype, EntitiesToCreate, EntitiesCreated);
+		EntityManager->BatchCreateEntities(FloatsArchetype, EntitiesToCreate, EntitiesCreated);
+		EntityManager->BatchCreateEntities(FloatsIntsArchetype, EntitiesToCreate, EntitiesCreated);
+		EntityManager->BatchCreateEntities(IntsArchetype, EntitiesToCreate, EntitiesCreated);
 		// note that only two of these archetypes match 
-		int32 EntityProcessedCount = SimpleProcessorRun<UMassTestProcessor_Floats>(*EntitySubsystem);
+		int32 EntityProcessedCount = SimpleProcessorRun<UMassTestProcessor_Floats>(*EntityManager);
 		AITEST_EQUAL("Two of given three archetypes should match.", EntityProcessedCount, EntitiesToCreate * 2);
 
 		return true;
diff --git a/Engine/Plugins/Runtime/MassEntity/Source/MassEntityTestSuite/Private/MassQueryTest.cpp b/Engine/Plugins/Runtime/MassEntity/Source/MassEntityTestSuite/Private/MassQueryTest.cpp
--- a/Engine/Plugins/Runtime/MassEntity/Source/MassEntityTestSuite/Private/MassQueryTest.cpp
+++ b/Engine/Plugins/Runtime/MassEntity/Source/MassEntityTestSuite/Private/MassQueryTest.cpp
@@ -31,9 +31,9 @@ struct FQueryTest_ProcessorRequirements : FEntityTestBase
 {
 	virtual bool InstantTest() override
 	{
-		CA_ASSUME(EntitySubsystem);
+		CA_ASSUME(EntityManager);
 
-		UMassTestProcessor_Floats* Processor = NewObject<UMassTestProcessor_Floats>(EntitySubsystem);
+		UMassTestProcessor_Floats* Processor = NewObject<UMassTestProcessor_Floats>();
 		TConstArrayView<FMassFragmentRequirementDescription> Requirements = Processor->TestGetQuery().GetFragmentRequirements();
 		
 		AITEST_TRUE("Query should have extracted some requirements from the given Processor", Requirements.Num() > 0);
@@ -50,7 +50,7 @@ struct FQueryTest_ExplicitRequirements : FEntityTestBase
 {
 	virtual bool InstantTest() override
 	{
-		CA_ASSUME(EntitySubsystem);
+		CA_ASSUME(EntityManager);
 				
 		FMassEntityQuery Query({ FTestFragment_Float::StaticStruct()});
 		TConstArrayView<FMassFragmentRequirementDescription> Requirements = Query.GetFragmentRequirements();
@@ -69,14 +69,14 @@ struct FQueryTest_FragmentViewBinding : FEntityTestBase
 {
 	virtual bool InstantTest() override
 	{
-		CA_ASSUME(EntitySubsystem);
+		CA_ASSUME(EntityManager);
 
-		FMassEntityHandle Entity = EntitySubsystem->CreateEntity(FloatsArchetype);
-		FTestFragment_Float& TestedFragment = EntitySubsystem->GetFragmentDataChecked<FTestFragment_Float>(Entity);
+		FMassEntityHandle Entity = EntityManager->CreateEntity(FloatsArchetype);
+		FTestFragment_Float& TestedFragment = EntityManager->GetFragmentDataChecked<FTestFragment_Float>(Entity);
 		AITEST_TRUE("Initial value of the fragment should match expectations", TestedFragment.Value == 0.f);
 
-		UMassTestProcessor_Floats* Processor = NewObject<UMassTestProcessor_Floats>(EntitySubsystem);
-		Processor->ExecutionFunction = [Processor](UMassEntitySubsystem& InEntitySubsystem, FMassExecutionContext& Context) {
+		UMassTestProcessor_Floats* Processor = NewObject<UMassTestProcessor_Floats>();
+		Processor->ExecutionFunction = [Processor](FMassEntityManager& InEntitySubsystem, FMassExecutionContext& Context) {
 			check(Processor);
 			//FMassEntityQuery Query(*Processor);
 
@@ -91,7 +91,7 @@ struct FQueryTest_FragmentViewBinding : FEntityTestBase
 				});
 		};
 
-		FMassProcessingContext ProcessingContext(*EntitySubsystem, /*DeltaSeconds=*/0.f);
+		FMassProcessingContext ProcessingContext(*EntityManager, /*DeltaSeconds=*/0.f);
 		UE::Mass::Executor::Run(*Processor, ProcessingContext);
 
 		AITEST_EQUAL("Fragment value should have changed to the expected value", TestedFragment.Value, 13.f);
@@ -105,17 +105,17 @@ struct FQueryTest_ExecuteSingleArchetype : FEntityTestBase
 {
 	virtual bool InstantTest() override
 	{
-		CA_ASSUME(EntitySubsystem);
+		CA_ASSUME(EntityManager);
 
 		const int32 NumToCreate = 10;
 		TArray<FMassEntityHandle> EntitiesCreated;
-		EntitySubsystem->BatchCreateEntities(FloatsArchetype, NumToCreate, EntitiesCreated);
+		EntityManager->BatchCreateEntities(FloatsArchetype, NumToCreate, EntitiesCreated);
 		
 		int TotalProcessed = 0;
 
 		FMassExecutionContext ExecContext;
 		FMassEntityQuery Query({ FTestFragment_Float::StaticStruct() });
-		Query.ForEachEntityChunk(*EntitySubsystem, ExecContext, [&TotalProcessed](FMassExecutionContext& Context)
+		Query.ForEachEntityChunk(*EntityManager, ExecContext, [&TotalProcessed](FMassExecutionContext& Context)
 			{
 				TotalProcessed += Context.GetNumEntities();
 				TArrayView<FTestFragment_Float> Floats = Context.GetMutableFragmentView<FTestFragment_Float>();
@@ -130,7 +130,7 @@ struct FQueryTest_ExecuteSingleArchetype : FEntityTestBase
 
 		for (FMassEntityHandle& Entity : EntitiesCreated)
 		{
-			const FTestFragment_Float& TestedFragment = EntitySubsystem->GetFragmentDataChecked<FTestFragment_Float>(Entity);
+			const FTestFragment_Float& TestedFragment = EntityManager->GetFragmentDataChecked<FTestFragment_Float>(Entity);
 			AITEST_EQUAL("Every fragment value should have changed to the expected value", TestedFragment.Value, 13.f);
 		}
 
@@ -144,22 +144,22 @@ struct FQueryTest_ExecuteMultipleArchetypes : FEntityTestBase
 {
 	virtual bool InstantTest() override
 	{
-		CA_ASSUME(EntitySubsystem);
+		CA_ASSUME(EntityManager);
 
 		const int32 FloatsArchetypeCreated = 7;
 		const int32 IntsArchetypeCreated = 11;
 		const int32 FloatsIntsArchetypeCreated = 13;
 		TArray<FMassEntityHandle> EntitiesCreated;
-		EntitySubsystem->BatchCreateEntities(IntsArchetype, IntsArchetypeCreated, EntitiesCreated);
+		EntityManager->BatchCreateEntities(IntsArchetype, IntsArchetypeCreated, EntitiesCreated);
 		// clear to store only the float-related entities
 		EntitiesCreated.Reset();
-		EntitySubsystem->BatchCreateEntities(FloatsArchetype, FloatsArchetypeCreated, EntitiesCreated);
-		EntitySubsystem->BatchCreateEntities(FloatsIntsArchetype, FloatsIntsArchetypeCreated, EntitiesCreated);
+		EntityManager->BatchCreateEntities(FloatsArchetype, FloatsArchetypeCreated, EntitiesCreated);
+		EntityManager->BatchCreateEntities(FloatsIntsArchetype, FloatsIntsArchetypeCreated, EntitiesCreated);
 
 		int TotalProcessed = 0;
 		FMassExecutionContext ExecContext;
 		FMassEntityQuery Query({ FTestFragment_Float::StaticStruct() });
-		Query.ForEachEntityChunk(*EntitySubsystem, ExecContext, [&TotalProcessed](FMassExecutionContext& Context)
+		Query.ForEachEntityChunk(*EntityManager, ExecContext, [&TotalProcessed](FMassExecutionContext& Context)
 			{
 				TotalProcessed += Context.GetNumEntities();
 				TArrayView<FTestFragment_Float> Floats = Context.GetMutableFragmentView<FTestFragment_Float>();
@@ -174,7 +174,7 @@ struct FQueryTest_ExecuteMultipleArchetypes : FEntityTestBase
 
 		for (FMassEntityHandle& Entity : EntitiesCreated)
 		{
-			const FTestFragment_Float& TestedFragment = EntitySubsystem->GetFragmentDataChecked<FTestFragment_Float>(Entity);
+			const FTestFragment_Float& TestedFragment = EntityManager->GetFragmentDataChecked<FTestFragment_Float>(Entity);
 			AITEST_EQUAL("Every fragment value should have changed to the expected value", TestedFragment.Value, 13.f);
 		}
 
@@ -189,11 +189,11 @@ struct FQueryTest_ExecuteSparse : FEntityTestBase
 {
 	virtual bool InstantTest() override
 	{
-		CA_ASSUME(EntitySubsystem);
+		CA_ASSUME(EntityManager);
 
 		const int32 NumToCreate = 10;
 		TArray<FMassEntityHandle> AllEntitiesCreated;
-		EntitySubsystem->BatchCreateEntities(FloatsArchetype, NumToCreate, AllEntitiesCreated);
+		EntityManager->BatchCreateEntities(FloatsArchetype, NumToCreate, AllEntitiesCreated);
 		
 		TArray<int32> IndicesToProcess = { 1, 2, 3, 6, 7};
 		TArray<FMassEntityHandle> EntitiesToProcess;
@@ -217,7 +217,7 @@ struct FQueryTest_ExecuteSparse : FEntityTestBase
 		FMassEntityQuery TestQuery;
 		TestQuery.AddRequirement<FTestFragment_Float>(EMassFragmentAccess::ReadWrite);
 		TestQuery.ForEachEntityChunk(FMassArchetypeEntityCollection(FloatsArchetype, EntitiesToProcess, FMassArchetypeEntityCollection::NoDuplicates)
-								, *EntitySubsystem, ExecContext, [&TotalProcessed](FMassExecutionContext& Context)
+								, *EntityManager, ExecContext, [&TotalProcessed](FMassExecutionContext& Context)
 			{
 				TotalProcessed += Context.GetNumEntities();
 				TArrayView<FTestFragment_Float> Floats = Context.GetMutableFragmentView<FTestFragment_Float>();
@@ -232,13 +232,13 @@ struct FQueryTest_ExecuteSparse : FEntityTestBase
 
 		for (FMassEntityHandle& Entity : EntitiesToProcess)
 		{
-			const FTestFragment_Float& TestedFragment = EntitySubsystem->GetFragmentDataChecked<FTestFragment_Float>(Entity);
+			const FTestFragment_Float& TestedFragment = EntityManager->GetFragmentDataChecked<FTestFragment_Float>(Entity);
 			AITEST_EQUAL("Every fragment value should have changed to the expected value", TestedFragment.Value, 13.f);
 		}
 		
 		for (FMassEntityHandle& Entity : EntitiesToIgnore)
 		{
-			const FTestFragment_Float& TestedFragment = EntitySubsystem->GetFragmentDataChecked<FTestFragment_Float>(Entity);
+			const FTestFragment_Float& TestedFragment = EntityManager->GetFragmentDataChecked<FTestFragment_Float>(Entity);
 			AITEST_EQUAL("Untouched entites should retain default fragment value ", TestedFragment.Value, 0.f);
 		}
 
@@ -252,15 +252,15 @@ struct FQueryTest_TagPresent : FEntityTestBase
 {
 	virtual bool InstantTest() override
 	{
-		CA_ASSUME(EntitySubsystem);
+		CA_ASSUME(EntityManager);
 
 		TArray<const UScriptStruct*> Fragments = {FTestFragment_Float::StaticStruct(), FTestFragment_Tag::StaticStruct()};
-		const FMassArchetypeHandle FloatsTagArchetype = EntitySubsystem->CreateArchetype(Fragments);
+		const FMassArchetypeHandle FloatsTagArchetype = EntityManager->CreateArchetype(Fragments);
 
 		FMassEntityQuery Query;
 		Query.AddRequirement<FTestFragment_Float>(EMassFragmentAccess::ReadWrite);
 		Query.AddTagRequirement<FTestFragment_Tag>(EMassFragmentPresence::All);
-		Query.CacheArchetypes(*EntitySubsystem);
+		Query.CacheArchetypes(*EntityManager);
 
 		AITEST_EQUAL("There's a single archetype matching the requirements", Query.GetArchetypes().Num(), 1);
 		AITEST_TRUE("The only valid archetype is FloatsTagArchetype", FloatsTagArchetype == Query.GetArchetypes()[0]);
@@ -275,15 +275,15 @@ struct FQueryTest_TagAbsent : FEntityTestBase
 {
 	virtual bool InstantTest() override
 	{
-		CA_ASSUME(EntitySubsystem);
+		CA_ASSUME(EntityManager);
 
 		TArray<const UScriptStruct*> Fragments = { FTestFragment_Float::StaticStruct(), FTestFragment_Tag::StaticStruct() };
-		const FMassArchetypeHandle FloatsTagArchetype = EntitySubsystem->CreateArchetype(Fragments);
+		const FMassArchetypeHandle FloatsTagArchetype = EntityManager->CreateArchetype(Fragments);
 
 		FMassEntityQuery Query;
 		Query.AddRequirement<FTestFragment_Float>(EMassFragmentAccess::ReadWrite);
 		Query.AddTagRequirement<FTestFragment_Tag>(EMassFragmentPresence::None);
-		Query.CacheArchetypes(*EntitySubsystem);
+		Query.CacheArchetypes(*EntityManager);
 
 		AITEST_EQUAL("There are exactly two archetypes matching the requirements", Query.GetArchetypes().Num(), 2);
 		AITEST_TRUE("FloatsTagArchetype is not amongst matching archetypes"
@@ -300,11 +300,11 @@ struct FQueryTest_FragmentPresent : FEntityTestBase
 {
 	virtual bool InstantTest() override
 	{
-		CA_ASSUME(EntitySubsystem);
+		CA_ASSUME(EntityManager);
 
 		FMassEntityQuery Query;
 		Query.AddRequirement<FTestFragment_Int>(EMassFragmentAccess::ReadOnly);
-		Query.CacheArchetypes(*EntitySubsystem);
+		Query.CacheArchetypes(*EntityManager);
 
 		AITEST_EQUAL("There are exactly two archetypes matching the requirements", Query.GetArchetypes().Num(), 2);
 		AITEST_TRUE("FloatsArchetype is not amongst matching archetypes"
@@ -320,7 +320,7 @@ struct FQueryTest_OnlySingleAbsentFragment : FEntityTestBase
 {
 	virtual bool InstantTest() override
 	{
-		CA_ASSUME(EntitySubsystem);
+		CA_ASSUME(EntityManager);
 
 		FMassEntityQuery Query;
 		Query.AddRequirement<FTestFragment_Int>(EMassFragmentAccess::ReadOnly, EMassFragmentPresence::None);
@@ -328,7 +328,7 @@ struct FQueryTest_OnlySingleAbsentFragment : FEntityTestBase
 		AITEST_FALSE("The query is not valid", Query.CheckValidity());
 
 		GetTestRunner().AddExpectedError(TEXT("requirements not valid"), EAutomationExpectedErrorFlags::Contains, 1);
-		Query.CacheArchetypes(*EntitySubsystem);
+		Query.CacheArchetypes(*EntityManager);
 
 		// this is an invalid query. We expect an error and an empty valid archetypes array
 		AITEST_EQUAL("The query is invalid and we expect no archetypes match", Query.GetArchetypes().Num(), 0);
@@ -343,7 +343,7 @@ struct FQueryTest_OnlyMultipleAbsentFragments : FEntityTestBase
 {
 	virtual bool InstantTest() override
 	{
-		CA_ASSUME(EntitySubsystem);
+		CA_ASSUME(EntityManager);
 
 		FMassEntityQuery Query;
 		Query.AddRequirement<FTestFragment_Int>(EMassFragmentAccess::ReadOnly, EMassFragmentPresence::None);
@@ -352,7 +352,7 @@ struct FQueryTest_OnlyMultipleAbsentFragments : FEntityTestBase
 		AITEST_FALSE("The query is not valid", Query.CheckValidity());
 
 		GetTestRunner().AddExpectedError(TEXT("requirements not valid"), EAutomationExpectedErrorFlags::Contains, 1);
-		Query.CacheArchetypes(*EntitySubsystem);
+		Query.CacheArchetypes(*EntityManager);
 
 		// this is an invalid query. We expect an error and an empty valid archetypes array
 		AITEST_EQUAL("The query is invalid and we expect no archetypes match", Query.GetArchetypes().Num(), 0);
@@ -367,14 +367,14 @@ struct FQueryTest_AbsentAndPresentFragments : FEntityTestBase
 {
 	virtual bool InstantTest() override
 	{
-		CA_ASSUME(EntitySubsystem);
+		CA_ASSUME(EntityManager);
 
 		FMassEntityQuery Query;
 		Query.AddRequirement<FTestFragment_Int>(EMassFragmentAccess::None, EMassFragmentPresence::None);
 		Query.AddRequirement<FTestFragment_Float>(EMassFragmentAccess::ReadOnly, EMassFragmentPresence::All);
 
 		AITEST_TRUE("The query is valid", Query.CheckValidity());
-		Query.CacheArchetypes(*EntitySubsystem);
+		Query.CacheArchetypes(*EntityManager);
 		AITEST_EQUAL("There is only one archetype matching the query", Query.GetArchetypes().Num(), 1);
 		AITEST_TRUE("FloatsArchetype is the only one matching the query", FloatsArchetype == Query.GetArchetypes()[0]);
 
@@ -388,11 +388,11 @@ struct FQueryTest_SingleOptionalFragment : FEntityTestBase
 {
 	virtual bool InstantTest() override
 	{
-		CA_ASSUME(EntitySubsystem);
+		CA_ASSUME(EntityManager);
 
 		FMassEntityQuery Query;
 		Query.AddRequirement<FTestFragment_Int>(EMassFragmentAccess::ReadWrite, EMassFragmentPresence::Optional);
-		Query.CacheArchetypes(*EntitySubsystem);
+		Query.CacheArchetypes(*EntityManager);
 
 		AITEST_EQUAL("There are exactly two archetypes matching the requirements", Query.GetArchetypes().Num(), 2);
 		AITEST_TRUE("FloatsArchetype is not amongst matching archetypes"
@@ -408,12 +408,12 @@ struct FQueryTest_MultipleOptionalFragment : FEntityTestBase
 {
 	virtual bool InstantTest() override
 	{
-		CA_ASSUME(EntitySubsystem);
+		CA_ASSUME(EntityManager);
 
 		FMassEntityQuery Query;
 		Query.AddRequirement<FTestFragment_Int>(EMassFragmentAccess::ReadWrite, EMassFragmentPresence::Optional);
 		Query.AddRequirement<FTestFragment_Float>(EMassFragmentAccess::ReadWrite, EMassFragmentPresence::Optional);
-		Query.CacheArchetypes(*EntitySubsystem);
+		Query.CacheArchetypes(*EntityManager);
 
 		AITEST_EQUAL("All three archetype meet requirements", Query.GetArchetypes().Num(), 3);
 		return true;
@@ -428,11 +428,11 @@ struct FQueryTest_UsingOptionalFragment : FEntityTestBase
 {
 	virtual bool InstantTest() override
 	{
-		CA_ASSUME(EntitySubsystem);
+		CA_ASSUME(EntityManager);
 
-		EntitySubsystem->CreateEntity(FloatsArchetype);
-		const FMassEntityHandle EntityWithFloatsInts = EntitySubsystem->CreateEntity(FloatsIntsArchetype);
-		EntitySubsystem->CreateEntity(IntsArchetype);
+		EntityManager->CreateEntity(FloatsArchetype);
+		const FMassEntityHandle EntityWithFloatsInts = EntityManager->CreateEntity(FloatsIntsArchetype);
+		EntityManager->CreateEntity(IntsArchetype);
 
 		const int32 IntValueSet = 123;
 		int TotalProcessed = 0;
@@ -442,7 +442,7 @@ struct FQueryTest_UsingOptionalFragment : FEntityTestBase
 		Query.AddRequirement<FTestFragment_Int>(EMassFragmentAccess::ReadWrite, EMassFragmentPresence::Optional);
 		Query.AddRequirement<FTestFragment_Float>(EMassFragmentAccess::ReadWrite, EMassFragmentPresence::All);
 		FMassExecutionContext ExecContext;
-		Query.ForEachEntityChunk(*EntitySubsystem, ExecContext, [&TotalProcessed, &EmptyIntsViewCount, IntValueSet](FMassExecutionContext& Context) {
+		Query.ForEachEntityChunk(*EntityManager, ExecContext, [&TotalProcessed, &EmptyIntsViewCount, IntValueSet](FMassExecutionContext& Context) {
 			++TotalProcessed;
 			TArrayView<FTestFragment_Int> Ints = Context.GetMutableFragmentView<FTestFragment_Int>();
 			if (Ints.Num() == 0)
@@ -461,7 +461,7 @@ struct FQueryTest_UsingOptionalFragment : FEntityTestBase
 		AITEST_EQUAL("Two archetypes total should get processed", TotalProcessed, 2);
 		AITEST_EQUAL("Only one of these archetypes should get an empty Ints array view", EmptyIntsViewCount, 1);
 
-		const FTestFragment_Int& TestFragment = EntitySubsystem->GetFragmentDataChecked<FTestFragment_Int>(EntityWithFloatsInts);
+		const FTestFragment_Int& TestFragment = EntityManager->GetFragmentDataChecked<FTestFragment_Int>(EntityWithFloatsInts);
 		AITEST_TRUE("The optional fragment\'s value should get modified where present", TestFragment.Value == IntValueSet);
 
 		return true;
@@ -474,32 +474,32 @@ struct FQueryTest_AnyFragment : FEntityTestBase
 {
 	virtual bool InstantTest() override
 	{
-		CA_ASSUME(EntitySubsystem);
+		CA_ASSUME(EntityManager);
 
 		// From FEntityTestBase:
 		// FMassArchetypeHandle FloatsArchetype;
 		// FMassArchetypeHandle IntsArchetype;
 		// FMassArchetypeHandle FloatsIntsArchetype;
-		const FMassArchetypeHandle BoolArchetype = EntitySubsystem->CreateArchetype({ FTestFragment_Bool::StaticStruct() });
-		const FMassArchetypeHandle BoolFloatArchetype = EntitySubsystem->CreateArchetype({ FTestFragment_Bool::StaticStruct(), FTestFragment_Float::StaticStruct() });
+		const FMassArchetypeHandle BoolArchetype = EntityManager->CreateArchetype({ FTestFragment_Bool::StaticStruct() });
+		const FMassArchetypeHandle BoolFloatArchetype = EntityManager->CreateArchetype({ FTestFragment_Bool::StaticStruct(), FTestFragment_Float::StaticStruct() });
 
 		FMassEntityQuery Query;
 		Query.AddRequirement<FTestFragment_Int>(EMassFragmentAccess::ReadWrite, EMassFragmentPresence::Any);
 		Query.AddRequirement<FTestFragment_Bool>(EMassFragmentAccess::ReadWrite, EMassFragmentPresence::Any);
 		// this query should match: 
 		// IntsArchetype, FloatsIntsArchetype, BoolArchetype, BoolFloatArchetype
-		Query.CacheArchetypes(*EntitySubsystem);
+		Query.CacheArchetypes(*EntityManager);
 
 		AITEST_EQUAL("Archetypes containing Int or Bool should meet requirements", Query.GetArchetypes().Num(), 4);
 
 		// populate the archetypes so that we can test fragment binding
 		for (auto ArchetypeHandle : Query.GetArchetypes())
 		{
-			EntitySubsystem->CreateEntity(ArchetypeHandle);
+			EntityManager->CreateEntity(ArchetypeHandle);
 		}
 
 		FMassExecutionContext TestContext;
-		Query.ForEachEntityChunk(*EntitySubsystem, TestContext, [this](FMassExecutionContext& Context)
+		Query.ForEachEntityChunk(*EntityManager, TestContext, [this](FMassExecutionContext& Context)
 			{
 				TArrayView<FTestFragment_Bool> BoolView = Context.GetMutableFragmentView<FTestFragment_Bool>();
 				TArrayView<FTestFragment_Int> IntView = Context.GetMutableFragmentView<FTestFragment_Int>();
@@ -517,13 +517,13 @@ struct FQueryTest_AnyTag : FEntityTestBase
 {
 	virtual bool InstantTest() override
 	{
-		CA_ASSUME(EntitySubsystem);
+		CA_ASSUME(EntityManager);
 
-		const FMassArchetypeHandle ABArchetype = EntitySubsystem->CreateArchetype({ FTestFragment_Int::StaticStruct(), FTestTag_A::StaticStruct(), FTestTag_B::StaticStruct() });
-		const FMassArchetypeHandle ACArchetype = EntitySubsystem->CreateArchetype({ FTestFragment_Int::StaticStruct(), FTestTag_A::StaticStruct(), FTestTag_C::StaticStruct() });
-		const FMassArchetypeHandle BCArchetype = EntitySubsystem->CreateArchetype({ FTestFragment_Int::StaticStruct(), FTestTag_B::StaticStruct(), FTestTag_C::StaticStruct() });
-		const FMassArchetypeHandle BDArchetype = EntitySubsystem->CreateArchetype({ FTestFragment_Int::StaticStruct(), FTestTag_B::StaticStruct(), FTestTag_D::StaticStruct() });
-		const FMassArchetypeHandle FloatACArchetype = EntitySubsystem->CreateArchetype({ FTestFragment_Float::StaticStruct(), FTestTag_A::StaticStruct(), FTestTag_C::StaticStruct() });
+		const FMassArchetypeHandle ABArchetype = EntityManager->CreateArchetype({ FTestFragment_Int::StaticStruct(), FTestTag_A::StaticStruct(), FTestTag_B::StaticStruct() });
+		const FMassArchetypeHandle ACArchetype = EntityManager->CreateArchetype({ FTestFragment_Int::StaticStruct(), FTestTag_A::StaticStruct(), FTestTag_C::StaticStruct() });
+		const FMassArchetypeHandle BCArchetype = EntityManager->CreateArchetype({ FTestFragment_Int::StaticStruct(), FTestTag_B::StaticStruct(), FTestTag_C::StaticStruct() });
+		const FMassArchetypeHandle BDArchetype = EntityManager->CreateArchetype({ FTestFragment_Int::StaticStruct(), FTestTag_B::StaticStruct(), FTestTag_D::StaticStruct() });
+		const FMassArchetypeHandle FloatACArchetype = EntityManager->CreateArchetype({ FTestFragment_Float::StaticStruct(), FTestTag_A::StaticStruct(), FTestTag_C::StaticStruct() });
 
 		FMassEntityQuery Query;
 		// at least one fragment requirement needs to be present for the query to be valid
@@ -532,7 +532,7 @@ struct FQueryTest_AnyTag : FEntityTestBase
 		Query.AddTagRequirement<FTestTag_C>(EMassFragmentPresence::Any);
 		// this query should match: 
 		// ABArchetype, ACArchetype and ABCrchetype but not BDArchetype nor FEntityTestBase.IntsArchetype
-		Query.CacheArchetypes(*EntitySubsystem);
+		Query.CacheArchetypes(*EntityManager);
 
 		AITEST_EQUAL("Only Archetypes tagged with A or C should matched the query", Query.GetArchetypes().Num(), 3);
 		AITEST_TRUE("ABArchetype should be amongst the matched archetypes", Query.GetArchetypes().Find(ABArchetype) != INDEX_NONE);
diff --git a/Engine/Plugins/Runtime/MassEntity/Source/MassEntityTestSuite/Private/MassSystemRequirementTest.cpp b/Engine/Plugins/Runtime/MassEntity/Source/MassEntityTestSuite/Private/MassSystemRequirementTest.cpp
--- a/Engine/Plugins/Runtime/MassEntity/Source/MassEntityTestSuite/Private/MassSystemRequirementTest.cpp
+++ b/Engine/Plugins/Runtime/MassEntity/Source/MassEntityTestSuite/Private/MassSystemRequirementTest.cpp
@@ -14,7 +14,7 @@
 PRAGMA_DISABLE_OPTIMIZATION
 
 //template<>
-//struct TMassExternalSystemGetter<UMassEntitySubsystem>
+//struct TMassExternalSystemGetter<FMassEntityManager>
 //{
 //	
 //};
@@ -22,8 +22,7 @@ PRAGMA_DISABLE_OPTIMIZATION
 namespace FMassSystemRequirementTest
 {
 
-//FEntityTestBase : FExecutionTestBase
-struct FSystemRequirementTestBase : FEntityTestBase//FExecutionTestBase
+struct FSystemRequirementTestBase : FEntityTestBase
 {
 	using Super = FEntityTestBase;
 	virtual bool SetUp() override
@@ -36,20 +35,19 @@ struct FMutableRequirement : FSystemRequirementTestBase
 {
 	virtual bool InstantTest() override
 	{
-		EntitySubsystem->CreateEntity(FloatsArchetype);
+		EntityManager->CreateEntity(FloatsArchetype);
 
 		FMassEntityQuery EntityQuery;
 		EntityQuery.AddRequirement<FTestFragment_Float>(EMassFragmentAccess::ReadOnly);
 		EntityQuery.AddSubsystemRequirement<UMassTestWorldSubsystem>(EMassFragmentAccess::ReadWrite);
 
-		check(EntitySubsystem->GetWorld());
-		UMassTestWorldSubsystem* TestSystemActual = EntitySubsystem->GetWorld()->GetSubsystem<UMassTestWorldSubsystem>();
+		check(World);
+		UMassTestWorldSubsystem* TestSystemActual = World->GetSubsystem<UMassTestWorldSubsystem>();
 		UMassTestWorldSubsystem* TestSystem = nullptr;
 		const UMassTestWorldSubsystem* TestSystem2 = nullptr;
 		FMassExecutionContext ExecutionContext;
-		EntityQuery.ForEachEntityChunk(*EntitySubsystem, ExecutionContext, [this, &TestSystem, &TestSystem2, World = EntitySubsystem->GetWorld()](FMassExecutionContext& Context)
+		EntityQuery.ForEachEntityChunk(*EntityManager, ExecutionContext, [this, &TestSystem, &TestSystem2](FMassExecutionContext& Context)
 		{
-			//UMassEntitySubsystem* NavSystem = Context.GetSubsystem<UMassEntitySubsystem>(World);
 			TestSystem = Context.GetMutableSubsystem<UMassTestWorldSubsystem>(World);
 			TestSystem2 = Context.GetSubsystem<UMassTestWorldSubsystem>(World);
 		});
@@ -67,20 +65,19 @@ struct FConstRequirement : FSystemRequirementTestBase
 {
 	virtual bool InstantTest() override
 	{
-		EntitySubsystem->CreateEntity(FloatsArchetype);
+		EntityManager->CreateEntity(FloatsArchetype);
 
 		FMassEntityQuery EntityQuery;
 		EntityQuery.AddRequirement<FTestFragment_Float>(EMassFragmentAccess::ReadOnly);
 		EntityQuery.AddSubsystemRequirement<UMassTestWorldSubsystem>(EMassFragmentAccess::ReadOnly);
 
-		check(EntitySubsystem->GetWorld());
+		check(World);
 		UMassTestWorldSubsystem* TestMutableSystem = nullptr;
 		const UMassTestWorldSubsystem* TestConstSystem = nullptr;
 		FMassExecutionContext ExecutionContext;
-		EntityQuery.ForEachEntityChunk(*EntitySubsystem, ExecutionContext, [this, &TestMutableSystem, &TestConstSystem, World = EntitySubsystem->GetWorld()](FMassExecutionContext& Context)
+		EntityQuery.ForEachEntityChunk(*EntityManager, ExecutionContext, [this, &TestMutableSystem, &TestConstSystem](FMassExecutionContext& Context)
 		{
 			// commented out since there's no way to seamlessly handle failed ensures, and the line below should cause that.
-			//TestMutableSystem = Context.GetMutableSubsystem<UMassTestWorldSubsystem>(World);
 			TestConstSystem = Context.GetSubsystem<UMassTestWorldSubsystem>(World);
 		});
 
diff --git a/Engine/Plugins/Runtime/MassEntity/Source/MassEntityTestSuite/Private/MassTestTypes.cpp b/Engine/Plugins/Runtime/MassEntity/Source/MassEntityTestSuite/Private/MassTestTypes.cpp
--- a/Engine/Plugins/Runtime/MassEntity/Source/MassEntityTestSuite/Private/MassTestTypes.cpp
+++ b/Engine/Plugins/Runtime/MassEntity/Source/MassEntityTestSuite/Private/MassTestTypes.cpp
@@ -11,14 +11,8 @@
 bool FExecutionTestBase::SetUp()
 {
 	World = FAITestHelpers::GetWorld();
-	EntitySubsystem = NewObject<UMassEntitySubsystem>(World);
-	check(EntitySubsystem);
-	struct FSubsystemCollection_TestInit : FSubsystemCollectionBase
-	{
-		FSubsystemCollection_TestInit(){}
-	};
-	FSubsystemCollection_TestInit Collection;
-	EntitySubsystem->Initialize(Collection);
+	EntityManager = MakeShareable(new FMassEntityManager);
+	EntityManager->Initialize();
 
 	return true;
 }
@@ -26,14 +20,14 @@ bool FExecutionTestBase::SetUp()
 bool FEntityTestBase::SetUp()
 {
 	FExecutionTestBase::SetUp();
-	check(EntitySubsystem);
+	check(EntityManager);
 
 	const UScriptStruct* FragmentTypes[] = { FTestFragment_Float::StaticStruct(), FTestFragment_Int::StaticStruct() };
 
-	EmptyArchetype = EntitySubsystem->CreateArchetype(MakeArrayView<const UScriptStruct*>(nullptr, 0));
-	FloatsArchetype = EntitySubsystem->CreateArchetype(MakeArrayView(&FragmentTypes[0], 1));
-	IntsArchetype = EntitySubsystem->CreateArchetype(MakeArrayView(&FragmentTypes[1], 1));
-	FloatsIntsArchetype = EntitySubsystem->CreateArchetype(MakeArrayView(FragmentTypes, 2));
+	EmptyArchetype = EntityManager->CreateArchetype(MakeArrayView<const UScriptStruct*>(nullptr, 0));
+	FloatsArchetype = EntityManager->CreateArchetype(MakeArrayView(&FragmentTypes[0], 1));
+	IntsArchetype = EntityManager->CreateArchetype(MakeArrayView(&FragmentTypes[1], 1));
+	FloatsIntsArchetype = EntityManager->CreateArchetype(MakeArrayView(FragmentTypes, 2));
 
 	FTestFragment_Int IntFrag;
 	IntFrag.Value = TestIntValue;
@@ -55,7 +49,7 @@ UMassTestProcessorBase::UMassTestProcessorBase()
 	bAutoRegisterWithProcessingPhases = false;
 	ExecutionFlags = int32(EProcessorExecutionFlags::All);
 
-	ExecutionFunction = [](UMassEntitySubsystem& InEntitySubsystem, FMassExecutionContext& Context) {};
+	ExecutionFunction = [](FMassEntityManager& InEntitySubsystem, FMassExecutionContext& Context) {};
 	RequirementsFunction = [](FMassEntityQuery& Query){};
 }
 
diff --git a/Engine/Plugins/Runtime/MassEntity/Source/MassEntityTestSuite/Public/LWComponentTestFarmPlot.h b/Engine/Plugins/Runtime/MassEntity/Source/MassEntityTestSuite/Public/LWComponentTestFarmPlot.h
--- a/Engine/Plugins/Runtime/MassEntity/Source/MassEntityTestSuite/Public/LWComponentTestFarmPlot.h
+++ b/Engine/Plugins/Runtime/MassEntity/Source/MassEntityTestSuite/Public/LWComponentTestFarmPlot.h
@@ -101,7 +101,7 @@ class UFragmentUpdateSystem : public UObject
 	GENERATED_BODY()
 public:
 	virtual void PostInitProperties() override;
-	virtual void Execute(UMassEntitySubsystem& EntitySubsystem, FMassExecutionContext& Context) {}
+	virtual void Execute(FMassEntityManager& EntityManager, FMassExecutionContext& Context) {}
 protected:
 	virtual void ConfigureQueries() {}
 protected:
@@ -120,10 +120,10 @@ public:
 		EntityQuery.AddRequirement<FFarmWaterFragment>(EMassFragmentAccess::ReadWrite);
 	}
 
-	virtual void Execute(UMassEntitySubsystem& EntitySubsystem, FMassExecutionContext& Context) override
+	virtual void Execute(FMassEntityManager& EntityManager, FMassExecutionContext& Context) override
 	{
 		QUICK_SCOPE_CYCLE_COUNTER(UFarmWaterUpdateSystem_Run);
-		EntityQuery.ForEachEntityChunk(EntitySubsystem, Context, [this](FMassExecutionContext& Context) {
+		EntityQuery.ForEachEntityChunk(EntityManager, Context, [this](FMassExecutionContext& Context) {
 
 			TArrayView<FFarmWaterFragment> WaterList = Context.GetMutableFragmentView<FFarmWaterFragment>();
 
@@ -149,11 +149,11 @@ public:
 		EntityQuery.AddRequirement<FFarmFlowerFragment>(EMassFragmentAccess::ReadWrite);
 	}
 
-	virtual void Execute(UMassEntitySubsystem& EntitySubsystem, FMassExecutionContext& Context) override
+	virtual void Execute(FMassEntityManager& EntityManager, FMassExecutionContext& Context) override
 	{
 		QUICK_SCOPE_CYCLE_COUNTER(UFarmHarvestTimerSystem_Flowers_Run);
 		
-		EntityQuery.ForEachEntityChunk(EntitySubsystem, Context, [this](FMassExecutionContext& Context) {
+		EntityQuery.ForEachEntityChunk(EntityManager, Context, [this](FMassExecutionContext& Context) {
 			const int32 NumEntities = Context.GetNumEntities();
 			const float WellWateredThreshold = 0.25f;
 			TArrayView<FHarvestTimerFragment> TimerList = Context.GetMutableFragmentView<FHarvestTimerFragment>();
@@ -192,11 +192,11 @@ public:
 		EntityQuery.AddRequirement<FFarmCropFragment>(EMassFragmentAccess::ReadOnly);
 	}
 
-	virtual void Execute(UMassEntitySubsystem& EntitySubsystem, FMassExecutionContext& Context) override
+	virtual void Execute(FMassEntityManager& EntityManager, FMassExecutionContext& Context) override
 	{
 		QUICK_SCOPE_CYCLE_COUNTER(UFarmHarvestTimerSystem_Crops_Run);
 		
-		EntityQuery.ForEachEntityChunk(EntitySubsystem, Context, [this](FMassExecutionContext& Context) {
+		EntityQuery.ForEachEntityChunk(EntityManager, Context, [this](FMassExecutionContext& Context) {
 			
 			const int32 NumEntities = Context.GetNumEntities();
 			const float WellWateredThreshold = 0.25f;
@@ -227,11 +227,11 @@ public:
 		EntityQuery.AddTagRequirement<FFarmReadyToHarvestTag>(EMassFragmentPresence::None);
 	}
 
-	virtual void Execute(UMassEntitySubsystem& EntitySubsystem, FMassExecutionContext& Context) override
+	virtual void Execute(FMassEntityManager& EntityManager, FMassExecutionContext& Context) override
 	{
 		QUICK_SCOPE_CYCLE_COUNTER(UFarmHarvestTimerExpired_Run);
 		
-		EntityQuery.ForEachEntityChunk(EntitySubsystem, Context, [this](FMassExecutionContext& Context) {
+		EntityQuery.ForEachEntityChunk(EntityManager, Context, [this](FMassExecutionContext& Context) {
 			const int32 NumEntities = Context.GetNumEntities();
 			TConstArrayView<FHarvestTimerFragment> TimerList = Context.GetFragmentView<FHarvestTimerFragment>();
 
@@ -269,7 +269,7 @@ public:
 		EntityQuery.AddTagRequirement<FFarmJustBecameReadyToHarvestTag>(EMassFragmentPresence::All);
 	}
 
-	virtual void Execute(UMassEntitySubsystem& EntitySubsystem, FMassExecutionContext& Context) override;
+	virtual void Execute(FMassEntityManager& EntityManager, FMassExecutionContext& Context) override;
 };
 
 //////////////////////////////////////////////////////////////////////
@@ -337,7 +337,7 @@ private:
 	TArray<int32> FreeHarvestIconIndicies;
 
 private:
-	void AddItemToGrid(UMassEntitySubsystem* EntitySystem, uint16 X, uint16 Y, FMassArchetypeHandle Archetype, uint16 VisualIndex);
+	void AddItemToGrid(FMassEntityManager& EntityManager, uint16 X, uint16 Y, FMassArchetypeHandle Archetype, uint16 VisualIndex);
 
 public:
 	ALWFragmentTestFarmPlot();
diff --git a/Engine/Plugins/Runtime/MassEntity/Source/MassEntityTestSuite/Public/MassEntityTestTypes.h b/Engine/Plugins/Runtime/MassEntity/Source/MassEntityTestSuite/Public/MassEntityTestTypes.h
--- a/Engine/Plugins/Runtime/MassEntity/Source/MassEntityTestSuite/Public/MassEntityTestTypes.h
+++ b/Engine/Plugins/Runtime/MassEntity/Source/MassEntityTestSuite/Public/MassEntityTestTypes.h
@@ -14,7 +14,7 @@
 
 
 class UWorld;
-class UMassEntitySubsystem;
+struct FMassEntityManager;
 
 USTRUCT()
 struct FTestFragment_Float : public FMassFragment
@@ -106,16 +106,16 @@ class MASSENTITYTESTSUITE_API UMassTestProcessorBase : public UMassProcessor
 public:
 	UMassTestProcessorBase();
 	FMassProcessorExecutionOrder& GetMutableExecutionOrder() { return ExecutionOrder; }
-	virtual void Execute(UMassEntitySubsystem& EntitySubsystem, FMassExecutionContext& Context) override 
+	virtual void Execute(FMassEntityManager& EntityManager, FMassExecutionContext& Context) override 
 	{
-		ExecutionFunction(EntitySubsystem, Context);
+		ExecutionFunction(EntityManager, Context);
 	}
 	virtual void ConfigureQueries() override
 	{
 		RequirementsFunction(EntityQuery);
 	}
 
-	TFunction<void(UMassEntitySubsystem& EntitySubsystem, FMassExecutionContext& Context)> ExecutionFunction;
+	TFunction<void(FMassEntityManager& EntityManager, FMassExecutionContext& Context)> ExecutionFunction;
 	TFunction<void(FMassEntityQuery& Query)> RequirementsFunction;
 
 	FMassEntityQuery& TestGetQuery() { return EntityQuery; }
@@ -190,7 +190,7 @@ public:
 
 struct MASSENTITYTESTSUITE_API FExecutionTestBase : FAITestBase
 {
-	UMassEntitySubsystem* EntitySubsystem = nullptr;
+	TSharedPtr<FMassEntityManager> EntityManager;
 	UWorld* World = nullptr;
 
 	virtual bool SetUp() override;
diff --git a/Engine/Plugins/Runtime/MassGameplay/Source/MassActors/Private/MassActorHelper.cpp b/Engine/Plugins/Runtime/MassGameplay/Source/MassActors/Private/MassActorHelper.cpp
--- a/Engine/Plugins/Runtime/MassGameplay/Source/MassActors/Private/MassActorHelper.cpp
+++ b/Engine/Plugins/Runtime/MassGameplay/Source/MassActors/Private/MassActorHelper.cpp
@@ -11,14 +11,15 @@ namespace UE::MassActor
 	bool AddEntityTagToActor(const AActor& Actor, const UScriptStruct& TagType)
 	{
 		UWorld* World = Actor.GetWorld();
-		UMassEntitySubsystem* EntitySystem = UWorld::GetSubsystem<UMassEntitySubsystem>(World);
+		UMassEntitySubsystem* EntitySubsystem = UWorld::GetSubsystem<UMassEntitySubsystem>(World);
 		UMassActorSubsystem* MassActorSubsystem = UWorld::GetSubsystem<UMassActorSubsystem>(World);
-		if (EntitySystem && MassActorSubsystem)
+		if (EntitySubsystem && MassActorSubsystem)
 		{
+			FMassEntityManager& EntityManager = EntitySubsystem->GetMutableEntityManager();
 			const FMassEntityHandle AgentHandle = MassActorSubsystem->GetEntityHandleFromActor(&Actor);
 			if (AgentHandle.IsValid())
 			{
-				EntitySystem->AddTagToEntity(AgentHandle, &TagType);
+				EntityManager.AddTagToEntity(AgentHandle, &TagType);
 				return true;
 			}
 			UE_LOG(LogMassActor, Warning, TEXT("Failed to add tag %s to actor %s due to it not having an associated entity")
@@ -27,7 +28,7 @@ namespace UE::MassActor
 		else
 		{
 			UE_LOG(LogMassActor, Warning, TEXT("Failed to add tag %s to actor %s due to missing: %s%s")
-				, *TagType.GetName(), *Actor.GetName(), EntitySystem ? TEXT("EntitySubsystem, ") : TEXT("")
+				, *TagType.GetName(), *Actor.GetName(), EntitySubsystem ? TEXT("EntitySubsystem, ") : TEXT("")
 				, MassActorSubsystem ? TEXT("MassActorSubsystem, ") : TEXT(""));
 		}
 		return false;
@@ -36,14 +37,15 @@ namespace UE::MassActor
 	bool RemoveEntityTagFromActor(const AActor& Actor, const UScriptStruct& TagType)
 	{
 		UWorld* World = Actor.GetWorld();
-		UMassEntitySubsystem* EntitySystem = UWorld::GetSubsystem<UMassEntitySubsystem>(World);
+		UMassEntitySubsystem* EntitySubsystem = UWorld::GetSubsystem<UMassEntitySubsystem>(World);
 		UMassActorSubsystem* MassActorSubsystem = UWorld::GetSubsystem<UMassActorSubsystem>(World);
-		if (EntitySystem && MassActorSubsystem)
+		if (EntitySubsystem && MassActorSubsystem)
 		{
 			const FMassEntityHandle AgentHandle = MassActorSubsystem->GetEntityHandleFromActor(&Actor);
 			if (AgentHandle.IsValid())
 			{
-				EntitySystem->RemoveTagFromEntity(AgentHandle, &TagType);
+				FMassEntityManager& EntityManager = EntitySubsystem->GetMutableEntityManager();
+				EntityManager.RemoveTagFromEntity(AgentHandle, &TagType);
 				return true;
 			}
 			UE_LOG(LogMassActor, Warning, TEXT("Failed to remove tag %s from actor %s due to it not having an associated entity")
@@ -52,7 +54,7 @@ namespace UE::MassActor
 		else
 		{
 			UE_LOG(LogMassActor, Warning, TEXT("Failed to remove tag %s from actor %s due to missing: %s%s")
-				, *TagType.GetName(), *Actor.GetName(), EntitySystem ? TEXT("EntitySubsystem, ") : TEXT("")
+				, *TagType.GetName(), *Actor.GetName(), EntitySubsystem ? TEXT("EntitySubsystem, ") : TEXT("")
 				, MassActorSubsystem ? TEXT("MassActorSubsystem, ") : TEXT(""));
 		}
 		return false;
diff --git a/Engine/Plugins/Runtime/MassGameplay/Source/MassActors/Private/MassActorSubsystem.cpp b/Engine/Plugins/Runtime/MassGameplay/Source/MassActors/Private/MassActorSubsystem.cpp
--- a/Engine/Plugins/Runtime/MassGameplay/Source/MassActors/Private/MassActorSubsystem.cpp
+++ b/Engine/Plugins/Runtime/MassGameplay/Source/MassActors/Private/MassActorSubsystem.cpp
@@ -62,7 +62,10 @@ void UMassActorSubsystem::Initialize(FSubsystemCollectionBase& Collection)
 	// making sure UMassSimulationSubsystem gets created before the MassActorSubsystem
 	Collection.InitializeDependency<UMassSimulationSubsystem>();
 	
-	EntitySystem = UWorld::GetSubsystem<UMassEntitySubsystem>(GetWorld());
+	if (UMassEntitySubsystem* EntitySubsystem = UWorld::GetSubsystem<UMassEntitySubsystem>(GetWorld()))
+	{
+		EntityManager = EntitySubsystem->GetMutableEntityManager().AsShared();
+	}
 }
 
 FMassEntityHandle UMassActorSubsystem::GetEntityHandleFromActor(const TObjectKey<const AActor> Actor)
@@ -71,7 +74,7 @@ FMassEntityHandle UMassActorSubsystem::GetEntityHandleFromActor(const TObjectKey
 	FMassEntityHandle* Entity = ActorHandleMap.Find(Actor);
 	if (!Entity)
 	{
-		return UMassEntitySubsystem::InvalidEntity;
+		return FMassEntityManager::InvalidEntity;
 	}
 
 	check(TObjectKey<const AActor>(GetActorFromHandle(*Entity)) == Actor);
@@ -80,8 +83,8 @@ FMassEntityHandle UMassActorSubsystem::GetEntityHandleFromActor(const TObjectKey
 
 AActor* UMassActorSubsystem::GetActorFromHandle(const FMassEntityHandle Handle) const
 {
-	check(EntitySystem);
-	FMassActorFragment* Data = EntitySystem->GetFragmentDataPtr<FMassActorFragment>(Handle);
+	check(EntityManager);
+	FMassActorFragment* Data = EntityManager->GetFragmentDataPtr<FMassActorFragment>(Handle);
 	return Data != nullptr ? Data->GetMutable() : nullptr;
 }
 
@@ -118,8 +121,8 @@ void UMassActorSubsystem::DisconnectActor(const TObjectKey<const AActor> Actor,
 
 	if (FoundEntity == Handle)
 	{
-		check(EntitySystem);
-		if (FMassActorFragment* Data = EntitySystem->GetFragmentDataPtr<FMassActorFragment>(Handle))
+		check(EntityManager);
+		if (FMassActorFragment* Data = EntityManager->GetFragmentDataPtr<FMassActorFragment>(Handle))
 		{
 			Data->ResetAndUpdateHandleMap();
 		}
diff --git a/Engine/Plugins/Runtime/MassGameplay/Source/MassActors/Private/MassAgentComponent.cpp b/Engine/Plugins/Runtime/MassGameplay/Source/MassActors/Private/MassAgentComponent.cpp
--- a/Engine/Plugins/Runtime/MassGameplay/Source/MassActors/Private/MassAgentComponent.cpp
+++ b/Engine/Plugins/Runtime/MassGameplay/Source/MassActors/Private/MassAgentComponent.cpp
@@ -161,11 +161,12 @@ void UMassAgentComponent::SetEntityHandleInternal(const FMassEntityHandle NewHan
 	ensureMsgf((AgentHandle.IsValid() && NewHandle.IsValid()) == false, TEXT("Overriding an existing entity ID might result in a dangling entity still affecting the simulation"));
 	AgentHandle = NewHandle;
 
+	UMassEntitySubsystem* EntitySubsystem = UWorld::GetSubsystem<UMassEntitySubsystem>(GetWorld());
 #if	UE_REPLICATION_COMPILE_SERVER_CODE
 	// Fetch NetID if it exist
-	if (const UMassEntitySubsystem* EntitySubsystem = UWorld::GetSubsystem<UMassEntitySubsystem>(GetWorld()))
+	if (EntitySubsystem)
 	{
-		if (const FMassNetworkIDFragment* NetIDFragment = EntitySubsystem->GetFragmentDataPtr<FMassNetworkIDFragment>(AgentHandle))
+		if (const FMassNetworkIDFragment* NetIDFragment = EntitySubsystem->GetEntityManager().GetFragmentDataPtr<FMassNetworkIDFragment>(AgentHandle))
 		{
 			if (!IsNetSimulating())
 			{
@@ -185,11 +186,11 @@ void UMassAgentComponent::SetEntityHandleInternal(const FMassEntityHandle NewHan
 	}
 
 	// Sync up with mass
-	if (const UMassEntitySubsystem* EntitySubsystem = UWorld::GetSubsystem<UMassEntitySubsystem>(GetWorld()))
+	if (EntitySubsystem)
 	{
 		if (IsNetSimulating())
 		{
-			const FMassEntityView EntityView(*EntitySubsystem, AgentHandle);
+			const FMassEntityView EntityView(EntitySubsystem->GetEntityManager(), AgentHandle);
 
 			// @todo Find a way to add these initialization into either translator initializer or adding new fragments
 			// Make sure to fetch the fragment after any release, as that action can move the entity around into new archetype and 
@@ -278,7 +279,7 @@ void UMassAgentComponent::ClearEntityHandleInternal()
 	{
 		if (IsNetSimulating())
 		{
-			const FMassEntityView EntityView(*EntitySubsystem, AgentHandle);
+			const FMassEntityView EntityView(EntitySubsystem->GetEntityManager(), AgentHandle);
 			if (FMassActorFragment* ActorInfo = EntityView.GetFragmentDataPtr<FMassActorFragment>())
 			{
 				checkf(!ActorInfo->IsValid() || ActorInfo->Get() == GetOwner(), TEXT("Expecting actor pointer to be the Component\'s owner"));
@@ -287,7 +288,7 @@ void UMassAgentComponent::ClearEntityHandleInternal()
 		}
 	}
 	
-	AgentHandle = UMassEntitySubsystem::InvalidEntity;
+	AgentHandle = FMassEntityManager::InvalidEntity;
 }
 
 void UMassAgentComponent::PuppetUnregistrationDone()
@@ -346,21 +347,22 @@ void UMassAgentComponent::DebugCheckStateConsistency()
 			MASSAGENT_CHECK(bValidAgentHandle, TEXT("Expecting a valid mass agent handle in state %s"), *UEnum::GetValueAsString(State));
 			if (bValidAgentHandle)
 			{
-				if (UMassEntitySubsystem* EntitySubsystem = UWorld::GetSubsystem<UMassEntitySubsystem>(GetWorld()))
+				if (const UMassEntitySubsystem* EntitySubsystem = UWorld::GetSubsystem<UMassEntitySubsystem>(GetWorld()))
 				{
-					const bool bIsValidEntity = EntitySubsystem->IsEntityValid(AgentHandle);
+					const FMassEntityManager& EntityManager = EntitySubsystem->GetEntityManager();
+					const bool bIsValidEntity = EntityManager.IsEntityValid(AgentHandle);
 					MASSAGENT_CHECK(bIsValidEntity, TEXT("Exepecting a valid entity in state"), *UEnum::GetValueAsString(State))
-						if (bIsValidEntity)
+					if (bIsValidEntity)
+					{
+						const bool bIsBuiltEntity = EntityManager.IsEntityBuilt(AgentHandle);
+						MASSAGENT_CHECK(bIsBuiltEntity, TEXT("Expecting a fully built entity in state %s"), *UEnum::GetValueAsString(State));
+						if (bIsBuiltEntity)
 						{
-							const bool bIsBuiltEntity = EntitySubsystem->IsEntityBuilt(AgentHandle);
-							MASSAGENT_CHECK(bIsBuiltEntity, TEXT("Expecting a fully built entity in state %s"), *UEnum::GetValueAsString(State));
-							if (bIsBuiltEntity)
-							{
-								AActor* Owner = GetOwner();
-								const AActor* Actor = EntitySubsystem->GetFragmentDataChecked<FMassActorFragment>(AgentHandle).Get();
-								MASSAGENT_CHECK(Actor == nullptr || Actor == Owner, TEXT("Mass Actor and Owner mismatched in state %s"), *UEnum::GetValueAsString(State));
-							}
+							AActor* Owner = GetOwner();
+							const AActor* Actor = EntityManager.GetFragmentDataChecked<FMassActorFragment>(AgentHandle).Get();
+							MASSAGENT_CHECK(Actor == nullptr || Actor == Owner, TEXT("Mass Actor and Owner mismatched in state %s"), *UEnum::GetValueAsString(State));
 						}
+					}
 				}
 			}
 			break;
@@ -412,8 +414,7 @@ void UMassAgentComponent::KillEntity(const bool bDestroyActor)
 
 	AActor* Owner = GetOwner();
 	UWorld* World = GetWorld();
-	UMassEntitySubsystem* EntitySubsystem = UWorld::GetSubsystem<UMassEntitySubsystem>(World);
-	if (Owner == nullptr || EntitySubsystem == nullptr)
+	if (Owner == nullptr)
 	{
 		return;
 	}
diff --git a/Engine/Plugins/Runtime/MassGameplay/Source/MassActors/Private/MassAgentSubsystem.cpp b/Engine/Plugins/Runtime/MassGameplay/Source/MassActors/Private/MassAgentSubsystem.cpp
--- a/Engine/Plugins/Runtime/MassGameplay/Source/MassActors/Private/MassAgentSubsystem.cpp
+++ b/Engine/Plugins/Runtime/MassGameplay/Source/MassActors/Private/MassAgentSubsystem.cpp
@@ -14,6 +14,7 @@
 #include "MassEntityView.h"
 #include "MassReplicationSubsystem.h"
 #include "Engine/NetDriver.h"
+#include "MassEntityUtils.h"
 
 
 namespace FMassAgentSubsystemHelper
@@ -46,7 +47,10 @@ void UMassAgentSubsystem::Initialize(FSubsystemCollectionBase& Collection)
 #endif // UE_REPLICATION_COMPILE_CLIENT_CODE
 	
 	UWorld* World = GetWorld();
-	EntitySystem = UWorld::GetSubsystem<UMassEntitySubsystem>(World);
+	check(World);
+
+	EntityManager = UE::Mass::Utils::GetEntityManagerChecked(*World).AsShared();
+
 	SpawnerSystem = UWorld::GetSubsystem<UMassSpawnerSubsystem>(World);
 
 	SimulationSystem = UWorld::GetSubsystem<UMassSimulationSubsystem>(World);
@@ -63,7 +67,7 @@ void UMassAgentSubsystem::Initialize(FSubsystemCollectionBase& Collection)
 
 FMassEntityTemplateID UMassAgentSubsystem::RegisterAgentComponent(UMassAgentComponent& AgentComp)
 {
-	check(EntitySystem);
+	check(EntityManager);
 	check(SpawnerSystem);
 
 	if (AgentComp.IsPuppet())
@@ -110,7 +114,7 @@ FMassEntityTemplateID UMassAgentSubsystem::RegisterAgentComponent(UMassAgentComp
 
 void UMassAgentSubsystem::UpdateAgentComponent(const UMassAgentComponent& AgentComp)
 {
-	check(EntitySystem);
+	check(EntityManager);
 	check(SpawnerSystem)
 
 	if (!ensureMsgf(AgentComp.GetEntityHandle().IsValid(), TEXT("Caling %s is valid only for already registered MassAgentComponents"), ANSI_TO_TCHAR(__FUNCTION__)))
@@ -129,7 +133,7 @@ void UMassAgentSubsystem::UpdateAgentComponent(const UMassAgentComponent& AgentC
 	const FMassEntityTemplate& EntityTemplate = EntityConfig.GetOrCreateEntityTemplate(*World, AgentComp);
 
 	const FMassEntityHandle Entity = AgentComp.GetEntityHandle();
-	const FMassArchetypeHandle CurrentArchetypeHandle = EntitySystem->GetArchetypeForEntity(Entity);
+	const FMassArchetypeHandle CurrentArchetypeHandle = EntityManager->GetArchetypeForEntity(Entity);
 	if (CurrentArchetypeHandle == EntityTemplate.GetArchetype())
 	{
 		UE_VLOG(this, LogMassActor, Log, TEXT("%s called for %s but no archetype changes have been found")
@@ -149,7 +153,7 @@ void UMassAgentSubsystem::UpdateAgentComponent(const UMassAgentComponent& AgentC
 
 void UMassAgentSubsystem::UnregisterAgentComponent(UMassAgentComponent& AgentComp)
 {
-	if (EntitySystem == nullptr || SpawnerSystem == nullptr || SimulationSystem == nullptr)
+	if (!EntityManager || SpawnerSystem == nullptr || SimulationSystem == nullptr)
 	{
 		return;
 	}
@@ -174,11 +178,11 @@ void UMassAgentSubsystem::UnregisterAgentComponent(UMassAgentComponent& AgentCom
 
 			    // remove fragments that have been added for the puppet agent
 			    if (SimulationSystem->GetPhaseManager().IsDuringMassProcessing()
-					|| EntitySystem->IsProcessing())
+					|| EntityManager->IsProcessing())
 			    {
 				    // need to request via command buffer since we can't move entities while processing is happening
 					FMassArchetypeCompositionDescriptor Composition = AgentComp.GetPuppetSpecificAddition();
-					EntitySystem->Defer().PushCommand<FMassDeferredRemoveCommand>([Entity, Composition](UMassEntitySubsystem& System)
+					EntityManager->Defer().PushCommand<FMassDeferredRemoveCommand>([Entity, Composition](FMassEntityManager& System)
 						{
 							if (System.IsEntityValid(Entity) == false)
 							{
@@ -189,7 +193,7 @@ void UMassAgentSubsystem::UnregisterAgentComponent(UMassAgentComponent& AgentCom
 			    }
 			    else
 			    {
-					EntitySystem->RemoveCompositionFromEntity(Entity, AgentComp.GetPuppetSpecificAddition());
+					EntityManager->RemoveCompositionFromEntity(Entity, AgentComp.GetPuppetSpecificAddition());
 			    }
 			}
 			AgentComp.PuppetUnregistrationDone();
@@ -215,10 +219,10 @@ void UMassAgentSubsystem::UnregisterAgentComponent(UMassAgentComponent& AgentCom
 
 			// Destroy the entity
 			if (SimulationSystem->GetPhaseManager().IsDuringMassProcessing()
-				|| EntitySystem->IsProcessing())
+				|| EntityManager->IsProcessing())
 			{
 				// need to request via command buffer since we can't move entities while processing is happening
-				EntitySystem->Defer().DestroyEntity(Entity);
+				EntityManager->Defer().DestroyEntity(Entity);
 			}
 			else if (ensure(EntityTemplate))
 			{
@@ -291,7 +295,7 @@ void UMassAgentSubsystem::OnProcessingPhaseStarted(const float DeltaSeconds, con
 
 void UMassAgentSubsystem::HandlePendingInitialization()
 {
-	check(EntitySystem);
+	check(EntityManager);
 	check(SpawnerSystem);
 	check(SimulationSystem);
 	check(SimulationSystem->GetPhaseManager().IsDuringMassProcessing() == false);
@@ -358,11 +362,11 @@ void UMassAgentSubsystem::HandlePendingInitialization()
 
 			FMassArchetypeCompositionDescriptor& PuppetDescriptor = AgentComp->GetMutablePuppetSpecificAddition();
 			PuppetDescriptor = TemplateDescriptor;
-			EntitySystem->AddCompositionToEntity_GetDelta(PuppetEntity, PuppetDescriptor);
+			EntityManager->AddCompositionToEntity_GetDelta(PuppetEntity, PuppetDescriptor);
 			
 			if (EntityTemplate->GetObjectFragmentInitializers().Num())
 			{
-				const FMassArchetypeHandle ArchetypeHandle = EntitySystem->GetArchetypeForEntity(PuppetEntity);
+				const FMassArchetypeHandle ArchetypeHandle = EntityManager->GetArchetypeForEntity(PuppetEntity);
 				FMassEntityView EntityView(ArchetypeHandle, PuppetEntity);
 				FMassAgentSubsystemHelper::InitializeAgentComponentFragments(*AgentComp, EntityView, EMassTranslationDirection::MassToActor, EntityTemplate->GetObjectFragmentInitializers());
 			}
diff --git a/Engine/Plugins/Runtime/MassGameplay/Source/MassActors/Private/Translators/MassCapsuleComponentTranslators.cpp b/Engine/Plugins/Runtime/MassGameplay/Source/MassActors/Private/Translators/MassCapsuleComponentTranslators.cpp
--- a/Engine/Plugins/Runtime/MassGameplay/Source/MassActors/Private/Translators/MassCapsuleComponentTranslators.cpp
+++ b/Engine/Plugins/Runtime/MassGameplay/Source/MassActors/Private/Translators/MassCapsuleComponentTranslators.cpp
@@ -24,9 +24,9 @@ void UMassCapsuleTransformToMassTranslator::ConfigureQueries()
 	EntityQuery.AddRequirement<FTransformFragment>(EMassFragmentAccess::ReadWrite);
 }
 
-void UMassCapsuleTransformToMassTranslator::Execute(UMassEntitySubsystem& EntitySubsystem, FMassExecutionContext& Context)
+void UMassCapsuleTransformToMassTranslator::Execute(FMassEntityManager& EntityManager, FMassExecutionContext& Context)
 {
-	EntityQuery.ForEachEntityChunk(EntitySubsystem, Context, [this](FMassExecutionContext& Context)
+	EntityQuery.ForEachEntityChunk(EntityManager, Context, [this](FMassExecutionContext& Context)
 		{
 			const TConstArrayView<FCapsuleComponentWrapperFragment> CapsuleComponentList = Context.GetFragmentView<FCapsuleComponentWrapperFragment>();
 			const TArrayView<FTransformFragment> LocationList = Context.GetMutableFragmentView<FTransformFragment>();
@@ -61,9 +61,9 @@ void UMassTransformToActorCapsuleTranslator::ConfigureQueries()
 	EntityQuery.RequireMutatingWorldAccess(); // due to mutating World by setting component transform
 }
 
-void UMassTransformToActorCapsuleTranslator::Execute(UMassEntitySubsystem& EntitySubsystem, FMassExecutionContext& Context)
+void UMassTransformToActorCapsuleTranslator::Execute(FMassEntityManager& EntityManager, FMassExecutionContext& Context)
 {
-	EntityQuery.ForEachEntityChunk(EntitySubsystem, Context, [this](FMassExecutionContext& Context)
+	EntityQuery.ForEachEntityChunk(EntityManager, Context, [this](FMassExecutionContext& Context)
 		{
 			const TArrayView<FCapsuleComponentWrapperFragment> CapsuleComponentList = Context.GetMutableFragmentView<FCapsuleComponentWrapperFragment>();
 			const TConstArrayView<FTransformFragment> LocationList = Context.GetFragmentView<FTransformFragment>();
diff --git a/Engine/Plugins/Runtime/MassGameplay/Source/MassActors/Private/Translators/MassCharacterMovementTranslators.cpp b/Engine/Plugins/Runtime/MassGameplay/Source/MassActors/Private/Translators/MassCharacterMovementTranslators.cpp
--- a/Engine/Plugins/Runtime/MassGameplay/Source/MassActors/Private/Translators/MassCharacterMovementTranslators.cpp
+++ b/Engine/Plugins/Runtime/MassGameplay/Source/MassActors/Private/Translators/MassCharacterMovementTranslators.cpp
@@ -28,9 +28,9 @@ void UMassCharacterMovementToMassTranslator::ConfigureQueries()
 	EntityQuery.AddRequirement<FMassVelocityFragment>(EMassFragmentAccess::ReadWrite);
 }
 
-void UMassCharacterMovementToMassTranslator::Execute(UMassEntitySubsystem& EntitySubsystem, FMassExecutionContext& Context)
+void UMassCharacterMovementToMassTranslator::Execute(FMassEntityManager& EntityManager, FMassExecutionContext& Context)
 {
-	EntityQuery.ForEachEntityChunk(EntitySubsystem, Context, [this](FMassExecutionContext& Context)
+	EntityQuery.ForEachEntityChunk(EntityManager, Context, [this](FMassExecutionContext& Context)
 	{
 		const TConstArrayView<FCharacterMovementComponentWrapperFragment> ComponentList = Context.GetFragmentView<FCharacterMovementComponentWrapperFragment>();
 		const TArrayView<FTransformFragment> LocationList = Context.GetMutableFragmentView<FTransformFragment>();
@@ -69,9 +69,9 @@ void UMassCharacterMovementToActorTranslator::ConfigureQueries()
 	EntityQuery.AddRequirement<FMassVelocityFragment>(EMassFragmentAccess::ReadOnly);
 }
 
-void UMassCharacterMovementToActorTranslator::Execute(UMassEntitySubsystem& EntitySubsystem, FMassExecutionContext& Context)
+void UMassCharacterMovementToActorTranslator::Execute(FMassEntityManager& EntityManager, FMassExecutionContext& Context)
 {
-	EntityQuery.ForEachEntityChunk(EntitySubsystem, Context, [this](FMassExecutionContext& Context)
+	EntityQuery.ForEachEntityChunk(EntityManager, Context, [this](FMassExecutionContext& Context)
 	{
 		const TArrayView<FCharacterMovementComponentWrapperFragment> ComponentList = Context.GetMutableFragmentView<FCharacterMovementComponentWrapperFragment>();
 		const TConstArrayView<FMassVelocityFragment> VelocityList = Context.GetFragmentView<FMassVelocityFragment>();
@@ -106,9 +106,9 @@ void UMassCharacterOrientationToMassTranslator::ConfigureQueries()
 	EntityQuery.AddRequirement<FTransformFragment>(EMassFragmentAccess::ReadWrite);
 }
 
-void UMassCharacterOrientationToMassTranslator::Execute(UMassEntitySubsystem& EntitySubsystem, FMassExecutionContext& Context)
+void UMassCharacterOrientationToMassTranslator::Execute(FMassEntityManager& EntityManager, FMassExecutionContext& Context)
 {
-	EntityQuery.ForEachEntityChunk(EntitySubsystem, Context, [this](FMassExecutionContext& Context)
+	EntityQuery.ForEachEntityChunk(EntityManager, Context, [this](FMassExecutionContext& Context)
 	{
 		const TConstArrayView<FCharacterMovementComponentWrapperFragment> ComponentList = Context.GetFragmentView<FCharacterMovementComponentWrapperFragment>();
 		const TArrayView<FTransformFragment> LocationList = Context.GetMutableFragmentView<FTransformFragment>();
@@ -149,9 +149,9 @@ void UMassCharacterOrientationToActorTranslator::ConfigureQueries()
 	EntityQuery.RequireMutatingWorldAccess(); // due to mutating World by setting component rotation
 }
 
-void UMassCharacterOrientationToActorTranslator::Execute(UMassEntitySubsystem& EntitySubsystem, FMassExecutionContext& Context)
+void UMassCharacterOrientationToActorTranslator::Execute(FMassEntityManager& EntityManager, FMassExecutionContext& Context)
 {
-	EntityQuery.ForEachEntityChunk(EntitySubsystem, Context, [this](FMassExecutionContext& Context)
+	EntityQuery.ForEachEntityChunk(EntityManager, Context, [this](FMassExecutionContext& Context)
 	{
 		const TArrayView<FCharacterMovementComponentWrapperFragment> ComponentList = Context.GetMutableFragmentView<FCharacterMovementComponentWrapperFragment>();
 		const TConstArrayView<FTransformFragment> TransformList = Context.GetFragmentView<FTransformFragment>();
diff --git a/Engine/Plugins/Runtime/MassGameplay/Source/MassActors/Private/Translators/MassSceneComponentLocationTranslator.cpp b/Engine/Plugins/Runtime/MassGameplay/Source/MassActors/Private/Translators/MassSceneComponentLocationTranslator.cpp
--- a/Engine/Plugins/Runtime/MassGameplay/Source/MassActors/Private/Translators/MassSceneComponentLocationTranslator.cpp
+++ b/Engine/Plugins/Runtime/MassGameplay/Source/MassActors/Private/Translators/MassSceneComponentLocationTranslator.cpp
@@ -24,9 +24,9 @@ void UMassSceneComponentLocationToMassTranslator::ConfigureQueries()
 	EntityQuery.AddRequirement<FTransformFragment>(EMassFragmentAccess::ReadWrite);
 }
 
-void UMassSceneComponentLocationToMassTranslator::Execute(UMassEntitySubsystem& EntitySubsystem, FMassExecutionContext& Context)
+void UMassSceneComponentLocationToMassTranslator::Execute(FMassEntityManager& EntityManager, FMassExecutionContext& Context)
 {
-	EntityQuery.ForEachEntityChunk(EntitySubsystem, Context, [this](FMassExecutionContext& Context)
+	EntityQuery.ForEachEntityChunk(EntityManager, Context, [this](FMassExecutionContext& Context)
 	{
 		const TConstArrayView<FMassSceneComponentWrapperFragment> ComponentList = Context.GetFragmentView<FMassSceneComponentWrapperFragment>();
 		const TArrayView<FTransformFragment> LocationList = Context.GetMutableFragmentView<FTransformFragment>();
@@ -63,9 +63,9 @@ void UMassSceneComponentLocationToActorTranslator::ConfigureQueries()
 	EntityQuery.RequireMutatingWorldAccess(); // due to mutating World by setting actor's location
 }
 
-void UMassSceneComponentLocationToActorTranslator::Execute(UMassEntitySubsystem& EntitySubsystem, FMassExecutionContext& Context)
+void UMassSceneComponentLocationToActorTranslator::Execute(FMassEntityManager& EntityManager, FMassExecutionContext& Context)
 {
-	EntityQuery.ForEachEntityChunk(EntitySubsystem, Context, [this](FMassExecutionContext& Context)
+	EntityQuery.ForEachEntityChunk(EntityManager, Context, [this](FMassExecutionContext& Context)
 		{
 			const TConstArrayView<FMassSceneComponentWrapperFragment> ComponentList = Context.GetFragmentView<FMassSceneComponentWrapperFragment>();
 			const TArrayView<FTransformFragment> LocationList = Context.GetMutableFragmentView<FTransformFragment>();
diff --git a/Engine/Plugins/Runtime/MassGameplay/Source/MassActors/Public/MassActorSubsystem.h b/Engine/Plugins/Runtime/MassGameplay/Source/MassActors/Public/MassActorSubsystem.h
--- a/Engine/Plugins/Runtime/MassGameplay/Source/MassActors/Public/MassActorSubsystem.h
+++ b/Engine/Plugins/Runtime/MassGameplay/Source/MassActors/Public/MassActorSubsystem.h
@@ -12,7 +12,7 @@
 
 struct FMassEntityHandle;
 class AActor;
-class UMassEntitySubsystem;
+struct FMassEntityManager;
 
 /**
  * Fragment to save the actor pointer of a mass entity if it exist
@@ -107,8 +107,7 @@ protected:
 	TMap<TObjectKey<const AActor>, FMassEntityHandle> ActorHandleMap;
 	UE_MT_DECLARE_RW_ACCESS_DETECTOR(ActorHandleMapDetector);
 	
-	UPROPERTY()
-	UMassEntitySubsystem* EntitySystem;
+	TSharedPtr<FMassEntityManager> EntityManager;
 };
 
 template<>
diff --git a/Engine/Plugins/Runtime/MassGameplay/Source/MassActors/Public/MassAgentSubsystem.h b/Engine/Plugins/Runtime/MassGameplay/Source/MassActors/Public/MassAgentSubsystem.h
--- a/Engine/Plugins/Runtime/MassGameplay/Source/MassActors/Public/MassAgentSubsystem.h
+++ b/Engine/Plugins/Runtime/MassGameplay/Source/MassActors/Public/MassAgentSubsystem.h
@@ -9,7 +9,7 @@
 #include "MassAgentSubsystem.generated.h"
 
 class AActor;
-class UMassEntitySubsystem;
+struct FMassEntityManager;
 class UMassSpawnerSubsystem;
 class UMassAgentComponent;
 class UMassSimulationSubsystem;
@@ -121,8 +121,7 @@ protected:
 	void OnMassAgentRemovedFromReplication(FMassNetworkID NetID, FMassEntityHandle Entity);
 
 protected:
-	UPROPERTY()
-	UMassEntitySubsystem* EntitySystem;
+	TSharedPtr<FMassEntityManager> EntityManager;
 
 	UPROPERTY()
 	UMassSpawnerSubsystem* SpawnerSystem;
diff --git a/Engine/Plugins/Runtime/MassGameplay/Source/MassActors/Public/Translators/MassCapsuleComponentTranslators.h b/Engine/Plugins/Runtime/MassGameplay/Source/MassActors/Public/Translators/MassCapsuleComponentTranslators.h
--- a/Engine/Plugins/Runtime/MassGameplay/Source/MassActors/Public/Translators/MassCapsuleComponentTranslators.h
+++ b/Engine/Plugins/Runtime/MassGameplay/Source/MassActors/Public/Translators/MassCapsuleComponentTranslators.h
@@ -40,7 +40,7 @@ public:
 
 protected:
 	virtual void ConfigureQueries() override;
-	virtual void Execute(UMassEntitySubsystem& EntitySubsystem, FMassExecutionContext& Context) override;	
+	virtual void Execute(FMassEntityManager& EntityManager, FMassExecutionContext& Context) override;	
 
 	FMassEntityQuery EntityQuery;
 };
@@ -60,7 +60,7 @@ public:
 
 protected:
 	virtual void ConfigureQueries() override;
-	virtual void Execute(UMassEntitySubsystem& EntitySubsystem, FMassExecutionContext& Context) override;
+	virtual void Execute(FMassEntityManager& EntityManager, FMassExecutionContext& Context) override;
 
 	FMassEntityQuery EntityQuery;
 };
diff --git a/Engine/Plugins/Runtime/MassGameplay/Source/MassActors/Public/Translators/MassCharacterMovementTranslators.h b/Engine/Plugins/Runtime/MassGameplay/Source/MassActors/Public/Translators/MassCharacterMovementTranslators.h
--- a/Engine/Plugins/Runtime/MassGameplay/Source/MassActors/Public/Translators/MassCharacterMovementTranslators.h
+++ b/Engine/Plugins/Runtime/MassGameplay/Source/MassActors/Public/Translators/MassCharacterMovementTranslators.h
@@ -32,7 +32,7 @@ public:
 
 protected:
 	virtual void ConfigureQueries() override;
-	virtual void Execute(UMassEntitySubsystem& EntitySubsystem, FMassExecutionContext& Context) override;
+	virtual void Execute(FMassEntityManager& EntityManager, FMassExecutionContext& Context) override;
 
 	FMassEntityQuery EntityQuery;
 };
@@ -54,7 +54,7 @@ public:
 
 protected:
 	virtual void ConfigureQueries() override;
-	virtual void Execute(UMassEntitySubsystem& EntitySubsystem, FMassExecutionContext& Context) override;
+	virtual void Execute(FMassEntityManager& EntityManager, FMassExecutionContext& Context) override;
 
 	FMassEntityQuery EntityQuery;
 };
@@ -76,7 +76,7 @@ public:
 
 protected:
 	virtual void ConfigureQueries() override;
-	virtual void Execute(UMassEntitySubsystem& EntitySubsystem, FMassExecutionContext& Context) override;
+	virtual void Execute(FMassEntityManager& EntityManager, FMassExecutionContext& Context) override;
 
 	FMassEntityQuery EntityQuery;
 };
@@ -98,7 +98,7 @@ public:
 
 protected:
 	virtual void ConfigureQueries() override;
-	virtual void Execute(UMassEntitySubsystem& EntitySubsystem, FMassExecutionContext& Context) override;
+	virtual void Execute(FMassEntityManager& EntityManager, FMassExecutionContext& Context) override;
 
 	FMassEntityQuery EntityQuery;
 };
diff --git a/Engine/Plugins/Runtime/MassGameplay/Source/MassActors/Public/Translators/MassTranslators_BehaviorTree.h b/Engine/Plugins/Runtime/MassGameplay/Source/MassActors/Public/Translators/MassTranslators_BehaviorTree.h
--- a/Engine/Plugins/Runtime/MassGameplay/Source/MassActors/Public/Translators/MassTranslators_BehaviorTree.h
+++ b/Engine/Plugins/Runtime/MassGameplay/Source/MassActors/Public/Translators/MassTranslators_BehaviorTree.h
@@ -26,7 +26,7 @@ public:
 
 protected:
 	virtual void ConfigureQueries() override;
-	virtual void Execute(UMassEntitySubsystem& EntitySubsystem, FMassExecutionContext& Context) override {}
+	virtual void Execute(FMassEntityManager& EntityManager, FMassExecutionContext& Context) override {}
 
 	FMassEntityQuery EntityQuery;
 };
diff --git a/Engine/Plugins/Runtime/MassGameplay/Source/MassCommon/Public/MassCommonTypes.h b/Engine/Plugins/Runtime/MassGameplay/Source/MassCommon/Public/MassCommonTypes.h
--- a/Engine/Plugins/Runtime/MassGameplay/Source/MassCommon/Public/MassCommonTypes.h
+++ b/Engine/Plugins/Runtime/MassGameplay/Source/MassCommon/Public/MassCommonTypes.h
@@ -12,7 +12,6 @@
 
 class UStaticMesh;
 class UMaterialInterface;
-class UMassEntitySubsystem;
 
 #if WITH_MASSGAMEPLAY_DEBUG
 namespace UE::Mass::Debug
diff --git a/Engine/Plugins/Runtime/MassGameplay/Source/MassGameplayDebug/Private/AssignDebugVisProcessor.cpp b/Engine/Plugins/Runtime/MassGameplay/Source/MassGameplayDebug/Private/AssignDebugVisProcessor.cpp
--- a/Engine/Plugins/Runtime/MassGameplay/Source/MassGameplayDebug/Private/AssignDebugVisProcessor.cpp
+++ b/Engine/Plugins/Runtime/MassGameplay/Source/MassGameplayDebug/Private/AssignDebugVisProcessor.cpp
@@ -26,12 +26,12 @@ void UAssignDebugVisProcessor::ConfigureQueries()
 	ProcessorRequirements.AddSubsystemRequirement<UMassDebuggerSubsystem>(EMassFragmentAccess::ReadWrite);
 }
 
-void UAssignDebugVisProcessor::Execute(UMassEntitySubsystem& EntitySubsystem, FMassExecutionContext& Context)
+void UAssignDebugVisProcessor::Execute(FMassEntityManager& EntityManager, FMassExecutionContext& Context)
 {
 #if WITH_EDITORONLY_DATA
 	QUICK_SCOPE_CYCLE_COUNTER(AssignDebugVisProcessor_Execute);
 	
-	UWorld* World = EntitySubsystem.GetWorld();
+	UWorld* World = EntityManager.GetWorld();
 	// @todo this code bit is temporary, so is the Visualizer->DirtyVisuals at the end of the function. Will be wrapped in
 	// "executable task" once that's implemented. 
 	UMassDebugVisualizationComponent* Visualizer = nullptr;
@@ -41,7 +41,7 @@ void UAssignDebugVisProcessor::Execute(UMassEntitySubsystem& EntitySubsystem, FM
 	// note that this function will create the "visual components" only it they're missing or out of sync. 
 	Debugger.GetVisualizationComponent()->ConditionallyConstructVisualComponent();
 
-	EntityQuery.ForEachEntityChunk(EntitySubsystem, Context, [this, World](FMassExecutionContext& Context)
+	EntityQuery.ForEachEntityChunk(EntityManager, Context, [this, World](FMassExecutionContext& Context)
 	{
 		UMassDebuggerSubsystem& Debugger = Context.GetMutableSubsystemChecked<UMassDebuggerSubsystem>(World);
 		UMassDebugVisualizationComponent* Visualizer = Debugger.GetVisualizationComponent();
diff --git a/Engine/Plugins/Runtime/MassGameplay/Source/MassGameplayDebug/Private/DebugVisLocationProcessor.cpp b/Engine/Plugins/Runtime/MassGameplay/Source/MassGameplayDebug/Private/DebugVisLocationProcessor.cpp
--- a/Engine/Plugins/Runtime/MassGameplay/Source/MassGameplayDebug/Private/DebugVisLocationProcessor.cpp
+++ b/Engine/Plugins/Runtime/MassGameplay/Source/MassGameplayDebug/Private/DebugVisLocationProcessor.cpp
@@ -26,12 +26,12 @@ void UDebugVisLocationProcessor::ConfigureQueries()
 	EntityQuery.AddSubsystemRequirement<UMassDebuggerSubsystem>(EMassFragmentAccess::ReadWrite);
 }
 
-void UDebugVisLocationProcessor::Execute(UMassEntitySubsystem& EntitySubsystem, FMassExecutionContext& Context)
+void UDebugVisLocationProcessor::Execute(FMassEntityManager& EntityManager, FMassExecutionContext& Context)
 {
 #if WITH_EDITORONLY_DATA
 	QUICK_SCOPE_CYCLE_COUNTER(DebugVisLocationProcessor_Run);
 
-	EntityQuery.ForEachEntityChunk(EntitySubsystem, Context, [this, World = EntitySubsystem.GetWorld()](FMassExecutionContext& Context)
+	EntityQuery.ForEachEntityChunk(EntityManager, Context, [this, World = EntityManager.GetWorld()](FMassExecutionContext& Context)
 	{
 		UMassDebuggerSubsystem& Debugger = Context.GetMutableSubsystemChecked<UMassDebuggerSubsystem>(World);
 		UMassDebugVisualizationComponent* Visualizer = Debugger.GetVisualizationComponent();
@@ -61,7 +61,7 @@ void UDebugVisLocationProcessor::Execute(UMassEntitySubsystem& EntitySubsystem,
 		}
 	});
 
-	UMassDebuggerSubsystem* Debugger = UWorld::GetSubsystem<UMassDebuggerSubsystem>(EntitySubsystem.GetWorld());
+	UMassDebuggerSubsystem* Debugger = UWorld::GetSubsystem<UMassDebuggerSubsystem>(EntityManager.GetWorld());
 	if (ensure(Debugger))
 	{
 		Debugger->GetVisualizationComponent()->DirtyVisuals();
@@ -89,11 +89,11 @@ void UMassProcessor_UpdateDebugVis::ConfigureQueries()
 	EntityQuery.AddSubsystemRequirement<UMassDebuggerSubsystem>(EMassFragmentAccess::ReadWrite);
 }
 
-void UMassProcessor_UpdateDebugVis::Execute(UMassEntitySubsystem& EntitySubsystem, FMassExecutionContext& Context)
+void UMassProcessor_UpdateDebugVis::Execute(FMassEntityManager& EntityManager, FMassExecutionContext& Context)
 {
 	QUICK_SCOPE_CYCLE_COUNTER(UMassProcessor_UpdateDebugVis_Run);
 
-	EntityQuery.ForEachEntityChunk(EntitySubsystem, Context, [this, World = EntitySubsystem.GetWorld()](FMassExecutionContext& Context)
+	EntityQuery.ForEachEntityChunk(EntityManager, Context, [this, World = EntityManager.GetWorld()](FMassExecutionContext& Context)
 		{
 			UMassDebuggerSubsystem& Debugger = Context.GetMutableSubsystemChecked<UMassDebuggerSubsystem>(World);
 
diff --git a/Engine/Plugins/Runtime/MassGameplay/Source/MassGameplayDebug/Private/MassDebuggerSubsystem.cpp b/Engine/Plugins/Runtime/MassGameplay/Source/MassGameplayDebug/Private/MassDebuggerSubsystem.cpp
--- a/Engine/Plugins/Runtime/MassGameplay/Source/MassGameplayDebug/Private/MassDebuggerSubsystem.cpp
+++ b/Engine/Plugins/Runtime/MassGameplay/Source/MassGameplayDebug/Private/MassDebuggerSubsystem.cpp
@@ -33,9 +33,9 @@ void UMassDebuggerSubsystem::Deinitialize()
 }
 
 #if WITH_MASSENTITY_DEBUG
-void UMassDebuggerSubsystem::OnEntitySelected(const UMassEntitySubsystem& EntitySubsystem, const FMassEntityHandle EntityHandle)
+void UMassDebuggerSubsystem::OnEntitySelected(const FMassEntityManager& EntityManager, const FMassEntityHandle EntityHandle)
 {
-	if (EntitySubsystem.GetWorld() == GetWorld())
+	if (EntityManager.GetWorld() == GetWorld())
 	{
 		SetSelectedEntity(EntityHandle);
 	}
diff --git a/Engine/Plugins/Runtime/MassGameplay/Source/MassGameplayDebug/Private/MassGameplayDebugTypes.cpp b/Engine/Plugins/Runtime/MassGameplay/Source/MassGameplayDebug/Private/MassGameplayDebugTypes.cpp
--- a/Engine/Plugins/Runtime/MassGameplay/Source/MassGameplayDebug/Private/MassGameplayDebugTypes.cpp
+++ b/Engine/Plugins/Runtime/MassGameplay/Source/MassGameplayDebug/Private/MassGameplayDebugTypes.cpp
@@ -11,7 +11,7 @@
 namespace UE::Mass::Debug
 {
 
-void GetDebugEntitiesAndLocations(const UMassEntitySubsystem& EntitySystem, TArray<FMassEntityHandle>& OutEntities, TArray<FVector>& OutLocations)
+void GetDebugEntitiesAndLocations(const FMassEntityManager& EntityManager, TArray<FMassEntityHandle>& OutEntities, TArray<FVector>& OutLocations)
 {
 	int32 DebugEntityEnd, DebugEntityBegin;
 	if (GetDebugEntitiesRange(DebugEntityEnd, DebugEntityBegin) == false)
@@ -24,10 +24,10 @@ void GetDebugEntitiesAndLocations(const UMassEntitySubsystem& EntitySystem, TArr
 
 	for (int32 i = DebugEntityBegin; i <= DebugEntityEnd; ++i)
 	{
-		const FMassEntityHandle EntityHandle = ConvertEntityIndexToHandle(EntitySystem, i);
+		const FMassEntityHandle EntityHandle = ConvertEntityIndexToHandle(EntityManager, i);
 		if (EntityHandle.IsSet())
 		{
-			if (const FTransformFragment* TransformFragment = EntitySystem.GetFragmentDataPtr<FTransformFragment>(EntityHandle))
+			if (const FTransformFragment* TransformFragment = EntityManager.GetFragmentDataPtr<FTransformFragment>(EntityHandle))
 			{
 				OutEntities.Add(EntityHandle);
 				OutLocations.Add(TransformFragment->GetTransform().GetLocation());
@@ -36,9 +36,9 @@ void GetDebugEntitiesAndLocations(const UMassEntitySubsystem& EntitySystem, TArr
 	}
 }
 
-FMassEntityHandle ConvertEntityIndexToHandle(const UMassEntitySubsystem& EntitySystem, const int32 EntityIndex)
+FMassEntityHandle ConvertEntityIndexToHandle(const FMassEntityManager& EntityManager, const int32 EntityIndex)
 {
-	return EntitySystem.DebugGetEntityIndexHandle(EntityIndex);
+	return EntityManager.DebugGetEntityIndexHandle(EntityIndex);
 }
 
 } // namespace UE::Mass::Debug
diff --git a/Engine/Plugins/Runtime/MassGameplay/Source/MassGameplayDebug/Public/MassDebuggerSubsystem.h b/Engine/Plugins/Runtime/MassGameplay/Source/MassGameplayDebug/Public/MassDebuggerSubsystem.h
--- a/Engine/Plugins/Runtime/MassGameplay/Source/MassGameplayDebug/Public/MassDebuggerSubsystem.h
+++ b/Engine/Plugins/Runtime/MassGameplay/Source/MassGameplayDebug/Public/MassDebuggerSubsystem.h
@@ -52,7 +52,7 @@ protected:
 	
 	void OnProcessingPhaseStarted(const float DeltaSeconds);
 	void PreTickProcessors();
-	void OnEntitySelected(const UMassEntitySubsystem& EntitySubsystem, const FMassEntityHandle EntityHandle);
+	void OnEntitySelected(const FMassEntityManager& EntityManager, const FMassEntityHandle EntityHandle);
 
 protected:
 	bool bCollectingData = false;
diff --git a/Engine/Plugins/Runtime/MassGameplay/Source/MassGameplayDebug/Public/MassGameplayDebugTypes.h b/Engine/Plugins/Runtime/MassGameplay/Source/MassGameplayDebug/Public/MassGameplayDebugTypes.h
--- a/Engine/Plugins/Runtime/MassGameplay/Source/MassGameplayDebug/Public/MassGameplayDebugTypes.h
+++ b/Engine/Plugins/Runtime/MassGameplay/Source/MassGameplayDebug/Public/MassGameplayDebugTypes.h
@@ -26,8 +26,8 @@ namespace UE::Mass::Debug
 	 * Fetches entity handles and their locations for entities indicated by index range as set by
 	 * ai.debug.mass.SetDebugEntityRange or ai.debug.mass.DebugEntity console commands.
 	 */
-	MASSGAMEPLAYDEBUG_API extern void GetDebugEntitiesAndLocations(const UMassEntitySubsystem& EntitySystem, TArray<FMassEntityHandle>& OutEntities, TArray<FVector>& OutLocations);
-	MASSGAMEPLAYDEBUG_API extern FMassEntityHandle ConvertEntityIndexToHandle(const UMassEntitySubsystem& EntitySystem, const int32 EntityIndex);
+	MASSGAMEPLAYDEBUG_API extern void GetDebugEntitiesAndLocations(const FMassEntityManager& EntitySubsystem, TArray<FMassEntityHandle>& OutEntities, TArray<FVector>& OutLocations);
+	MASSGAMEPLAYDEBUG_API extern FMassEntityHandle ConvertEntityIndexToHandle(const FMassEntityManager& EntitySubsystem, const int32 EntityIndex);
 } // namespace UE::Mass::Debug
 #endif // WITH_MASSGAMEPLAY_DEBUG
 
diff --git a/Engine/Plugins/Runtime/MassGameplay/Source/MassLOD/Private/MassLODCollectorProcessor.cpp b/Engine/Plugins/Runtime/MassGameplay/Source/MassLOD/Private/MassLODCollectorProcessor.cpp
--- a/Engine/Plugins/Runtime/MassGameplay/Source/MassLOD/Private/MassLODCollectorProcessor.cpp
+++ b/Engine/Plugins/Runtime/MassGameplay/Source/MassLOD/Private/MassLODCollectorProcessor.cpp
@@ -63,36 +63,36 @@ void UMassLODCollectorProcessor::CollectLODForChunk(FMassExecutionContext& Conte
 }
 
 template <bool bLocalViewersOnly>
-void UMassLODCollectorProcessor::ExecuteInternal(UMassEntitySubsystem& EntitySubsystem, FMassExecutionContext& Context)
+void UMassLODCollectorProcessor::ExecuteInternal(FMassEntityManager& EntityManager, FMassExecutionContext& Context)
 {
 	{
 		TRACE_CPUPROFILER_EVENT_SCOPE(TEXT("Close"));
-		EntityQuery_VisibleRangeAndOnLOD.ForEachEntityChunk(EntitySubsystem, Context, [this](FMassExecutionContext& Context) { CollectLODForChunk<bLocalViewersOnly>(Context); });
-		EntityQuery_VisibleRangeOnly.ForEachEntityChunk(EntitySubsystem, Context, [this](FMassExecutionContext& Context) { CollectLODForChunk<bLocalViewersOnly>(Context); });
-		EntityQuery_OnLODOnly.ForEachEntityChunk(EntitySubsystem, Context, [this](FMassExecutionContext& Context) { CollectLODForChunk<bLocalViewersOnly>(Context); });
+		EntityQuery_VisibleRangeAndOnLOD.ForEachEntityChunk(EntityManager, Context, [this](FMassExecutionContext& Context) { CollectLODForChunk<bLocalViewersOnly>(Context); });
+		EntityQuery_VisibleRangeOnly.ForEachEntityChunk(EntityManager, Context, [this](FMassExecutionContext& Context) { CollectLODForChunk<bLocalViewersOnly>(Context); });
+		EntityQuery_OnLODOnly.ForEachEntityChunk(EntityManager, Context, [this](FMassExecutionContext& Context) { CollectLODForChunk<bLocalViewersOnly>(Context); });
 	}
 
 	{
 		TRACE_CPUPROFILER_EVENT_SCOPE(TEXT("Far"));
-		EntityQuery_NotVisibleRangeAndOffLOD.ForEachEntityChunk(EntitySubsystem, Context, [this](FMassExecutionContext& Context) { CollectLODForChunk<bLocalViewersOnly>(Context); });
+		EntityQuery_NotVisibleRangeAndOffLOD.ForEachEntityChunk(EntityManager, Context, [this](FMassExecutionContext& Context) { CollectLODForChunk<bLocalViewersOnly>(Context); });
 	}
 }
 
-void UMassLODCollectorProcessor::Execute(UMassEntitySubsystem& EntitySubsystem, FMassExecutionContext& Context)
+void UMassLODCollectorProcessor::Execute(FMassEntityManager& EntityManager, FMassExecutionContext& Context)
 {
-	const UMassLODSubsystem& LODSubsystem = Context.GetSubsystemChecked<UMassLODSubsystem>(EntitySubsystem.GetWorld());
+	const UMassLODSubsystem& LODSubsystem = Context.GetSubsystemChecked<UMassLODSubsystem>(EntityManager.GetWorld());
 	const TArray<FViewerInfo>& Viewers = LODSubsystem.GetViewers();
 	Collector.PrepareExecution(Viewers);
 
-	UWorld* World = EntitySubsystem.GetWorld();
+	UWorld* World = EntityManager.GetWorld();
 	check(World);
 	if (World->IsNetMode(NM_DedicatedServer))
 	{
-		ExecuteInternal<false/*bLocalViewersOnly*/>(EntitySubsystem, Context);
+		ExecuteInternal<false/*bLocalViewersOnly*/>(EntityManager, Context);
 	}
 	else
 	{
-		ExecuteInternal<true/*bLocalViewersOnly*/>(EntitySubsystem, Context);
+		ExecuteInternal<true/*bLocalViewersOnly*/>(EntityManager, Context);
 	}
 
 }
diff --git a/Engine/Plugins/Runtime/MassGameplay/Source/MassLOD/Private/MassLODTrait.cpp b/Engine/Plugins/Runtime/MassGameplay/Source/MassLOD/Private/MassLODTrait.cpp
--- a/Engine/Plugins/Runtime/MassGameplay/Source/MassLOD/Private/MassLODTrait.cpp
+++ b/Engine/Plugins/Runtime/MassGameplay/Source/MassLOD/Private/MassLODTrait.cpp
@@ -6,6 +6,8 @@
 #include "MassCommonFragments.h"
 #include "StructUtilsTypes.h"
 #include "MassLODFragments.h"
+#include "MassEntityUtils.h"
+
 
 void UMassLODCollectorTrait::BuildTemplate(FMassEntityTemplateBuildContext& BuildContext, const UWorld& World) const
 {
@@ -28,13 +30,12 @@ void UMassSimulationLODTrait::BuildTemplate(FMassEntityTemplateBuildContext& Bui
 		BuildContext.AddTag<FMassOffLODTag>();
 	}
 
-	UMassEntitySubsystem* EntitySubsystem = UWorld::GetSubsystem<UMassEntitySubsystem>(&World);
-	check(EntitySubsystem);
+	FMassEntityManager& EntityManager = UE::Mass::Utils::GetEntityManagerChecked(World);
 
 	uint32 ParamsHash = UE::StructUtils::GetStructCrc32(FConstStructView::Make(Params));
-	FConstSharedStruct ParamsFragment = EntitySubsystem->GetOrCreateConstSharedFragment(ParamsHash, Params);
+	FConstSharedStruct ParamsFragment = EntityManager.GetOrCreateConstSharedFragment(ParamsHash, Params);
 	BuildContext.AddConstSharedFragment(ParamsFragment);
-	FSharedStruct SharedFragment = EntitySubsystem->GetOrCreateSharedFragment<FMassSimulationLODSharedFragment>(ParamsHash, Params);
+	FSharedStruct SharedFragment = EntityManager.GetOrCreateSharedFragment<FMassSimulationLODSharedFragment>(ParamsHash, Params);
 	BuildContext.AddSharedFragment(SharedFragment);
 
 	// Variable ticking from simulation LOD
@@ -44,9 +45,9 @@ void UMassSimulationLODTrait::BuildTemplate(FMassEntityTemplateBuildContext& Bui
 		BuildContext.AddChunkFragment<FMassSimulationVariableTickChunkFragment>();
 
 		uint32 VariableTickParamsHash = UE::StructUtils::GetStructCrc32(FConstStructView::Make(VariableTickParams));
-		FConstSharedStruct VariableTickParamsFragment = EntitySubsystem->GetOrCreateConstSharedFragment(VariableTickParamsHash, VariableTickParams);
+		FConstSharedStruct VariableTickParamsFragment = EntityManager.GetOrCreateConstSharedFragment(VariableTickParamsHash, VariableTickParams);
 		BuildContext.AddConstSharedFragment(VariableTickParamsFragment);
-		FSharedStruct VariableTickSharedFragment = EntitySubsystem->GetOrCreateSharedFragment<FMassSimulationVariableTickSharedFragment>(VariableTickParamsHash, VariableTickParams);
+		FSharedStruct VariableTickSharedFragment = EntityManager.GetOrCreateSharedFragment<FMassSimulationVariableTickSharedFragment>(VariableTickParamsHash, VariableTickParams);
 		BuildContext.AddSharedFragment(VariableTickSharedFragment);
 	}
 }
diff --git a/Engine/Plugins/Runtime/MassGameplay/Source/MassLOD/Private/MassSimulationLOD.cpp b/Engine/Plugins/Runtime/MassGameplay/Source/MassLOD/Private/MassSimulationLOD.cpp
--- a/Engine/Plugins/Runtime/MassGameplay/Source/MassLOD/Private/MassSimulationLOD.cpp
+++ b/Engine/Plugins/Runtime/MassGameplay/Source/MassLOD/Private/MassSimulationLOD.cpp
@@ -112,17 +112,17 @@ void UMassSimulationLODProcessor::ConfigureQueries()
 	ProcessorRequirements.AddSubsystemRequirement<UMassLODSubsystem>(EMassFragmentAccess::ReadOnly);
 }
 
-void UMassSimulationLODProcessor::Execute(UMassEntitySubsystem& EntitySubsystem, FMassExecutionContext& Context)
+void UMassSimulationLODProcessor::Execute(FMassEntityManager& EntityManager, FMassExecutionContext& Context)
 {
 	TRACE_CPUPROFILER_EVENT_SCOPE(TEXT("SimulationLOD"))
 
 	{
 		TRACE_CPUPROFILER_EVENT_SCOPE(TEXT("PrepareExecution"));
 
-		const UMassLODSubsystem& LODSubsystem = Context.GetSubsystemChecked<UMassLODSubsystem>(EntitySubsystem.GetWorld());
+		const UMassLODSubsystem& LODSubsystem = Context.GetSubsystemChecked<UMassLODSubsystem>(EntityManager.GetWorld());
 		const TArray<FViewerInfo>& Viewers = LODSubsystem.GetViewers();
 
-		EntitySubsystem.ForEachSharedFragment<FMassSimulationLODSharedFragment>([&Viewers](FMassSimulationLODSharedFragment& LODSharedFragment)
+		EntityManager.ForEachSharedFragment<FMassSimulationLODSharedFragment>([&Viewers](FMassSimulationLODSharedFragment& LODSharedFragment)
 		{
 			LODSharedFragment.LODCalculator.PrepareExecution(Viewers);
 		});
@@ -130,7 +130,7 @@ void UMassSimulationLODProcessor::Execute(UMassEntitySubsystem& EntitySubsystem,
 
 	{
 		TRACE_CPUPROFILER_EVENT_SCOPE(TEXT("CalculateLOD"));
-		EntityQueryCalculateLOD.ForEachEntityChunk(EntitySubsystem, Context, [](FMassExecutionContext& Context)
+		EntityQueryCalculateLOD.ForEachEntityChunk(EntityManager, Context, [](FMassExecutionContext& Context)
 		{
 			FMassSimulationLODSharedFragment& LODSharedFragment = Context.GetMutableSharedFragment<FMassSimulationLODSharedFragment>();
 			TConstArrayView<FMassViewerInfoFragment> ViewersInfoList = Context.GetFragmentView<FMassViewerInfoFragment>();
@@ -141,12 +141,12 @@ void UMassSimulationLODProcessor::Execute(UMassEntitySubsystem& EntitySubsystem,
 
 	{
 		TRACE_CPUPROFILER_EVENT_SCOPE(TEXT("AdjustDistancesAndLODFromCount"));
-		EntitySubsystem.ForEachSharedFragment<FMassSimulationLODSharedFragment>([](FMassSimulationLODSharedFragment& LODSharedFragment)
+		EntityManager.ForEachSharedFragment<FMassSimulationLODSharedFragment>([](FMassSimulationLODSharedFragment& LODSharedFragment)
 		{
 			LODSharedFragment.bHasAdjustedDistancesFromCount = LODSharedFragment.LODCalculator.AdjustDistancesFromCount();
 		});
 
-		EntityQueryAdjustDistances.ForEachEntityChunk(EntitySubsystem, Context, [](FMassExecutionContext& Context)
+		EntityQueryAdjustDistances.ForEachEntityChunk(EntityManager, Context, [](FMassExecutionContext& Context)
 		{
 			FMassSimulationLODSharedFragment& LODSharedFragment = Context.GetMutableSharedFragment<FMassSimulationLODSharedFragment>();
 			TConstArrayView<FMassViewerInfoFragment> ViewersInfoList = Context.GetFragmentView<FMassViewerInfoFragment>();
@@ -155,12 +155,12 @@ void UMassSimulationLODProcessor::Execute(UMassEntitySubsystem& EntitySubsystem,
 		});
 	}
 
-	UWorld* World = EntitySubsystem.GetWorld();
+	UWorld* World = EntityManager.GetWorld();
 	{
 		TRACE_CPUPROFILER_EVENT_SCOPE(TEXT("VariableTickRates"))
 		check(World);
 		const float Time = World->GetTimeSeconds();
-		EntityQueryVariableTick.ForEachEntityChunk(EntitySubsystem, Context, [Time](FMassExecutionContext& Context)
+		EntityQueryVariableTick.ForEachEntityChunk(EntityManager, Context, [Time](FMassExecutionContext& Context)
 		{
 			FMassSimulationVariableTickSharedFragment& TickRateSharedFragment = Context.GetMutableSharedFragment<FMassSimulationVariableTickSharedFragment>();
 			TConstArrayView<FMassSimulationLODFragment> SimulationLODFragments = Context.GetMutableFragmentView<FMassSimulationLODFragment>();
@@ -173,7 +173,7 @@ void UMassSimulationLODProcessor::Execute(UMassEntitySubsystem& EntitySubsystem,
 	{
 		TRACE_CPUPROFILER_EVENT_SCOPE(TEXT("SetLODTags"))
 		check(World);
-		EntityQuerySetLODTag.ForEachEntityChunk(EntitySubsystem, Context, [](FMassExecutionContext& Context)
+		EntityQuerySetLODTag.ForEachEntityChunk(EntityManager, Context, [](FMassExecutionContext& Context)
 		{
 			TConstArrayView<FMassSimulationLODFragment> SimulationLODFragments = Context.GetMutableFragmentView<FMassSimulationLODFragment>();
 			const int32 NumEntities = Context.GetNumEntities();
@@ -193,7 +193,7 @@ void UMassSimulationLODProcessor::Execute(UMassEntitySubsystem& EntitySubsystem,
 	if (UE::MassLOD::bDebugSimulationLOD)
 	{
 		TRACE_CPUPROFILER_EVENT_SCOPE(TEXT("DebugDisplayLOD"));
-		EntityQuery.ForEachEntityChunk(EntitySubsystem, Context, [World](FMassExecutionContext& Context)
+		EntityQuery.ForEachEntityChunk(EntityManager, Context, [World](FMassExecutionContext& Context)
 		{
 			FMassSimulationLODSharedFragment& LODSharedFragment = Context.GetMutableSharedFragment<FMassSimulationLODSharedFragment>();
 			TConstArrayView<FTransformFragment> LocationList = Context.GetFragmentView<FTransformFragment>();
diff --git a/Engine/Plugins/Runtime/MassGameplay/Source/MassLOD/Public/MassLODCollectorProcessor.h b/Engine/Plugins/Runtime/MassGameplay/Source/MassLOD/Public/MassLODCollectorProcessor.h
--- a/Engine/Plugins/Runtime/MassGameplay/Source/MassLOD/Public/MassLODCollectorProcessor.h
+++ b/Engine/Plugins/Runtime/MassGameplay/Source/MassLOD/Public/MassLODCollectorProcessor.h
@@ -29,13 +29,13 @@ public:
 
 protected:
 	virtual void ConfigureQueries() override;
-	virtual void Execute(UMassEntitySubsystem& EntitySubsystem, FMassExecutionContext& Context) override;
+	virtual void Execute(FMassEntityManager& EntityManager, FMassExecutionContext& Context) override;
 
 	template <bool bLocalViewersOnly>
 	void CollectLODForChunk(FMassExecutionContext& Context);
 
 	template <bool bLocalViewersOnly>
-	void ExecuteInternal(UMassEntitySubsystem& EntitySubsystem, FMassExecutionContext& Context);
+	void ExecuteInternal(FMassEntityManager& EntityManager, FMassExecutionContext& Context);
 
 	TMassLODCollector<FMassGenericCollectorLogic> Collector;
 
diff --git a/Engine/Plugins/Runtime/MassGameplay/Source/MassLOD/Public/MassSimulationLOD.h b/Engine/Plugins/Runtime/MassGameplay/Source/MassLOD/Public/MassSimulationLOD.h
--- a/Engine/Plugins/Runtime/MassGameplay/Source/MassLOD/Public/MassSimulationLOD.h
+++ b/Engine/Plugins/Runtime/MassGameplay/Source/MassLOD/Public/MassSimulationLOD.h
@@ -172,9 +172,9 @@ public:
 protected:
 
 	virtual void ConfigureQueries() override;
-	virtual void Execute(UMassEntitySubsystem& EntitySubsystem, FMassExecutionContext& Context) override;
+	virtual void Execute(FMassEntityManager& EntityManager, FMassExecutionContext& Context) override;
 
-	void CalculateLODForConfig(UMassEntitySubsystem& EntitySubsystem, FMassExecutionContext& Context, FMassSimulationLODParameters& LODParams);
+	void CalculateLODForConfig(FMassEntityManager& EntityManager, FMassExecutionContext& Context, FMassSimulationLODParameters& LODParams);
 
 	FMassEntityQuery EntityQuery;
 	FMassEntityQuery EntityQueryCalculateLOD;
