Enabled Mass MT by default for all builds except for UE_SERVER

diff --git a/Engine/Plugins/Runtime/MassEntity/Source/MassEntity/MassEntity.Build.cs b/Engine/Plugins/Runtime/MassEntity/Source/MassEntity/MassEntity.Build.cs
--- a/Engine/Plugins/Runtime/MassEntity/Source/MassEntity/MassEntity.Build.cs
+++ b/Engine/Plugins/Runtime/MassEntity/Source/MassEntity/MassEntity.Build.cs
@@ -35,9 +35,6 @@ namespace UnrealBuildTool.Rules
 			{
 				PrivateDependencyModuleNames.Add("MessageLog");
 			}
-
-			// @todo will be enables once we're confident in MASS_DO_PARALLEL implementation 
-			// PublicDefinitions.Add("MASS_DO_PARALLEL=" + ((Target.Type != TargetType.Server) ? 1 : 0));
 		}
 	}
 }
\ No newline at end of file
diff --git a/Engine/Plugins/Runtime/MassEntity/Source/MassEntity/Public/MassProcessingTypes.h b/Engine/Plugins/Runtime/MassEntity/Source/MassEntity/Public/MassProcessingTypes.h
--- a/Engine/Plugins/Runtime/MassEntity/Source/MassEntity/Public/MassProcessingTypes.h
+++ b/Engine/Plugins/Runtime/MassEntity/Source/MassEntity/Public/MassProcessingTypes.h
@@ -7,7 +7,7 @@
 #include "MassProcessingTypes.generated.h"
 
 #ifndef MASS_DO_PARALLEL
-#define MASS_DO_PARALLEL 0
+#define MASS_DO_PARALLEL !UE_SERVER
 #endif // MASS_DO_PARALLEL
 
 #define WITH_MASSENTITY_DEBUG (!(UE_BUILD_SHIPPING || UE_BUILD_SHIPPING_WITH_EDITOR || UE_BUILD_TEST) && WITH_STRUCTUTILS_DEBUG && 1)
