Compilation issue #200190, fixed nonunity nopch Mass compilation issue resulting from CL#21355128

diff --git a/Engine/Plugins/Runtime/MassEntity/Source/MassEntity/Private/MassEntitySubsystem.cpp b/Engine/Plugins/Runtime/MassEntity/Source/MassEntity/Private/MassEntitySubsystem.cpp
--- a/Engine/Plugins/Runtime/MassEntity/Source/MassEntity/Private/MassEntitySubsystem.cpp
+++ b/Engine/Plugins/Runtime/MassEntity/Source/MassEntity/Private/MassEntitySubsystem.cpp
@@ -1,7 +1,6 @@
 // Copyright Epic Games, Inc. All Rights Reserved.
 
 #include "MassEntitySubsystem.h"
-#include "MassEntityManager.h"
 
 
 //////////////////////////////////////////////////////////////////////
diff --git a/Engine/Plugins/Runtime/MassEntity/Source/MassEntity/Public/MassEntitySubsystem.h b/Engine/Plugins/Runtime/MassEntity/Source/MassEntity/Public/MassEntitySubsystem.h
--- a/Engine/Plugins/Runtime/MassEntity/Source/MassEntity/Public/MassEntitySubsystem.h
+++ b/Engine/Plugins/Runtime/MassEntity/Source/MassEntity/Public/MassEntitySubsystem.h
@@ -3,11 +3,10 @@
 #pragma once
 
 #include "Subsystems/WorldSubsystem.h"
+#include "MassEntityManager.h"
 #include "MassEntitySubsystem.generated.h"
 
 
-struct FMassEntityManager;
-
 /** 
  * The sole responsibility of this world subsystem class is to host the default instance of FMassEntityManager
  * for a given UWorld. All the gameplay-related use cases of Mass (found in MassGameplay and related plugins) 
