Fixed heap-use-after-free issues recently intruduced to mass dependency solving.

diff --git a/Engine/Plugins/Runtime/MassEntity/Source/MassEntity/Private/MassProcessorDependencySolver.cpp b/Engine/Plugins/Runtime/MassEntity/Source/MassEntity/Private/MassProcessorDependencySolver.cpp
--- a/Engine/Plugins/Runtime/MassEntity/Source/MassEntity/Private/MassProcessorDependencySolver.cpp
+++ b/Engine/Plugins/Runtime/MassEntity/Source/MassEntity/Private/MassProcessorDependencySolver.cpp
@@ -320,10 +320,11 @@ void FProcessorDependencySolver::BuildDependencies()
 {
 	// at this point we have collected all the known processors and groups in AllNodes so we can transpose 
 	// A.ExecuteBefore(B) type of dependencies into B.ExecuteAfter(A)
-	for (int32 i = 0; i < AllNodes.Num(); ++i)
+	for (int32 NodeIndex = 0; NodeIndex < AllNodes.Num(); ++NodeIndex)
 	{
-		for (const FName& BeforeDependencyName : AllNodes[i].ExecuteBefore)
+		for (int i = 0; i < AllNodes[NodeIndex].ExecuteBefore.Num(); ++i)
 		{
+			const FName BeforeDependencyName = AllNodes[NodeIndex].ExecuteBefore[i];
 			int32 DependentNodeIndex = INDEX_NONE;
 			int32* DependentNodeIndexPtr = NodeIndexMap.Find(BeforeDependencyName);
 			if (DependentNodeIndexPtr == nullptr)
@@ -337,7 +338,7 @@ void FProcessorDependencySolver::BuildDependencies()
 				AllNodes.Add({ BeforeDependencyName, nullptr, DependentNodeIndex });
 
 				UE_LOG(LogMass, Log, TEXT("Unable to find dependency \"%s\" declared by %s. Creating a dummy dependency node.")
-					, *BeforeDependencyName.ToString(), *AllNodes[i].Name.ToString());
+					, *BeforeDependencyName.ToString(), *AllNodes[NodeIndex].Name.ToString());
 			}
 			else
 			{
@@ -345,9 +346,9 @@ void FProcessorDependencySolver::BuildDependencies()
 			}
 
 			check(AllNodes.IsValidIndex(DependentNodeIndex));
-			AllNodes[DependentNodeIndex].ExecuteAfter.Add(AllNodes[i].Name);
+			AllNodes[DependentNodeIndex].ExecuteAfter.Add(AllNodes[NodeIndex].Name);
 		}
-		AllNodes[i].ExecuteBefore.Reset();
+		AllNodes[NodeIndex].ExecuteBefore.Reset();
 	}
 
 	// at this point all nodes contain:
@@ -363,7 +364,7 @@ void FProcessorDependencySolver::BuildDependencies()
 	{
 		for (int i = 0; i < AllNodes[NodeIndex].ExecuteAfter.Num(); ++i)
 		{
-			const FName& AfterDependencyName = AllNodes[NodeIndex].ExecuteAfter[i];
+			const FName AfterDependencyName = AllNodes[NodeIndex].ExecuteAfter[i];
 			int32* PrerequisiteNodeIndexPtr = NodeIndexMap.Find(AfterDependencyName);
 			int32 PrerequisiteNodeIndex = INDEX_NONE;
 
