Added simple MassEntityManager type comment

diff --git a/Engine/Plugins/Runtime/MassEntity/Source/MassEntity/Public/MassEntityManager.h b/Engine/Plugins/Runtime/MassEntity/Source/MassEntity/Public/MassEntityManager.h
--- a/Engine/Plugins/Runtime/MassEntity/Source/MassEntity/Public/MassEntityManager.h
+++ b/Engine/Plugins/Runtime/MassEntity/Source/MassEntity/Public/MassEntityManager.h
@@ -23,7 +23,19 @@ class FOutputDevice;
 struct FMassDebugger;
 enum class EMassFragmentAccess : uint8;
 
-//@TODO: Comment this
+/** 
+ * The type responsible for hosting Entities managing Archetypes.
+ * Entities are stored as FEntityData entries in a chunked array. 
+ * Each valid entity is assigned to an Archetype that stored fragments associated with a given entity at the moment. 
+ * 
+ * FMassEntityManager supplies API for entity creation (that can result in archetype creation) and entity manipulation.
+ * Even though synchronized manipulation methods are available in most cases the entity operations are performed via a command 
+ * buffer. The default command buffer can be obtained with a Defer() call. @see FMassCommandBuffer for more details.
+ * 
+ * FMassEntityManager are meant to be stored with a TSharedPtr or TSharedRef. Some of Mass API pass around 
+ * FMassEntityManager& but programmers can always use AsShared() call to obtain a shared ref for a given manager instance 
+ * (as supplied by deriving from TSharedFromThis<FMassEntityManager>).
+ */
 struct MASSENTITY_API FMassEntityManager : public TSharedFromThis<FMassEntityManager>
 {
 	friend FMassEntityQuery;
