No need to do SaveConfig by hand, found out that the property grid does it automatically for any class you edit that has the config. On top of it, it only saves the member that has been touched, not all of the class. So this with the ConfigDoNotCheckDefaults does the trick!

diff --git a/Engine/Plugins/Runtime/MassEntity/Source/MassEntity/Private/MassProcessor.cpp b/Engine/Plugins/Runtime/MassEntity/Source/MassEntity/Private/MassProcessor.cpp
--- a/Engine/Plugins/Runtime/MassEntity/Source/MassEntity/Private/MassProcessor.cpp
+++ b/Engine/Plugins/Runtime/MassEntity/Source/MassEntity/Private/MassProcessor.cpp
@@ -122,7 +122,13 @@ void UMassProcessor::SetShouldAutoRegisterWithGlobalList(const bool bAutoRegiste
 	{
 		bAutoRegisterWithProcessingPhases = bAutoRegister;
 #if WITH_EDITOR
-		SaveConfig(CPF_Config, *GetDefaultConfigFilename());
+		if (UClass* Class = GetClass())
+		{
+			if (FProperty* AutoRegisterProperty = Class->FindPropertyByName(GET_MEMBER_NAME_CHECKED(UMassProcessor, bAutoRegisterWithProcessingPhases)))
+			{
+				UpdateSinglePropertyInConfigFile(AutoRegisterProperty, *GetDefaultConfigFilename());
+			}
+		}
 #endif // WITH_EDITOR
 	}
 }
@@ -224,19 +230,6 @@ void UMassProcessor::DebugOutputDescription(FOutputDevice& Ar, int32 Indent) con
 #endif // WITH_MASSENTITY_DEBUG
 }
 
-#if WITH_EDITOR
-void UMassProcessor::PostEditChangeProperty(FPropertyChangedEvent& PropertyChangedEvent)
-{
-	Super::PostEditChangeProperty(PropertyChangedEvent);
-
-	if (HasAnyFlags(RF_ClassDefaultObject))
-	{
-		// this is here to make sure all the changes to CDOs we do via MassSettings gets serialized to the ini file
-		SaveConfig(CPF_Config, *GetDefaultConfigFilename());
-	}
-}
-#endif // WITH_EDITOR
-
 //----------------------------------------------------------------------//
 //  UMassCompositeProcessor
 //----------------------------------------------------------------------//
diff --git a/Engine/Plugins/Runtime/MassEntity/Source/MassEntity/Public/MassProcessor.h b/Engine/Plugins/Runtime/MassEntity/Source/MassEntity/Public/MassProcessor.h
--- a/Engine/Plugins/Runtime/MassEntity/Source/MassEntity/Public/MassProcessor.h
+++ b/Engine/Plugins/Runtime/MassEntity/Source/MassEntity/Public/MassProcessor.h
@@ -40,7 +40,7 @@ struct FMassProcessorExecutionOrder
 };
 
 
-UCLASS(abstract, EditInlineNew, CollapseCategories, config = Mass, defaultconfig)
+UCLASS(abstract, EditInlineNew, CollapseCategories, config = Mass, defaultconfig, ConfigDoNotCheckDefaults)
 class MASSENTITY_API UMassProcessor : public UObject
 {
 	GENERATED_BODY()
@@ -100,9 +100,6 @@ protected:
 	virtual void ConfigureQueries() PURE_VIRTUAL(UMassProcessor::ConfigureQueries);
 	virtual void PostInitProperties() override;
 	virtual void Execute(FMassEntityManager& EntityManager, FMassExecutionContext& Context) PURE_VIRTUAL(UMassProcessor::Execute);
-#if WITH_EDITOR
-	virtual void PostEditChangeProperty(FPropertyChangedEvent& PropertyChangedEvent) override;
-#endif // WITH_EDITOR
 
 protected:
 	/** Whether this processor should be executed on StandAlone or Server or Client */
