Added "allocated memory" information to MassDebugger's archetype display

diff --git a/Engine/Plugins/Runtime/MassEntity/Source/MassEntity/Private/MassDebugger.cpp b/Engine/Plugins/Runtime/MassEntity/Source/MassEntity/Private/MassDebugger.cpp
--- a/Engine/Plugins/Runtime/MassEntity/Source/MassEntity/Private/MassDebugger.cpp
+++ b/Engine/Plugins/Runtime/MassEntity/Source/MassEntity/Private/MassDebugger.cpp
@@ -256,13 +256,13 @@ const FMassArchetypeCompositionDescriptor& FMassDebugger::GetArchetypeCompositio
 	return ArchetypeData.CompositionDescriptor;
 }
 
-void FMassDebugger::GetArchetypeEntityStats(const FMassArchetypeHandle& ArchetypeHandle
-	, int32& OutEntitiesCount, int32& OutNumEntitiesPerChunk, int32& OutChunksCount)
+void FMassDebugger::GetArchetypeEntityStats(const FMassArchetypeHandle& ArchetypeHandle, FArchetypeStats& OutStats)
 {
 	const FMassArchetypeData& ArchetypeData = FMassArchetypeHelper::ArchetypeDataFromHandleChecked(ArchetypeHandle);
-	OutEntitiesCount = ArchetypeData.GetNumEntities();
-	OutNumEntitiesPerChunk = ArchetypeData.GetNumEntitiesPerChunk();
-	OutChunksCount = ArchetypeData.GetChunkCount();
+	OutStats.EntitiesCount = ArchetypeData.GetNumEntities();
+	OutStats.EntitiesCountPerChunk = ArchetypeData.GetNumEntitiesPerChunk();
+	OutStats.ChunksCount = ArchetypeData.GetChunkCount();
+	OutStats.AllocatedSize = ArchetypeData.GetAllocatedSize();
 }
 
 const TConstArrayView<FName> FMassDebugger::GetArchetypeDebugNames(const FMassArchetypeHandle& ArchetypeHandle)
diff --git a/Engine/Plugins/Runtime/MassEntity/Source/MassEntity/Public/MassDebugger.h b/Engine/Plugins/Runtime/MassEntity/Source/MassEntity/Public/MassDebugger.h
--- a/Engine/Plugins/Runtime/MassEntity/Source/MassEntity/Public/MassDebugger.h
+++ b/Engine/Plugins/Runtime/MassEntity/Source/MassEntity/Public/MassDebugger.h
@@ -54,7 +54,18 @@ struct MASSENTITY_API FMassDebugger
 	static TArray<FMassArchetypeHandle> GetAllArchetypes(const FMassEntityManager& EntitySubsystem);
 	static const FMassArchetypeCompositionDescriptor& GetArchetypeComposition(const FMassArchetypeHandle& ArchetypeHandle);
 
-	static void GetArchetypeEntityStats(const FMassArchetypeHandle& ArchetypeHandle, int32& OutEntitiesCount, int32& OutNumEntitiesPerChunk, int32& OutChunksCount);
+	struct FArchetypeStats
+	{
+		/** Number of active entities of the archetype. */
+		int32 EntitiesCount = 0;
+		/** Number of entities that fit per chunk. */
+		int32 EntitiesCountPerChunk = 0;
+		/** Number of allocated chunks. */
+		int32 ChunksCount = 0;
+		/** Total amount of memory taken by this archetype */
+		SIZE_T AllocatedSize = 0;
+	};
+	static void GetArchetypeEntityStats(const FMassArchetypeHandle& ArchetypeHandle, FArchetypeStats& OutStats);
 	static const TConstArrayView<FName> GetArchetypeDebugNames(const FMassArchetypeHandle& ArchetypeHandle);
 
 	static TConstArrayView<UMassCompositeProcessor::FDependencyNode> GetProcessingGraph(const UMassCompositeProcessor& GraphOwner);
diff --git a/Engine/Plugins/Runtime/MassEntity/Source/MassEntityDebugger/Private/MassDebuggerModel.cpp b/Engine/Plugins/Runtime/MassEntity/Source/MassEntityDebugger/Private/MassDebuggerModel.cpp
--- a/Engine/Plugins/Runtime/MassEntity/Source/MassEntityDebugger/Private/MassDebuggerModel.cpp
+++ b/Engine/Plugins/Runtime/MassEntity/Source/MassEntityDebugger/Private/MassDebuggerModel.cpp
@@ -158,7 +158,7 @@ FMassDebuggerArchetypeData::FMassDebuggerArchetypeData(const FMassArchetypeHandl
 	BytesToHexLower(reinterpret_cast<const uint8*>(&FullHash), sizeof(FullHash), FullHashAsString);
 	HashLabel = FText::FromString(FullHashAsString);
 
-	FMassDebugger::GetArchetypeEntityStats(ArchetypeHandle, EntitiesCount, EntitiesCountPerChunk, ChunksCount);
+	FMassDebugger::GetArchetypeEntityStats(ArchetypeHandle, ArchetypeStats);
 
 	const TConstArrayView<FName> DebugNames = FMassDebugger::GetArchetypeDebugNames(ArchetypeHandle);
 
diff --git a/Engine/Plugins/Runtime/MassEntity/Source/MassEntityDebugger/Private/MassDebuggerModel.h b/Engine/Plugins/Runtime/MassEntity/Source/MassEntityDebugger/Private/MassDebuggerModel.h
--- a/Engine/Plugins/Runtime/MassEntity/Source/MassEntityDebugger/Private/MassDebuggerModel.h
+++ b/Engine/Plugins/Runtime/MassEntity/Source/MassEntityDebugger/Private/MassDebuggerModel.h
@@ -65,12 +65,8 @@ struct FMassDebuggerArchetypeData
 	/** Combined hash of composition and shared fragments. */
 	uint32 FullHash = 0;
 
-	/** Number of active entities of the archetype. */
-	int32 EntitiesCount = 0;
-	/** Number of entities that fit per chunk. */
-	int32 EntitiesCountPerChunk = 0;
-	/** Number of allocated chunks. */
-	int32 ChunksCount = 0;
+	/** Archetype statistics */
+	FMassDebugger::FArchetypeStats ArchetypeStats;
 
 	/** Child debugger data (same as parent, but changed in some way) */
 	TArray<TSharedPtr<FMassDebuggerArchetypeData>> Children;
diff --git a/Engine/Plugins/Runtime/MassEntity/Source/MassEntityDebugger/Private/SMassArchetype.cpp b/Engine/Plugins/Runtime/MassEntity/Source/MassEntityDebugger/Private/SMassArchetype.cpp
--- a/Engine/Plugins/Runtime/MassEntity/Source/MassEntityDebugger/Private/SMassArchetype.cpp
+++ b/Engine/Plugins/Runtime/MassEntity/Source/MassEntityDebugger/Private/SMassArchetype.cpp
@@ -55,17 +55,18 @@ void SMassArchetype::Construct(const FArguments& InArgs, TSharedPtr<FMassDebugge
 		.Font(FCoreStyle::GetDefaultFontStyle("Mono", 9))
 	];
 	
-	FString ArchetypeDescription = FString::Printf(TEXT("EntitiesCount: %d\nEntitiesCountPerChunk: %d\nChunksCount: %d")
-		, ArchetypeDebugData.EntitiesCount
-		, ArchetypeDebugData.EntitiesCountPerChunk
-		, ArchetypeDebugData.ChunksCount);
+	FText ArchetypeDescription = FText::Format(LOCTEXT("ArchetypeDescrption", "EntitiesCount: {0}\nEntitiesCountPerChunk: {1}\nChunksCount: {2}\nAllocated memory: {3}")
+		, FText::AsNumber(ArchetypeDebugData.ArchetypeStats.EntitiesCount)
+		, FText::AsNumber(ArchetypeDebugData.ArchetypeStats.EntitiesCountPerChunk)
+		, FText::AsNumber(ArchetypeDebugData.ArchetypeStats.ChunksCount)
+		, FText::AsMemory(ArchetypeDebugData.ArchetypeStats.AllocatedSize));
 	
 	Box->AddSlot()
 	.AutoHeight()
 	.Padding(0, 4)
 	[
 		SNew(STextBlock)
-		.Text(FText::FromString(ArchetypeDescription))
+		.Text(ArchetypeDescription)
 	];
 
 	const FMassArchetypeCompositionDescriptor& Composition = ArchetypeData->Composition;
diff --git a/Engine/Plugins/Runtime/MassEntity/Source/MassEntityDebugger/Private/SMassArchetypesView.cpp b/Engine/Plugins/Runtime/MassEntity/Source/MassEntityDebugger/Private/SMassArchetypesView.cpp
--- a/Engine/Plugins/Runtime/MassEntity/Source/MassEntityDebugger/Private/SMassArchetypesView.cpp
+++ b/Engine/Plugins/Runtime/MassEntity/Source/MassEntityDebugger/Private/SMassArchetypesView.cpp
@@ -73,7 +73,7 @@ public:
 				{
 					const float DistanceToSelected = DebuggerModel->MinDistanceToSelectedArchetypes(Item);
 					Color = UE::Mass::Debugger::UI::Private::GetArchetypeDistanceColor(DistanceToSelected);
-					if (Item->EntitiesCount == 0)
+					if (Item->ArchetypeStats.EntitiesCount == 0)
 					{
 						Color.A *= 0.5f;
 					}
@@ -82,7 +82,7 @@ public:
 			else
 			{
 				Color = FLinearColor::White;
-				if (Item->EntitiesCount == 0)
+				if (Item->ArchetypeStats.EntitiesCount == 0)
 				{
 					Color.A *= 0.5f;
 				}
@@ -130,7 +130,7 @@ public:
 				.Padding(5, 0)
 				[
 					SNew(STextBlock)
-					.Text(FText::FromString(FString::Printf(TEXT("%d"), Item->EntitiesCount)))
+					.Text(FText::FromString(FString::Printf(TEXT("%d"), Item->ArchetypeStats.EntitiesCount)))
 					.Justification(ETextJustify::Right)
 				];
 		}
@@ -201,7 +201,7 @@ void SMassArchetypesView::Construct(const FArguments& InArgs, TSharedRef<FMassDe
 						for (const TSharedPtr<FMassDebuggerArchetypeData>& ArchetypeData : DebuggerModel->CachedArchetypes)
 						{
 							TotalCount++;
-							if (ArchetypeData->EntitiesCount > 0)
+							if (ArchetypeData->ArchetypeStats.EntitiesCount > 0)
 							{
 								ActiveCount++;
 							}
@@ -209,7 +209,7 @@ void SMassArchetypesView::Construct(const FArguments& InArgs, TSharedRef<FMassDe
 							for (const TSharedPtr<FMassDebuggerArchetypeData>& ChildArchetypeData : ArchetypeData->Children)
 							{
 								TotalCount++;
-								if (ChildArchetypeData->EntitiesCount > 0)
+								if (ChildArchetypeData->ArchetypeStats.EntitiesCount > 0)
 								{
 									ActiveCount++;
 								}
