Fix CIS not compiling in Shipping Editor

diff --git a/Engine/Plugins/Runtime/MassEntity/Source/MassEntity/Private/MassDebugger.cpp b/Engine/Plugins/Runtime/MassEntity/Source/MassEntity/Private/MassDebugger.cpp
--- a/Engine/Plugins/Runtime/MassEntity/Source/MassEntity/Private/MassDebugger.cpp
+++ b/Engine/Plugins/Runtime/MassEntity/Source/MassEntity/Private/MassDebugger.cpp
@@ -256,7 +256,7 @@ const FMassArchetypeCompositionDescriptor& FMassDebugger::GetArchetypeCompositio
 	return ArchetypeData.CompositionDescriptor;
 }
 
-void FMassDebugger::GetArchetypeEntityStats(const FMassArchetypeHandle& ArchetypeHandle, FArchetypeStats& OutStats)
+void FMassDebugger::GetArchetypeEntityStats(const FMassArchetypeHandle& ArchetypeHandle, UE::Mass::Debug::FArchetypeStats& OutStats)
 {
 	const FMassArchetypeData& ArchetypeData = FMassArchetypeHelper::ArchetypeDataFromHandleChecked(ArchetypeHandle);
 	OutStats.EntitiesCount = ArchetypeData.GetNumEntities();
diff --git a/Engine/Plugins/Runtime/MassEntity/Source/MassEntity/Public/MassDebugger.h b/Engine/Plugins/Runtime/MassEntity/Source/MassEntity/Public/MassDebugger.h
--- a/Engine/Plugins/Runtime/MassEntity/Source/MassEntity/Public/MassDebugger.h
+++ b/Engine/Plugins/Runtime/MassEntity/Source/MassEntity/Public/MassDebugger.h
@@ -19,9 +19,9 @@ struct FMassFragmentRequirementDescription;
 enum class EMassFragmentAccess : uint8;
 enum class EMassFragmentPresence : uint8;
 
-#if WITH_MASSENTITY_DEBUG
 namespace UE::Mass::Debug
 {
+#if WITH_MASSENTITY_DEBUG
 struct MASSENTITY_API FQueryRequirementsView
 {
 	TConstArrayView<FMassFragmentRequirementDescription> FragmentRequirements;
@@ -37,8 +37,23 @@ struct MASSENTITY_API FQueryRequirementsView
 
 FString DebugGetFragmentAccessString(EMassFragmentAccess Access);
 MASSENTITY_API extern void DebugOutputDescription(TConstArrayView<UMassProcessor*> Processors, FOutputDevice& Ar);
+#endif // WITH_MASSENTITY_DEBUG
+
+struct FArchetypeStats
+{
+	/** Number of active entities of the archetype. */
+	int32 EntitiesCount = 0;
+	/** Number of entities that fit per chunk. */
+	int32 EntitiesCountPerChunk = 0;
+	/** Number of allocated chunks. */
+	int32 ChunksCount = 0;
+	/** Total amount of memory taken by this archetype */
+	SIZE_T AllocatedSize = 0;
+};
 } // namespace UE::Mass::Debug
 
+#if WITH_MASSENTITY_DEBUG
+
 struct MASSENTITY_API FMassDebugger
 {
 	DECLARE_MULTICAST_DELEGATE_TwoParams(FOnEntitySelected, const FMassEntityManager&, const FMassEntityHandle);
@@ -54,18 +69,7 @@ struct MASSENTITY_API FMassDebugger
 	static TArray<FMassArchetypeHandle> GetAllArchetypes(const FMassEntityManager& EntitySubsystem);
 	static const FMassArchetypeCompositionDescriptor& GetArchetypeComposition(const FMassArchetypeHandle& ArchetypeHandle);
 
-	struct FArchetypeStats
-	{
-		/** Number of active entities of the archetype. */
-		int32 EntitiesCount = 0;
-		/** Number of entities that fit per chunk. */
-		int32 EntitiesCountPerChunk = 0;
-		/** Number of allocated chunks. */
-		int32 ChunksCount = 0;
-		/** Total amount of memory taken by this archetype */
-		SIZE_T AllocatedSize = 0;
-	};
-	static void GetArchetypeEntityStats(const FMassArchetypeHandle& ArchetypeHandle, FArchetypeStats& OutStats);
+	static void GetArchetypeEntityStats(const FMassArchetypeHandle& ArchetypeHandle, UE::Mass::Debug::FArchetypeStats& OutStats);
 	static const TConstArrayView<FName> GetArchetypeDebugNames(const FMassArchetypeHandle& ArchetypeHandle);
 
 	static TConstArrayView<UMassCompositeProcessor::FDependencyNode> GetProcessingGraph(const UMassCompositeProcessor& GraphOwner);
diff --git a/Engine/Plugins/Runtime/MassEntity/Source/MassEntityDebugger/Private/MassDebuggerModel.h b/Engine/Plugins/Runtime/MassEntity/Source/MassEntityDebugger/Private/MassDebuggerModel.h
--- a/Engine/Plugins/Runtime/MassEntity/Source/MassEntityDebugger/Private/MassDebuggerModel.h
+++ b/Engine/Plugins/Runtime/MassEntity/Source/MassEntityDebugger/Private/MassDebuggerModel.h
@@ -66,7 +66,7 @@ struct FMassDebuggerArchetypeData
 	uint32 FullHash = 0;
 
 	/** Archetype statistics */
-	FMassDebugger::FArchetypeStats ArchetypeStats;
+	UE::Mass::Debug::FArchetypeStats ArchetypeStats;
 
 	/** Child debugger data (same as parent, but changed in some way) */
 	TArray<TSharedPtr<FMassDebuggerArchetypeData>> Children;
