Amendend FMassEntityManager's comment with an important detail.

diff --git a/Engine/Plugins/Runtime/MassEntity/Source/MassEntity/Public/MassEntityManager.h b/Engine/Plugins/Runtime/MassEntity/Source/MassEntity/Public/MassEntityManager.h
--- a/Engine/Plugins/Runtime/MassEntity/Source/MassEntity/Public/MassEntityManager.h
+++ b/Engine/Plugins/Runtime/MassEntity/Source/MassEntity/Public/MassEntityManager.h
@@ -35,6 +35,7 @@ enum class EMassFragmentAccess : uint8;
  * FMassEntityManager are meant to be stored with a TSharedPtr or TSharedRef. Some of Mass API pass around 
  * FMassEntityManager& but programmers can always use AsShared() call to obtain a shared ref for a given manager instance 
  * (as supplied by deriving from TSharedFromThis<FMassEntityManager>).
+ * IMPORTANT: if you create your own FMassEntityManager instance remember to call Initialize() before using it.
  */
 struct MASSENTITY_API FMassEntityManager : public TSharedFromThis<FMassEntityManager>
 {
