Brough back the Farm test as part of MassEntity's test suite.

diff --git a/Engine/Plugins/Runtime/MassEntity/Source/MassEntityTestSuite/Private/LWComponentTestFarmPlot.cpp b/Engine/Plugins/Runtime/MassEntity/Source/MassEntityTestSuite/Private/LWComponentTestFarmPlot.cpp
deleted file mode 100644
diff --git a/Engine/Plugins/Runtime/MassEntity/Source/MassEntityTestSuite/Private/MassEntityTestFarmPlot.cpp b/Engine/Plugins/Runtime/MassEntity/Source/MassEntityTestSuite/Private/MassEntityTestFarmPlot.cpp
new file mode 100644
--- /dev/null
+++ b/Engine/Plugins/Runtime/MassEntity/Source/MassEntityTestSuite/Private/MassEntityTestFarmPlot.cpp
@@ -0,0 +1,311 @@
+// Copyright Epic Games, Inc. All Rights Reserved.
+
+#include "MassEntityTestFarmPlot.h"
+
+#include "Components/HierarchicalInstancedStaticMeshComponent.h"
+#include "Engine/CollisionProfile.h"
+#include "MassEntityUtils.h"
+#include "MassExecutor.h"
+
+//@TODO: Can add a ReadyToHarvest tag Fragment on when things are ready to harvest, to stop them ticking and signal that we need to create an icon
+
+//////////////////////////////////////////////////////////////////////////
+// UFarmProcessorBase
+UFarmProcessorBase::UFarmProcessorBase()
+	: EntityQuery(*this)
+{
+	// not auto-registering to manually control execution
+	bAutoRegisterWithProcessingPhases = false;
+}
+
+//////////////////////////////////////////////////////////////////////////
+// UFarmWaterProcessor
+
+void UFarmWaterProcessor::ConfigureQueries()
+{
+	EntityQuery.AddRequirement<FFarmWaterFragment>(EMassFragmentAccess::ReadWrite);
+}
+
+void UFarmWaterProcessor::Execute(FMassEntityManager& EntityManager, FMassExecutionContext& Context)
+{
+	QUICK_SCOPE_CYCLE_COUNTER(UFarmWaterProcessor_Run);
+	EntityQuery.ForEachEntityChunk(EntityManager, Context, [this](FMassExecutionContext& Context) {
+
+		const float DeltaTimeSeconds = Context.GetDeltaTimeSeconds();
+		TArrayView<FFarmWaterFragment> WaterList = Context.GetMutableFragmentView<FFarmWaterFragment>();
+
+		for (FFarmWaterFragment& WaterFragment : WaterList)
+		{
+			WaterFragment.CurrentWater = FMath::Clamp(WaterFragment.CurrentWater + WaterFragment.DeltaWaterPerSecond * DeltaTimeSeconds, 0.0f, 1.0f);
+		}
+	});
+}
+
+//////////////////////////////////////////////////////////////////////////
+// UFarmHarvestTimerSystem_Flowers
+
+void UFarmHarvestTimerSystem_Flowers::ConfigureQueries()
+{
+	EntityQuery.AddRequirement<FHarvestTimerFragment>(EMassFragmentAccess::ReadWrite);
+	EntityQuery.AddRequirement<FFarmWaterFragment>(EMassFragmentAccess::ReadOnly);
+	EntityQuery.AddRequirement<FFarmFlowerFragment>(EMassFragmentAccess::ReadWrite);
+}
+
+void UFarmHarvestTimerSystem_Flowers::Execute(FMassEntityManager& EntityManager, FMassExecutionContext& Context)
+{
+	QUICK_SCOPE_CYCLE_COUNTER(UFarmHarvestTimerSystem_Flowers_Run);
+
+	EntityQuery.ForEachEntityChunk(EntityManager, Context, [this](FMassExecutionContext& Context) {
+		const int32 NumEntities = Context.GetNumEntities();
+		const float WellWateredThreshold = 0.25f;
+		TArrayView<FHarvestTimerFragment> TimerList = Context.GetMutableFragmentView<FHarvestTimerFragment>();
+		TConstArrayView<FFarmWaterFragment> WaterList = Context.GetFragmentView<FFarmWaterFragment>();
+		TArrayView<FFarmFlowerFragment> FlowerList = Context.GetMutableFragmentView<FFarmFlowerFragment>();
+
+		for (int32 i = 0; i < NumEntities; ++i)
+		{
+			if (TimerList[i].NumSecondsLeft > 0)
+			{
+				--TimerList[i].NumSecondsLeft;
+
+				if (WaterList[i].CurrentWater > WellWateredThreshold)
+				{
+					++FlowerList[i].NumBonusTicks;
+				}
+			}
+		}
+	});
+}
+
+//////////////////////////////////////////////////////////////////////////
+// UFarmHarvestTimerSystem_Crops
+
+void UFarmHarvestTimerSystem_Crops::ConfigureQueries()
+{
+	EntityQuery.AddRequirement<FHarvestTimerFragment>(EMassFragmentAccess::ReadWrite);
+	EntityQuery.AddRequirement<FFarmWaterFragment>(EMassFragmentAccess::ReadOnly);
+	EntityQuery.AddRequirement<FFarmCropFragment>(EMassFragmentAccess::ReadOnly);
+}
+
+void UFarmHarvestTimerSystem_Crops::Execute(FMassEntityManager & EntityManager, FMassExecutionContext & Context)
+{
+	QUICK_SCOPE_CYCLE_COUNTER(UFarmHarvestTimerSystem_Crops_Run);
+
+	EntityQuery.ForEachEntityChunk(EntityManager, Context, [this](FMassExecutionContext& Context) {
+
+		const int32 NumEntities = Context.GetNumEntities();
+		const float WellWateredThreshold = 0.25f;
+		TArrayView<FHarvestTimerFragment> TimerList = Context.GetMutableFragmentView<FHarvestTimerFragment>();
+		TConstArrayView<FFarmWaterFragment> WaterList = Context.GetMutableFragmentView<FFarmWaterFragment>();
+
+		for (int32 i = 0; i < NumEntities; ++i)
+		{
+			const uint32 TimeToSubtract = (WaterList[i].CurrentWater > WellWateredThreshold) ? 2 : 1;
+			TimerList[i].NumSecondsLeft = (TimerList[i].NumSecondsLeft >= TimeToSubtract) ? (TimerList[i].NumSecondsLeft - TimeToSubtract) : 0;
+		}
+	});
+}
+
+//////////////////////////////////////////////////////////////////////////
+// UFarmHarvestTimerExpired
+
+void UFarmHarvestTimerExpired::ConfigureQueries()
+{
+	EntityQuery.AddRequirement<FHarvestTimerFragment>(EMassFragmentAccess::ReadOnly);
+	EntityQuery.AddTagRequirement<FFarmJustBecameReadyToHarvestTag>(EMassFragmentPresence::None);
+	EntityQuery.AddTagRequirement<FFarmReadyToHarvestTag>(EMassFragmentPresence::None);
+}
+
+void UFarmHarvestTimerExpired::Execute(FMassEntityManager& EntityManager, FMassExecutionContext& Context)
+{
+	QUICK_SCOPE_CYCLE_COUNTER(UFarmHarvestTimerExpired_Run);
+
+	EntityQuery.ForEachEntityChunk(EntityManager, Context, [this](FMassExecutionContext& Context) {
+		const int32 NumEntities = Context.GetNumEntities();
+		TConstArrayView<FHarvestTimerFragment> TimerList = Context.GetFragmentView<FHarvestTimerFragment>();
+
+		for (int32 i = 0; i < NumEntities; ++i)
+		{
+			if (TimerList[i].NumSecondsLeft == 0)
+			{
+				Context.Defer().AddTag<FFarmJustBecameReadyToHarvestTag>(Context.GetEntity(i));
+			}
+		}
+	});
+}
+
+//////////////////////////////////////////////////////////////////////////
+// UFarmHarvestTimerSetIcon
+
+void UFarmHarvestTimerSetIcon::ConfigureQueries()
+{
+	EntityQuery.AddRequirement<FFarmGridCellData>(EMassFragmentAccess::ReadOnly);
+	EntityQuery.AddRequirement<FFarmVisualFragment>(EMassFragmentAccess::ReadWrite);
+	EntityQuery.AddTagRequirement<FFarmJustBecameReadyToHarvestTag>(EMassFragmentPresence::All);
+}
+
+void UFarmHarvestTimerSetIcon::Execute(FMassEntityManager& EntityManager, FMassExecutionContext& Context)
+{
+	if (HarvestIconISMC == nullptr)
+	{
+		return;
+	}
+
+	QUICK_SCOPE_CYCLE_COUNTER(SET_ICON_SET_ICON_SET_ICON);
+
+	EntityQuery.ForEachEntityChunk(EntityManager, Context, [this](FMassExecutionContext& Context) {
+
+		const int32 NumEntities = Context.GetNumEntities();
+		TConstArrayView<FFarmGridCellData> GridCoordList = Context.GetFragmentView<FFarmGridCellData>();
+		TArrayView<FFarmVisualFragment> VisualList = Context.GetMutableFragmentView<FFarmVisualFragment>();
+
+		for (int32 i = 0; i < NumEntities; ++i)
+		{
+			const FFarmGridCellData& GridCells = GridCoordList[i];
+
+			const FVector IconPosition(GridCells.CellX*GridCellWidth, GridCells.CellY*GridCellHeight, HarvestIconHeight);
+			const FTransform IconTransform(FQuat::Identity, IconPosition, FVector(HarvestIconScale, HarvestIconScale, HarvestIconScale));
+
+			VisualList[i].HarvestIconIndex = HarvestIconISMC->AddInstance(IconTransform);
+
+			FMassEntityHandle ThisEntity = Context.GetEntity(i);
+			Context.Defer().RemoveTag<FFarmJustBecameReadyToHarvestTag>(ThisEntity);
+			Context.Defer().AddTag<FFarmReadyToHarvestTag>(ThisEntity);
+		}
+	});
+}
+
+//////////////////////////////////////////////////////////////////////////
+// ALWFragmentTestFarmPlot
+
+AMassEntityTestFarmPlot::AMassEntityTestFarmPlot()
+	: SharedEntityManager(MakeShareable(new FMassEntityManager(this)))
+{
+	PrimaryActorTick.bCanEverTick = true;
+	PrimaryActorTick.bStartWithTickEnabled = true;
+
+	USceneComponent* SceneComponent = CreateDefaultSubobject<USceneComponent>(TEXT("SceneComp"));
+
+	RootComponent = SceneComponent;
+
+	HarvestIconISMC = CreateDefaultSubobject<UHierarchicalInstancedStaticMeshComponent>(TEXT("HarvestIconISMC"));
+	HarvestIconISMC->SetupAttachment(SceneComponent);
+	HarvestIconISMC->SetCollisionProfileName(UCollisionProfile::NoCollision_ProfileName);
+}
+
+void AMassEntityTestFarmPlot::AddItemToGrid(FMassEntityManager& InEntityManager, uint16 X, uint16 Y, FMassArchetypeHandle Archetype, uint16 VisualIndex)
+{
+	FMassEntityHandle NewItem = InEntityManager.CreateEntity(Archetype);
+	PlantedSquares[X + Y * GridWidth] = NewItem;
+
+	InEntityManager.GetFragmentDataChecked<FFarmWaterFragment>(NewItem).DeltaWaterPerSecond = FMath::FRandRange(-0.01f, -0.001f);
+	InEntityManager.GetFragmentDataChecked<FHarvestTimerFragment>(NewItem).NumSecondsLeft = 5 + (FMath::Rand() % 100);
+
+	FFarmGridCellData GridCoords;
+	GridCoords.CellX = X;
+	GridCoords.CellY = Y;
+	InEntityManager.GetFragmentDataChecked<FFarmGridCellData>(NewItem) = GridCoords;
+
+	const FVector MeshPosition(X*GridCellWidth, Y*GridCellHeight, 0.0f);
+	const FRotator MeshRotation(0.0f, FMath::FRand()*360.0f, 0.0f);
+	const FVector MeshScale(1.0f, 1.0f, 1.0f); //@TODO: plumb in scale param?
+	const FTransform MeshTransform(MeshRotation, MeshPosition, MeshScale);
+
+	FFarmVisualFragment& VisualComp = InEntityManager.GetFragmentDataChecked<FFarmVisualFragment>(NewItem);
+	VisualComp.VisualType = VisualIndex;
+	VisualComp.InstanceIndex = VisualDataISMCs[VisualComp.VisualType]->AddInstance(MeshTransform);
+}
+
+void AMassEntityTestFarmPlot::BeginPlay()
+{
+	Super::BeginPlay();
+
+	if (TestDataCropIndicies.Num() == 0
+		|| TestDataFlowerIndicies.Num() == 0
+		|| VisualDataTable.Num() == 0)
+	{
+		UE_LOG(LogMass, Error, TEXT("%s is misconfigured. Make sure TestDataCropIndicies, TestDataFlowerIndicies and VisualDataTable are not empty"), *GetName());
+		return;
+	}
+
+	FMassEntityManager& EntityManager = SharedEntityManager.Get();
+	EntityManager.Initialize();
+
+	FMassArchetypeHandle CropArchetype = EntityManager.CreateArchetype(TArray<const UScriptStruct*>{ FFarmWaterFragment::StaticStruct(), FFarmCropFragment::StaticStruct(), FHarvestTimerFragment::StaticStruct(), FFarmVisualFragment::StaticStruct(), FFarmGridCellData::StaticStruct() });
+	FMassArchetypeHandle FlowerArchetype = EntityManager.CreateArchetype(TArray<const UScriptStruct*>{ FFarmWaterFragment::StaticStruct(), FFarmFlowerFragment::StaticStruct(), FHarvestTimerFragment::StaticStruct(), FFarmVisualFragment::StaticStruct(), FFarmGridCellData::StaticStruct() });
+
+	PerFrameSystems.Add(NewObject<UFarmWaterProcessor>(this));
+
+	PerSecondSystems.Add(NewObject<UFarmHarvestTimerSystem_Flowers>(this));
+	PerSecondSystems.Add(NewObject<UFarmHarvestTimerSystem_Crops>(this));
+	PerSecondSystems.Add(NewObject<UFarmHarvestTimerExpired>(this));
+
+	UFarmHarvestTimerSetIcon* IconSetter = NewObject<UFarmHarvestTimerSetIcon>(this);
+	IconSetter->HarvestIconISMC = HarvestIconISMC;
+	IconSetter->GridCellWidth = GridCellWidth;
+	IconSetter->GridCellHeight = GridCellHeight;
+	IconSetter->HarvestIconHeight = 200.0f;
+	IconSetter->HarvestIconScale = HarvestIconScale;
+	PerSecondSystems.Add(IconSetter);
+
+	HarvestIconISMC->SetCullDistances(IconNearCullDistance, IconFarCullDistance);
+
+	// Create ISMCs for each mesh type
+	for (const FFarmVisualDataRow& VisualData : VisualDataTable)
+	{
+		UHierarchicalInstancedStaticMeshComponent* HISMC = NewObject<UHierarchicalInstancedStaticMeshComponent>(this);
+		HISMC->SetStaticMesh(VisualData.Mesh);
+		if (VisualData.MaterialOverride != nullptr)
+		{
+			HISMC->SetMaterial(0, VisualData.MaterialOverride);
+		}
+		HISMC->SetCullDistances(VisualNearCullDistance, VisualFarCullDistance);
+		HISMC->SetupAttachment(RootComponent);
+		HISMC->RegisterComponent();
+
+		VisualDataISMCs.Add(HISMC);
+	}
+
+	// Plant us a grid
+	const int32 NumGridCells = GridWidth * GridHeight;
+	PlantedSquares.AddDefaulted(NumGridCells);
+
+	for (uint16 Y = 0; Y < GridHeight; ++Y)
+	{
+		for (uint16 X = 0; X < GridWidth; ++X)
+		{
+			const bool bSpawnCrop = FMath::FRand() < 0.5f;
+			const uint16 VisualIndex = bSpawnCrop ? TestDataCropIndicies[FMath::RandRange(0, TestDataCropIndicies.Num() - 1)] : TestDataFlowerIndicies[FMath::RandRange(0, TestDataFlowerIndicies.Num() - 1)];
+
+			AddItemToGrid(EntityManager, X, Y, bSpawnCrop ? CropArchetype : FlowerArchetype, VisualIndex);
+		}
+	}
+}
+
+void AMassEntityTestFarmPlot::TickActor(float DeltaTime, enum ELevelTick TickType, FActorTickFunction& ThisTickFunction)
+{
+	Super::TickActor(DeltaTime, TickType, ThisTickFunction);
+
+	if (PlantedSquares.Num() == 0)
+	{
+		// not configured properly, ignore. 
+		return;
+	}
+
+	FMassEntityManager& EntityManager = SharedEntityManager.Get();
+	
+	// Run every frame systems
+	{
+		FMassProcessingContext Context(EntityManager, DeltaTime);
+		UE::Mass::Executor::RunProcessorsView(PerFrameSystems, Context);
+	}
+
+	// Run per-second systems when it's time
+	NextSecondTimer -= DeltaTime;
+	while (NextSecondTimer < 0.0f)
+	{
+		FMassProcessingContext Context(EntityManager, 1.f);
+		UE::Mass::Executor::RunProcessorsView(PerSecondSystems, Context);
+		NextSecondTimer += 1.0f;
+	}
+}
diff --git a/Engine/Plugins/Runtime/MassEntity/Source/MassEntityTestSuite/Public/LWComponentTestFarmPlot.h b/Engine/Plugins/Runtime/MassEntity/Source/MassEntityTestSuite/Public/LWComponentTestFarmPlot.h
deleted file mode 100644
diff --git a/Engine/Plugins/Runtime/MassEntity/Source/MassEntityTestSuite/Public/MassEntityTestFarmPlot.h b/Engine/Plugins/Runtime/MassEntity/Source/MassEntityTestSuite/Public/MassEntityTestFarmPlot.h
new file mode 100644
--- /dev/null
+++ b/Engine/Plugins/Runtime/MassEntity/Source/MassEntityTestSuite/Public/MassEntityTestFarmPlot.h
@@ -0,0 +1,251 @@
+// Copyright Epic Games, Inc. All Rights Reserved.
+
+#pragma once
+
+#include "CoreMinimal.h"
+#include "GameFramework/Actor.h"
+#include "Engine/DataTable.h"
+
+#include "MassEntityTypes.h"
+#include "MassEntityManager.h"
+#include "MassCommandBuffer.h"
+#include "MassProcessor.h"
+
+#include "MassEntityTestFarmPlot.generated.h"
+
+class UStaticMesh;
+class UMaterialInterface;
+class UHierarchicalInstancedStaticMeshComponent;
+
+USTRUCT()
+struct FFarmVisualDataRow : public FTableRowBase
+{
+	GENERATED_BODY()
+
+	UPROPERTY(EditAnywhere, Category=Farm)
+	TObjectPtr<UStaticMesh> Mesh = nullptr;
+
+	UPROPERTY(EditAnywhere, Category=Farm)
+	TObjectPtr<UMaterialInterface> MaterialOverride = nullptr;
+};
+
+USTRUCT()
+struct FFarmJustBecameReadyToHarvestTag : public FMassTag
+{
+	GENERATED_BODY()
+};
+
+USTRUCT()
+struct FFarmReadyToHarvestTag : public FMassTag
+{
+	GENERATED_BODY()
+};
+
+USTRUCT()
+struct FFarmGridCellData : public FMassFragment
+{
+	GENERATED_BODY()
+
+	uint16 CellX = 0;
+	uint16 CellY = 0;
+};
+
+USTRUCT()
+struct FFarmWaterFragment : public FMassFragment
+{
+	GENERATED_BODY()
+
+	float CurrentWater = 1.0f;
+	float DeltaWaterPerSecond = -0.01f;
+};
+
+USTRUCT()
+struct FFarmFlowerFragment : public FMassFragment
+{
+	GENERATED_BODY()
+		
+	uint32 NumBonusTicks = 0;
+	uint16 FlowerType = 0;
+};
+
+USTRUCT()
+struct FFarmCropFragment : public FMassFragment
+{
+	GENERATED_BODY()
+
+	uint16 CropType = 0;
+};
+
+
+USTRUCT()
+struct FFarmVisualFragment : public FMassFragment
+{
+	GENERATED_BODY()
+
+	int32 InstanceIndex = -1;
+	int32 HarvestIconIndex = -1;
+	uint16 VisualType = 0;
+};
+
+USTRUCT()
+struct FHarvestTimerFragment : public FMassFragment
+{
+	GENERATED_BODY()
+
+	uint32 NumSecondsLeft = 15;
+};
+
+//////////////////////////////////////////////////////////////////////
+
+UCLASS(abstract)
+class UFarmProcessorBase : public UMassProcessor
+{
+	GENERATED_BODY()
+
+public:
+	UFarmProcessorBase();
+
+protected:
+	FMassEntityQuery EntityQuery;
+};
+
+//////////////////////////////////////////////////////////////////////
+
+UCLASS()
+class UFarmWaterProcessor : public UFarmProcessorBase
+{
+	GENERATED_BODY()
+
+public:
+	virtual void ConfigureQueries() override;
+	virtual void Execute(FMassEntityManager& EntityManager, FMassExecutionContext& Context) override;
+};
+
+//////////////////////////////////////////////////////////////////////
+
+UCLASS()
+class UFarmHarvestTimerSystem_Flowers : public UFarmProcessorBase
+{
+	GENERATED_BODY()
+
+	virtual void ConfigureQueries() override;
+	virtual void Execute(FMassEntityManager& EntityManager, FMassExecutionContext& Context) override;
+};
+
+//////////////////////////////////////////////////////////////////////
+
+UCLASS()
+class UFarmHarvestTimerSystem_Crops : public UFarmProcessorBase
+{
+	GENERATED_BODY()
+
+	virtual void ConfigureQueries() override;
+	virtual void Execute(FMassEntityManager& EntityManager, FMassExecutionContext& Context) override;
+};
+
+//////////////////////////////////////////////////////////////////////
+
+UCLASS()
+class UFarmHarvestTimerExpired : public UFarmProcessorBase
+{
+	GENERATED_BODY()
+
+	virtual void ConfigureQueries() override;
+	virtual void Execute(FMassEntityManager& EntityManager, FMassExecutionContext& Context) override;
+};
+
+//////////////////////////////////////////////////////////////////////
+
+UCLASS()
+class UFarmHarvestTimerSetIcon : public UFarmProcessorBase
+{
+	GENERATED_BODY()
+
+public:
+	UPROPERTY()
+	TObjectPtr<UHierarchicalInstancedStaticMeshComponent> HarvestIconISMC;
+
+	float GridCellWidth;
+	float GridCellHeight;
+	float HarvestIconHeight;
+	float HarvestIconScale;
+
+protected:
+	virtual void ConfigureQueries() override;
+	virtual void Execute(FMassEntityManager& EntityManager, FMassExecutionContext& Context) override;
+};
+
+//////////////////////////////////////////////////////////////////////
+
+UCLASS(config=Game)
+class AMassEntityTestFarmPlot : public AActor
+{
+	GENERATED_BODY()
+
+public:
+	uint16 GridWidth = 40*7;
+	uint16 GridHeight = 20*7;
+
+	UPROPERTY(EditAnywhere, Category=Farm)
+	float GridCellWidth = 150.0f;
+	
+	UPROPERTY(EditAnywhere, Category=Farm)
+	float GridCellHeight = 150.0f;
+	
+	UPROPERTY(EditAnywhere, Category=Farm)
+	float HarvestIconScale = 0.3f;
+	
+	TArray<FMassEntityHandle> PlantedSquares;
+
+	
+	UPROPERTY(EditAnywhere, Category=Farm)
+	TArray<FFarmVisualDataRow> VisualDataTable;
+
+	UPROPERTY(Transient)
+	TArray<TObjectPtr<UHierarchicalInstancedStaticMeshComponent>> VisualDataISMCs;
+
+	float NextSecondTimer = 0.0f;
+
+	UPROPERTY(Transient)
+	TArray<TObjectPtr<UMassProcessor>> PerFrameSystems;
+
+	UPROPERTY(Transient)
+	TArray<TObjectPtr<UMassProcessor>> PerSecondSystems;
+
+	// Indices into VisualDataTable for flowers
+	UPROPERTY(EditAnywhere, Category=Farm)
+	TArray<uint16> TestDataFlowerIndicies;
+
+	// Indices into VisualDataTable for crops
+	UPROPERTY(EditAnywhere, Category=Farm)
+	TArray<uint16> TestDataCropIndicies;
+
+	UPROPERTY(EditAnywhere, Category = Farm)
+	uint32 VisualNearCullDistance = 1000;
+
+	UPROPERTY(EditAnywhere, Category = Farm)
+	uint32 VisualFarCullDistance = 1200;
+
+	UPROPERTY(EditAnywhere, Category = Farm)
+	uint32 IconNearCullDistance = 400;
+
+	UPROPERTY(EditAnywhere, Category = Farm)
+	uint32 IconFarCullDistance = 800;
+
+private:
+	UPROPERTY(VisibleAnywhere, BlueprintReadOnly, meta=(AllowPrivateAccess = "true"), Category=Farm)
+	TObjectPtr<UHierarchicalInstancedStaticMeshComponent> HarvestIconISMC;
+
+	TArray<int32> FreeHarvestIconIndicies;
+
+	TSharedRef<FMassEntityManager> SharedEntityManager;
+
+private:
+	void AddItemToGrid(FMassEntityManager& InEntityManager, uint16 X, uint16 Y, FMassArchetypeHandle Archetype, uint16 VisualIndex);
+
+public:
+	AMassEntityTestFarmPlot();
+
+	virtual void BeginPlay() override;
+	virtual void TickActor(float DeltaTime, enum ELevelTick TickType, FActorTickFunction& ThisTickFunction) override;
+};
