[MassDebugger] added support for EntityManager used outside MassEntitySubsystem.
- Current implementation still requires a valid associated World.
- fixed crash when switching environment while archetypes are selected
- as part of the change EntityManager calls its Deinitialize in its constructor

Original chabge by [REDACTED_HANDLE]

diff --git a/Engine/Plugins/Runtime/MassEntity/Source/MassEntity/Private/MassDebugger.cpp b/Engine/Plugins/Runtime/MassEntity/Source/MassEntity/Private/MassDebugger.cpp
--- a/Engine/Plugins/Runtime/MassEntity/Source/MassEntity/Private/MassDebugger.cpp
+++ b/Engine/Plugins/Runtime/MassEntity/Source/MassEntity/Private/MassDebugger.cpp
@@ -20,7 +20,6 @@ namespace UE::Mass::Debug
 
 	FString DebugGetFragmentAccessString(EMassFragmentAccess Access)
 	{
-#if WITH_MASSENTITY_DEBUG
 		switch (Access)
 		{
 		case EMassFragmentAccess::None:	return TEXT("--");
@@ -30,13 +29,11 @@ namespace UE::Mass::Debug
 			ensureMsgf(false, TEXT("Missing string conversion for EMassFragmentAccess=%d"), Access);
 			break;
 		}
-#endif // WITH_MASSENTITY_DEBUG
 		return TEXT("Missing string conversion");
 	}
 
 	void DebugOutputDescription(TConstArrayView<UMassProcessor*> Processors, FOutputDevice& Ar)
 	{
-#if WITH_MASSENTITY_DEBUG
 		const bool bAutoLineEnd = Ar.GetAutoEmitLineTerminator();
 		Ar.SetAutoEmitLineTerminator(false);
 		for (const UMassProcessor* Proc : Processors)
@@ -52,10 +49,8 @@ namespace UE::Mass::Debug
 			}
 		}
 		Ar.SetAutoEmitLineTerminator(bAutoLineEnd);
-#endif // WITH_MASSENTITY_DEBUG
 	}
 
-#if WITH_MASSENTITY_DEBUG && WITH_MASSENTITY_DEBUG
 	FAutoConsoleCommandWithWorldArgsAndOutputDevice PrintEntityFragmentsCmd(
 		TEXT("mass.PrintEntityFragments"),
 		TEXT("Prints all fragment types and values (uproperties) for the specified Entity index"),
@@ -194,8 +189,6 @@ namespace UE::Mass::Debug
 				PrintKnownTypes(FMassChunkFragmentBitSet::DebugGetAllStructTypes());
 			}));
 
-#endif // WITH_MASSENTITY_DEBUG
-
 } // namespace UE::Mass::Debug
 
 //----------------------------------------------------------------------//
@@ -203,6 +196,11 @@ namespace UE::Mass::Debug
 //----------------------------------------------------------------------//
 FMassDebugger::FOnEntitySelected FMassDebugger::OnEntitySelectedDelegate;
 
+FMassDebugger::FOnMassEntityManagerEvent FMassDebugger::OnEntityManagerInitialized;
+FMassDebugger::FOnMassEntityManagerEvent FMassDebugger::OnEntityManagerDeinitialized;
+TArray<TWeakPtr<const FMassEntityManager>> FMassDebugger::ActiveEntityManagers;
+UE::FSpinLock FMassDebugger::EntityManagerRegistrationLock;
+
 TConstArrayView<FMassEntityQuery*> FMassDebugger::GetProcessorQueries(const UMassProcessor& Processor)
 {
 	return Processor.OwnedQueries;
@@ -424,4 +422,30 @@ void FMassDebugger::SelectEntity(const FMassEntityManager& EntityManager, const
 	OnEntitySelectedDelegate.Broadcast(EntityManager, EntityHandle);
 }
 
+void FMassDebugger::RegisterEntityManager(FMassEntityManager& EntityManager)
+{
+	UE::TScopeLock<UE::FSpinLock> ScopeLock(EntityManagerRegistrationLock);
+
+	ActiveEntityManagers.Add(EntityManager.AsShared());
+	OnEntityManagerInitialized.Broadcast(EntityManager);
+}
+
+void FMassDebugger::UnregisterEntityManager(FMassEntityManager& EntityManager)
+{
+	UE::TScopeLock<UE::FSpinLock> ScopeLock(EntityManagerRegistrationLock);
+
+	if (EntityManager.DoesSharedInstanceExist())
+	{
+		ActiveEntityManagers.Remove(EntityManager.AsWeak());
+	}
+	else
+	{
+		ActiveEntityManagers.RemoveAll([](const TWeakPtr<const FMassEntityManager>& Item)
+			{
+				return Item.IsValid();
+			});
+	}
+	OnEntityManagerDeinitialized.Broadcast(EntityManager);
+}
+
 #endif // WITH_MASSENTITY_DEBUG
\ No newline at end of file
diff --git a/Engine/Plugins/Runtime/MassEntity/Source/MassEntity/Private/MassEntityManager.cpp b/Engine/Plugins/Runtime/MassEntity/Source/MassEntity/Private/MassEntityManager.cpp
--- a/Engine/Plugins/Runtime/MassEntity/Source/MassEntity/Private/MassEntityManager.cpp
+++ b/Engine/Plugins/Runtime/MassEntity/Source/MassEntity/Private/MassEntityManager.cpp
@@ -8,6 +8,8 @@
 #include "UObject/UObjectIterator.h"
 #include "VisualLogger/VisualLogger.h"
 #include "MassExecutionContext.h"
+#include "MassDebugger.h"
+
 
 const FMassEntityHandle FMassEntityManager::InvalidEntity;
 
@@ -42,6 +44,17 @@ FMassEntityManager::FMassEntityManager(UObject* InOwner)
 	: ObserverManager(*this)
 	, Owner(InOwner)
 {
+#if WITH_MASSENTITY_DEBUG
+	DebugName = InOwner ? (InOwner->GetName() + TEXT("_EntityManager")) : TEXT("Unset");
+#endif
+}
+
+FMassEntityManager::~FMassEntityManager()
+{
+	if (bInitialized)
+	{
+		Deinitialize();
+	}
 }
 
 void FMassEntityManager::GetResourceSizeEx(FResourceSizeEx& CumulativeResourceSize)
@@ -80,6 +93,13 @@ void FMassEntityManager::AddReferencedObjects(FReferenceCollector& Collector)
 
 void FMassEntityManager::Initialize()
 {
+	if (bInitialized)
+	{
+		UE_LOG(LogMass, Log, TEXT("Calling %s on already initialized entity manager owned by %s")
+			, ANSI_TO_TCHAR(__FUNCTION__), *GetNameSafe(Owner.Get()));
+		return;
+	}
+
 	// Index 0 is reserved so we can treat that index as an invalid entity handle
 	Entities.Add();
 	SerialNumberGenerator.fetch_add(FMath::Max(1,NumReservedEntities));
@@ -125,11 +145,15 @@ void FMassEntityManager::Initialize()
 	}
 #if WITH_MASSENTITY_DEBUG
 	RequirementAccessDetector.Initialize();	
+	FMassDebugger::RegisterEntityManager(*this);
 #endif // WITH_MASSENTITY_DEBUG
+
+	bInitialized = true;
 }
 
 void FMassEntityManager::PostInitialize()
 {
+	ensure(bInitialized);
 	// this needs to be done after all the subsystems have been initialized since some processors might want to access
 	// them during processors' initialization
 	ObserverManager.Initialize();
@@ -137,8 +161,22 @@ void FMassEntityManager::PostInitialize()
 
 void FMassEntityManager::Deinitialize()
 {
-	// closing down so no point in actually flushing commands, but need to clean them up to avoid warnings on destruction
-	DeferredCommandBuffer->CleanUp();
+	if (bInitialized)
+	{
+		// closing down so no point in actually flushing commands, but need to clean them up to avoid warnings on destruction
+		DeferredCommandBuffer->CleanUp();
+
+#if WITH_MASSENTITY_DEBUG
+		FMassDebugger::UnregisterEntityManager(*this);
+#endif // WITH_MASSENTITY_DEBUG
+
+		bInitialized = false;
+	}
+	else
+	{
+		UE_LOG(LogMass, Log, TEXT("Calling %s on already deinitialized entity manager owned by %s")
+			, ANSI_TO_TCHAR(__FUNCTION__), *GetNameSafe(Owner.Get()));
+	}
 }
 
 FMassArchetypeHandle FMassEntityManager::CreateArchetype(TConstArrayView<const UScriptStruct*> FragmentsAndTagsList, const FName ArchetypeDebugName)
@@ -1314,6 +1352,13 @@ void FMassEntityManager::FlushCommands(const TSharedPtr<FMassCommandBuffer>& InC
 	}
 }
 
+void FMassEntityManager::SetDebugName(const FString& NewDebugGame) 
+{ 
+#if WITH_MASSENTITY_DEBUG
+	DebugName = NewDebugGame; 
+#endif // WITH_MASSENTITY_DEBUG
+}
+
 #if WITH_MASSENTITY_DEBUG
 void FMassEntityManager::DebugPrintArchetypes(FOutputDevice& Ar, const bool bIncludeEmpty) const
 {
diff --git a/Engine/Plugins/Runtime/MassEntity/Source/MassEntity/Public/MassDebugger.h b/Engine/Plugins/Runtime/MassEntity/Source/MassEntity/Public/MassDebugger.h
--- a/Engine/Plugins/Runtime/MassEntity/Source/MassEntity/Public/MassDebugger.h
+++ b/Engine/Plugins/Runtime/MassEntity/Source/MassEntity/Public/MassDebugger.h
@@ -7,6 +7,7 @@
 #include "Containers/ContainersFwd.h"
 #include "MassEntityQuery.h"
 #include "MassProcessor.h"
+#include "Misc/SpinLock.h"
 #endif // WITH_MASSENTITY_DEBUG
 
 class FOutputDevice;
@@ -57,6 +58,8 @@ struct FArchetypeStats
 struct MASSENTITY_API FMassDebugger
 {
 	DECLARE_MULTICAST_DELEGATE_TwoParams(FOnEntitySelected, const FMassEntityManager&, const FMassEntityHandle);
+	DECLARE_MULTICAST_DELEGATE_OneParam(FOnMassEntityManagerEvent, const FMassEntityManager&);
+
 
 	static TConstArrayView<FMassEntityQuery*> GetProcessorQueries(const UMassProcessor& Processor);
 	/** fetches all queries registered for given Processor. Note that in order to get up to date information
@@ -85,6 +88,17 @@ struct MASSENTITY_API FMassDebugger
 	static void SelectEntity(const FMassEntityManager& EntityManager, const FMassEntityHandle EntityHandle);
 
 	static FOnEntitySelected OnEntitySelectedDelegate;
+
+	static FOnMassEntityManagerEvent OnEntityManagerInitialized;
+	static FOnMassEntityManagerEvent OnEntityManagerDeinitialized;
+
+	static void RegisterEntityManager(FMassEntityManager& EntityManager);
+	static void UnregisterEntityManager(FMassEntityManager& EntityManager);
+	static TConstArrayView<TWeakPtr<const FMassEntityManager>> GetEntityManagers() { return ActiveEntityManagers; }
+
+private:
+	static TArray<TWeakPtr<const FMassEntityManager>> ActiveEntityManagers;
+	static UE::FSpinLock EntityManagerRegistrationLock;
 };
 
 #else
diff --git a/Engine/Plugins/Runtime/MassEntity/Source/MassEntity/Public/MassEntityManager.h b/Engine/Plugins/Runtime/MassEntity/Source/MassEntity/Public/MassEntityManager.h
--- a/Engine/Plugins/Runtime/MassEntity/Source/MassEntity/Public/MassEntityManager.h
+++ b/Engine/Plugins/Runtime/MassEntity/Source/MassEntity/Public/MassEntityManager.h
@@ -82,6 +82,7 @@ public:
 
 	explicit FMassEntityManager(UObject* InOwner = nullptr);
 	FMassEntityManager(const FMassEntityManager& Other) = delete;
+	~FMassEntityManager();
 
 	void AddReferencedObjects(FReferenceCollector& Collector);
 	void GetResourceSizeEx(FResourceSizeEx& CumulativeResourceSize);
@@ -371,6 +372,7 @@ public:
 	UWorld* GetWorld() const { return Owner.IsValid() ? Owner->GetWorld() : nullptr; }
 	UObject* GetOwner() const { return Owner.Get(); }
 
+	void SetDebugName(const FString& NewDebugGame);
 #if WITH_MASSENTITY_DEBUG
 	void DebugPrintArchetypes(FOutputDevice& Ar, const bool bIncludeEmpty = true) const;
 	void DebugGetArchetypesStringDetails(FOutputDevice& Ar, const bool bIncludeEmpty = true) const;
@@ -383,6 +385,7 @@ public:
 	void DebugForceArchetypeDataVersionBump() { ++ArchetypeDataVersion; }
 	void DebugGetArchetypeStrings(const FMassArchetypeHandle& Archetype, TArray<FName>& OutFragmentNames, TArray<FName>& OutTagNames);
 	FMassEntityHandle DebugGetEntityIndexHandle(const int32 EntityIndex) const { return Entities.IsValidIndex(EntityIndex) ? FMassEntityHandle(EntityIndex, Entities[EntityIndex].SerialNumber) : FMassEntityHandle(); }
+	const FString& DebugGetName() const { return DebugName; }
 
 	FMassRequirementAccessDetector& GetRequirementAccessDetector() { return RequirementAccessDetector; }
 #endif // WITH_MASSENTITY_DEBUG
@@ -449,7 +452,10 @@ private:
 
 #if WITH_MASSENTITY_DEBUG
 	FMassRequirementAccessDetector RequirementAccessDetector;
+	FString DebugName;
 #endif // WITH_MASSENTITY_DEBUG
 
 	TWeakObjectPtr<UObject> Owner;
+
+	bool bInitialized = false;
 };
diff --git a/Engine/Plugins/Runtime/MassEntity/Source/MassEntityDebugger/Private/MassDebuggerModel.cpp b/Engine/Plugins/Runtime/MassEntity/Source/MassEntityDebugger/Private/MassDebuggerModel.cpp
--- a/Engine/Plugins/Runtime/MassEntity/Source/MassEntityDebugger/Private/MassDebuggerModel.cpp
+++ b/Engine/Plugins/Runtime/MassEntity/Source/MassEntityDebugger/Private/MassDebuggerModel.cpp
@@ -55,12 +55,27 @@ namespace UE::Mass::Debugger::Private
 //----------------------------------------------------------------------//
 FString FMassDebuggerEnvironment::GetDisplayName() const
 {
-	return World.IsValid() ? World->GetPathName() : TEXT("Invalid");
+	FString DisplayName;
+
+#if WITH_MASSENTITY_DEBUG
+	if (const FMassEntityManager* EntityManagerPtr = GetEntityManager())
+	{
+		DisplayName += EntityManagerPtr->DebugGetName();
+		if (DisplayName.Len())
+		{
+			DisplayName += TEXT(" - ");
+		}
+	}
+#endif // WITH_MASSENTITY_DEBUG
+	
+	const UWorld* WorldPtr = World.Get();
+	DisplayName += WorldPtr ? WorldPtr->GetPathName() : TEXT("Invalid");
+	return DisplayName;
 }
 
 const FMassEntityManager* FMassDebuggerEnvironment::GetEntityManager() const
 {
-	return UE::Mass::Utils::GetEntityManager(World.Get());
+	return EntityManager.Pin().Get();
 }
 
 //----------------------------------------------------------------------//
@@ -150,7 +165,7 @@ FMassDebuggerArchetypeData::FMassDebuggerArchetypeData(const FMassArchetypeHandl
 #if WITH_MASSENTITY_DEBUG
 	Composition = FMassDebugger::GetArchetypeComposition(ArchetypeHandle);
 
-	// @todo should ensure we're using same hashing as the EntitySubsystem here
+	// @todo should ensure we're using same hashing as the EntityManager here
 	CompositionHash = Composition.CalculateHash();
 	FullHash = CompositionHash;
 
@@ -304,6 +319,8 @@ void FMassDebuggerModel::RefreshAll()
 	CacheProcessorsData(TransientArchetypesMap);
 	CacheProcessingGraphs();
 
+	ClearArchetypeSelection();
+
 	OnRefreshDelegate.Broadcast();
 }
 
@@ -384,9 +401,9 @@ void FMassDebuggerModel::CacheProcessorsData(const TMap<FMassArchetypeHandle, TS
 	CachedProcessors.Reset();
 
 	UWorld* World = Environment ? Environment->World.Get() : nullptr;
-	const FMassEntityManager* EntitySubsystem = Environment ? Environment->GetEntityManager() : nullptr;
+	const FMassEntityManager* EntityManager = Environment ? Environment->GetEntityManager() : nullptr;
 
-	if (EntitySubsystem)
+	if (EntityManager)
 	{
 		check(World);
 		for (FThreadSafeObjectIterator It(UMassProcessor::StaticClass()); It; ++It)
@@ -396,7 +413,7 @@ void FMassDebuggerModel::CacheProcessorsData(const TMap<FMassArchetypeHandle, TS
 				&& Cast<UMassCompositeProcessor>(Processor) == nullptr
 				&& Processor->GetWorld() == World)
 			{
-				CachedProcessors.Add(MakeShareable(new FMassDebuggerProcessorData(*EntitySubsystem, *Processor, InTransientArchetypesMap)));
+				CachedProcessors.Add(MakeShareable(new FMassDebuggerProcessorData(*EntityManager, *Processor, InTransientArchetypesMap)));
 			}
 		}
 	}
@@ -442,9 +459,9 @@ void FMassDebuggerModel::CacheArchetypesData(TMap<FMassArchetypeHandle, TSharedP
 
 	if (Environment)
 	{
-		if (const FMassEntityManager* EntitySubsystem = Environment->GetEntityManager())
+		if (const FMassEntityManager* EntityManager = Environment->GetEntityManager())
 		{
-			StoreArchetypes(*EntitySubsystem, OutTransientArchetypesMap);
+			StoreArchetypes(*EntityManager, OutTransientArchetypesMap);
 		}
 	}
 }
diff --git a/Engine/Plugins/Runtime/MassEntity/Source/MassEntityDebugger/Private/MassDebuggerModel.h b/Engine/Plugins/Runtime/MassEntity/Source/MassEntityDebugger/Private/MassDebuggerModel.h
--- a/Engine/Plugins/Runtime/MassEntity/Source/MassEntityDebugger/Private/MassDebuggerModel.h
+++ b/Engine/Plugins/Runtime/MassEntity/Source/MassEntityDebugger/Private/MassDebuggerModel.h
@@ -135,16 +135,17 @@ struct FMassDebuggerProcessingGraph
 
 struct FMassDebuggerEnvironment
 {
-	explicit FMassDebuggerEnvironment(UWorld* InWorld)
-		: World(InWorld)
+	explicit FMassDebuggerEnvironment(const TSharedRef<const FMassEntityManager>& InEntityManager)
+		: EntityManager(InEntityManager), World(InEntityManager->GetWorld())
 	{}
 
-	bool operator==(const FMassDebuggerEnvironment& Other) const { return World == Other.World; }
+	bool operator==(const FMassDebuggerEnvironment& Other) const { return EntityManager == Other.EntityManager; }
 
 	FString GetDisplayName() const;
 	const FMassEntityManager* GetEntityManager() const;
 	bool IsWorldValid() const { return World.IsValid(); }
 	
+	TWeakPtr<const FMassEntityManager> EntityManager;
 	TWeakObjectPtr<UWorld> World;
 };
 
@@ -167,6 +168,7 @@ struct FMassDebuggerModel
 	void ClearArchetypeSelection();
 
 	bool IsCurrentEnvironment(const FMassDebuggerEnvironment& InEnvironment) const { return Environment && *Environment.Get() == InEnvironment; }
+	bool IsCurrentEnvironmentValid() const { return Environment && Environment->EntityManager.IsValid(); }
 
 	void CacheArchetypesData(TMap<FMassArchetypeHandle, TSharedPtr<FMassDebuggerArchetypeData>>& OutTransientArchetypesMap); 
 	void CacheProcessorsData(const TMap<FMassArchetypeHandle, TSharedPtr<FMassDebuggerArchetypeData>>& InTransientArchetypesMap);
diff --git a/Engine/Plugins/Runtime/MassEntity/Source/MassEntityDebugger/Private/SMassDebugger.cpp b/Engine/Plugins/Runtime/MassEntity/Source/MassEntityDebugger/Private/SMassDebugger.cpp
--- a/Engine/Plugins/Runtime/MassEntity/Source/MassEntityDebugger/Private/SMassDebugger.cpp
+++ b/Engine/Plugins/Runtime/MassEntity/Source/MassEntityDebugger/Private/SMassDebugger.cpp
@@ -53,12 +53,10 @@ SMassDebugger::SMassDebugger()
 
 SMassDebugger::~SMassDebugger()
 {
-#if WITH_EDITOR
-	FEditorDelegates::EndPIE.Remove(PIEEndHandle);
-	FWorldDelegates::OnPostDuplicate.Remove(PIEWorldInitialize);
-#endif // WITH_EDITOR
-	GEngine->OnWorldAdded().Remove(OnWorldAddedHandle);
-	GEngine->OnWorldDestroyed().Remove(OnWorldDestroyedHandle);
+#if WITH_MASSENTITY_DEBUG
+	FMassDebugger::OnEntityManagerInitialized.Remove(OnEntityManagerInitializedHandle);
+	FMassDebugger::OnEntityManagerDeinitialized.Remove(OnEntityManagerDeinitializedHandle);
+#endif // WITH_MASSENTITY_DEBUG
 }
 
 void SMassDebugger::Construct(const FArguments& InArgs, const TSharedRef<SDockTab>& ConstructUnderMajorTab, const TSharedPtr<SWindow>& ConstructUnderWindow)
@@ -199,31 +197,54 @@ TSharedRef<SDockTab> SMassDebugger::SpawnToolbar(const FSpawnTabArgs& Args)
 	return MajorTab;
 }
 
-void SMassDebugger::OnWorldAdded(UWorld* NewWorld)
+void SMassDebugger::OnEntityManagerInitialized(const FMassEntityManager& EntityManager)
 {
-	if (NewWorld &&  UE::Mass::Debugger::Private::IsSupportedWorldType(NewWorld->WorldType))
+	const UWorld* World = EntityManager.GetWorld();
+	if (World && UE::Mass::Debugger::Private::IsSupportedWorldType(World->WorldType))
 	{
-		EnvironmentsList.Add(MakeShareable(new FMassDebuggerEnvironment(NewWorld)));
+		EnvironmentsList.Add(MakeShareable(new FMassDebuggerEnvironment(EntityManager.AsShared())));
 	}
 }
 
-void SMassDebugger::OnWorldDestroyed(UWorld* InWorld)
+void SMassDebugger::OnEntityManagerDeinitialized(const FMassEntityManager& EntityManager)
 {
-	if (InWorld == nullptr || UE::Mass::Debugger::Private::IsSupportedWorldType(InWorld->WorldType) == false)
+	const UWorld* World = EntityManager.GetWorld();
+	if (World != nullptr && UE::Mass::Debugger::Private::IsSupportedWorldType(World->WorldType) == false)
 	{
 		return;
 	}
 
-	FMassDebuggerEnvironment InEnvironment(InWorld);
+	if (EntityManager.DoesSharedInstanceExist())
+	{
+		FMassDebuggerEnvironment InEnvironment(EntityManager.AsShared());
 
-	if (EnvironmentsList.RemoveAll([&InEnvironment](const TSharedPtr<FMassDebuggerEnvironment>& Element)
+		if (EnvironmentsList.RemoveAll([&InEnvironment](const TSharedPtr<FMassDebuggerEnvironment>& Element)
+			{
+				return *Element.Get() == InEnvironment;
+			}) > 0)
 		{
-			return *Element.Get() == InEnvironment;
-		}) > 0)
+			if (DebuggerModel->IsCurrentEnvironment(InEnvironment))
+			{
+				DebuggerModel->MarkAsStale();
+				EnvironmentComboLabel->SetText(DebuggerModel->GetDisplayName());
+			}
+		}
+	}
+	else
 	{
-		if (DebuggerModel->IsCurrentEnvironment(InEnvironment))
+		// EntityManager is either undergoing destruction or it has never been made sharable
+		// all we can do right now is remove all no longer valid environments
+		if (EnvironmentsList.RemoveAll([](const TSharedPtr<FMassDebuggerEnvironment>& Element)
+			{
+				check(Element);
+				return Element.Get()->EntityManager.IsValid() == false;
+			}) > 0)
 		{
-			EnvironmentComboLabel->SetText(FText::FromString(TEXT("Stale")));
+			if (DebuggerModel->IsCurrentEnvironmentValid() == false)
+			{
+				DebuggerModel->MarkAsStale();
+				EnvironmentComboLabel->SetText(DebuggerModel->GetDisplayName());
+			}
 		}
 	}
 }
@@ -236,15 +257,19 @@ void SMassDebugger::HandleEnvironmentChanged(TSharedPtr<FMassDebuggerEnvironment
 
 void SMassDebugger::RebuildEnvironmentsList()
 {
-	const TIndirectArray<FWorldContext>& WorldContexts = GEngine->GetWorldContexts();
-	EnvironmentsList.Reset();
-	for (const FWorldContext& Context : WorldContexts)
+#if WITH_MASSENTITY_DEBUG
+	for (const TWeakPtr<const FMassEntityManager>& WeakEntityManager : FMassDebugger::GetEntityManagers())
 	{
-		if (UE::Mass::Debugger::Private::IsSupportedWorldType(Context.WorldType))
+		if (const FMassEntityManager* EntityManagerPtr = WeakEntityManager.Pin().Get())
 		{
-			EnvironmentsList.Add(MakeShareable(new FMassDebuggerEnvironment(Context.World())));
+			const UWorld* World = EntityManagerPtr->GetWorld();
+			if (World != nullptr && UE::Mass::Debugger::Private::IsSupportedWorldType(World->WorldType))
+			{
+				EnvironmentsList.Add(MakeShareable(new FMassDebuggerEnvironment(EntityManagerPtr->AsShared())));
+			}
 		}
 	}
+#endif // WITH_MASSENTITY_DEBUG
 }
 
 TSharedRef<SDockTab> SMassDebugger::SpawnProcessorsTab(const FSpawnTabArgs& Args)
@@ -282,39 +307,10 @@ TSharedRef<SDockTab> SMassDebugger::SpawnArchetypesTab(const FSpawnTabArgs& Args
 
 void SMassDebugger::BindDelegates()
 {
-#if WITH_EDITOR
-	PIEWorldInitialize = FWorldDelegates::OnPostDuplicate.AddLambda([this](UWorld* World, bool bDuplicateForPIE, TMap<UObject*, UObject*>& ReplacementMap, TArray<UObject*>& ObjectsToFixReferences)
-		{
-			if (World && bDuplicateForPIE)
-			{
-				EnvironmentsList.Add(MakeShareable(new FMassDebuggerEnvironment(World)));
-			}
-		});
-
-	PIEEndHandle = FEditorDelegates::EndPIE.AddLambda([this](const bool bIsSimulating)
-		{
-			const TIndirectArray<FWorldContext>& WorldContexts = GEngine->GetWorldContexts();
-			EnvironmentsList.Reset();
-			for (const FWorldContext& Context : WorldContexts)
-			{
-				if (Context.WorldType == EWorldType::Game || Context.WorldType == EWorldType::Editor)
-				{
-					EnvironmentsList.Add(MakeShareable(new FMassDebuggerEnvironment(Context.World())));
-				}
-				else if (Context.WorldType == EWorldType::PIE)
-				{
-					FMassDebuggerEnvironment TempEnvironment(Context.World());
-					if (DebuggerModel->IsCurrentEnvironment(TempEnvironment))
-					{
-						DebuggerModel->MarkAsStale();
-					}
-				}
-			}
-			EnvironmentComboLabel->SetText(DebuggerModel->GetDisplayName());
-		});
-#endif
-	OnWorldAddedHandle = GEngine->OnWorldAdded().AddRaw(this, &SMassDebugger::OnWorldAdded);
-	OnWorldDestroyedHandle = GEngine->OnWorldDestroyed().AddRaw(this, &SMassDebugger::OnWorldDestroyed);
+#if WITH_MASSENTITY_DEBUG
+	OnEntityManagerInitializedHandle = FMassDebugger::OnEntityManagerInitialized.AddRaw(this, &SMassDebugger::OnEntityManagerInitialized);
+	OnEntityManagerDeinitializedHandle = FMassDebugger::OnEntityManagerDeinitialized.AddRaw(this, &SMassDebugger::OnEntityManagerDeinitialized);
+#endif // WITH_MASSENTITY_DEBUG
 }
 
 #undef LOCTEXT_NAMESPACE
diff --git a/Engine/Plugins/Runtime/MassEntity/Source/MassEntityDebugger/Private/SMassDebugger.h b/Engine/Plugins/Runtime/MassEntity/Source/MassEntityDebugger/Private/SMassDebugger.h
--- a/Engine/Plugins/Runtime/MassEntity/Source/MassEntityDebugger/Private/SMassDebugger.h
+++ b/Engine/Plugins/Runtime/MassEntity/Source/MassEntityDebugger/Private/SMassDebugger.h
@@ -6,6 +6,7 @@
 
 struct FMassDebuggerEnvironment;
 struct FMassDebuggerModel;
+struct FMassEntityManager;
 template<typename T> class SComboBox;
 
 
@@ -60,8 +61,8 @@ protected:
 	TSharedRef<SDockTab> SpawnArchetypesTab(const FSpawnTabArgs& Args);
 
 	void BindDelegates();
-	void OnWorldAdded(UWorld* NewWorld);
-	void OnWorldDestroyed(UWorld* InWorld);
+	void OnEntityManagerInitialized(const FMassEntityManager& EntityManager);
+	void OnEntityManagerDeinitialized(const FMassEntityManager& EntityManager);
 	void HandleEnvironmentChanged(TSharedPtr<FMassDebuggerEnvironment> Item, ESelectInfo::Type SelectInfo);
 	void RebuildEnvironmentsList();
 
@@ -78,10 +79,6 @@ protected:
 
 	TSharedRef<FMassDebuggerModel> DebuggerModel;
 
-#if WITH_EDITOR
-	FDelegateHandle PIEEndHandle;
-	FDelegateHandle PIEWorldInitialize;
-#endif // WITH_EDITOR
-	FDelegateHandle OnWorldAddedHandle;
-	FDelegateHandle OnWorldDestroyedHandle;
+	FDelegateHandle OnEntityManagerInitializedHandle;
+	FDelegateHandle OnEntityManagerDeinitializedHandle;
 };
diff --git a/Engine/Plugins/Runtime/MassEntity/Source/MassEntityTestSuite/Private/MassTestTypes.cpp b/Engine/Plugins/Runtime/MassEntity/Source/MassEntityTestSuite/Private/MassTestTypes.cpp
--- a/Engine/Plugins/Runtime/MassEntity/Source/MassEntityTestSuite/Private/MassTestTypes.cpp
+++ b/Engine/Plugins/Runtime/MassEntity/Source/MassEntityTestSuite/Private/MassTestTypes.cpp
@@ -12,6 +12,7 @@ bool FExecutionTestBase::SetUp()
 {
 	World = FAITestHelpers::GetWorld();
 	EntityManager = MakeShareable(new FMassEntityManager);
+	EntityManager->SetDebugName(TEXT("MassEntityTestSuite"));
 	EntityManager->Initialize();
 
 	return true;
