Added summary information to MassDebugger that's being shown when multiple archetypes are selected

diff --git a/Engine/Plugins/Runtime/MassEntity/Source/MassEntityDebugger/Private/MassDebuggerModel.h b/Engine/Plugins/Runtime/MassEntity/Source/MassEntityDebugger/Private/MassDebuggerModel.h
--- a/Engine/Plugins/Runtime/MassEntity/Source/MassEntityDebugger/Private/MassDebuggerModel.h
+++ b/Engine/Plugins/Runtime/MassEntity/Source/MassEntityDebugger/Private/MassDebuggerModel.h
@@ -58,7 +58,6 @@ struct FMassDebuggerArchetypeData
 	FMassDebuggerArchetypeData(const FMassArchetypeHandle& ArchetypeHandle);
 
 	FMassArchetypeCompositionDescriptor Composition;
-	FMassArchetypeSharedFragmentValues SharedFragments;
 
 	/** Hash of the Compositions. */
 	uint32 CompositionHash = 0;
diff --git a/Engine/Plugins/Runtime/MassEntity/Source/MassEntityDebugger/Private/SMassArchetypesView.cpp b/Engine/Plugins/Runtime/MassEntity/Source/MassEntityDebugger/Private/SMassArchetypesView.cpp
--- a/Engine/Plugins/Runtime/MassEntity/Source/MassEntityDebugger/Private/SMassArchetypesView.cpp
+++ b/Engine/Plugins/Runtime/MassEntity/Source/MassEntityDebugger/Private/SMassArchetypesView.cpp
@@ -5,7 +5,8 @@
 #include "SMassArchetype.h"
 #include "MassDebuggerStyle.h"
 #include "Styling/StyleColors.h"
-#include "Widgets/Layout/SScrollBox.h"		
+#include "Widgets/Layout/SScrollBox.h"
+#include "Widgets/Text/SRichTextBlock.h"
 
 #define LOCTEXT_NAMESPACE "SMassDebugger"
 
@@ -236,6 +237,47 @@ void SMassArchetypesView::RebuildSelectedView()
 {
 	SelectedArchetypesView->ClearChildren();
 
+	if (DebuggerModel->SelectedArchetypes.Num() > 1)
+	{
+		// print out entity count and memory totals
+		int32 EntitiesCount = 0;
+		SIZE_T AllocatedSize = 0;
+		for (const TSharedPtr<FMassDebuggerArchetypeData>& ArchetypeData : DebuggerModel->SelectedArchetypes)
+		{
+			EntitiesCount += ArchetypeData->ArchetypeStats.EntitiesCount;
+			AllocatedSize += ArchetypeData->ArchetypeStats.AllocatedSize;
+		}
+
+		const FText TotalDescription = FText::Format(LOCTEXT("ArchetypesSelectionTotalsContent", "EntitiesCount: {0}\nAllocated memory: {1}")
+			, FText::AsNumber(EntitiesCount), FText::AsMemory(AllocatedSize));
+
+		TSharedRef<SVerticalBox> Box = SNew(SVerticalBox);
+		
+		Box->AddSlot()
+			.AutoHeight()
+			.Padding(0, 4)
+			[
+				SNew(SRichTextBlock)
+				.Text(LOCTEXT("ArchetypesSelectionTotals", "Total"))
+				.DecoratorStyleSet(&FAppStyle::Get())
+				.TextStyle(FAppStyle::Get(), "LargeText")
+			];
+
+		Box->AddSlot()
+			.AutoHeight()
+			.Padding(0, 4)
+			[
+				SNew(STextBlock)
+				.Text(TotalDescription)
+			];
+
+		SelectedArchetypesView->AddSlot()
+			.Padding(4)
+			[
+				Box
+			];
+	}
+
 	const EMassBitSetDiffPrune Prune = DebuggerModel->SelectedArchetypes.Num() > 1 ? EMassBitSetDiffPrune::Same : EMassBitSetDiffPrune::None;
 	
 	for (const TSharedPtr<FMassDebuggerArchetypeData>& ArchetypeData : DebuggerModel->SelectedArchetypes)
