Added a TSharedPtr<FMassEntityManager> member to MassExecutionContext.
This change has long been coming, we always use given execution context with only a single entity manager instance, and the lack of connection between context and the manager required some use cases to cache the manager.

diff --git a/Engine/Plugins/AI/MassAI/Source/MassAIDebug/Private/GameplayDebuggerCategory_Mass.cpp b/Engine/Plugins/AI/MassAI/Source/MassAIDebug/Private/GameplayDebuggerCategory_Mass.cpp
--- a/Engine/Plugins/AI/MassAI/Source/MassAIDebug/Private/GameplayDebuggerCategory_Mass.cpp
+++ b/Engine/Plugins/AI/MassAI/Source/MassAIDebug/Private/GameplayDebuggerCategory_Mass.cpp
@@ -383,10 +383,12 @@ void FGameplayDebuggerCategory_Mass::CollectData(APlayerController* OwnerPC, AAc
 		
 		if (MassStateTreeSubsystem && SignalSubsystem && SmartObjectSubsystem)
 		{
-			FMassExecutionContext Context(0.0f);
+			FMassExecutionContext Context(EntityManager.AsShared(), 0.0f);
 		
-			EntityQuery.ForEachEntityChunk(EntityManager, Context, [this, MassStateTreeSubsystem, SignalSubsystem, &EntityManager, SmartObjectSubsystem, OwnerPC, ViewLocation, ViewDirection, CurrentTime](FMassExecutionContext& Context)
+			EntityQuery.ForEachEntityChunk(EntityManager, Context, [this, MassStateTreeSubsystem, SignalSubsystem, SmartObjectSubsystem, OwnerPC, ViewLocation, ViewDirection, CurrentTime](FMassExecutionContext& Context)
 			{
+				FMassEntityManager& EntityManager = Context.GetEntityManagerChecked();
+
 				const int32 NumEntities = Context.GetNumEntities();
 				const TConstArrayView<FMassStateTreeInstanceFragment> StateTreeInstanceList = Context.GetFragmentView<FMassStateTreeInstanceFragment>();
 				const TConstArrayView<FTransformFragment> TransformList = Context.GetFragmentView<FTransformFragment>();
@@ -628,7 +630,7 @@ void FGameplayDebuggerCategory_Mass::CollectData(APlayerController* OwnerPC, AAc
 			FMassEntityQuery EntityColliderQuery;
 			EntityColliderQuery.AddRequirement<FMassAvoidanceColliderFragment>(EMassFragmentAccess::ReadOnly);
 			EntityColliderQuery.AddRequirement<FTransformFragment>(EMassFragmentAccess::ReadOnly);
-			FMassExecutionContext Context(0.f);
+			FMassExecutionContext Context(EntityManager.AsShared(), 0.f);
 			EntityColliderQuery.ForEachEntityChunk(EntityManager, Context, [this, ViewLocation, ViewDirection](const FMassExecutionContext& Context)
 			{
 				const int32 NumEntities = Context.GetNumEntities();
diff --git a/Engine/Plugins/Runtime/MassEntity/Source/MassEntity/Private/MassEntityManager.cpp b/Engine/Plugins/Runtime/MassEntity/Source/MassEntity/Private/MassEntityManager.cpp
--- a/Engine/Plugins/Runtime/MassEntity/Source/MassEntity/Private/MassEntityManager.cpp
+++ b/Engine/Plugins/Runtime/MassEntity/Source/MassEntity/Private/MassEntityManager.cpp
@@ -1314,9 +1314,9 @@ void FMassEntityManager::GetValidArchetypes(const FMassEntityQuery& Query, TArra
 	}
 }
 
-FMassExecutionContext FMassEntityManager::CreateExecutionContext(const float DeltaSeconds) const
+FMassExecutionContext FMassEntityManager::CreateExecutionContext(const float DeltaSeconds)
 {
-	FMassExecutionContext ExecutionContext(DeltaSeconds);
+	FMassExecutionContext ExecutionContext(AsShared(), DeltaSeconds);
 	ExecutionContext.SetDeferredCommandBuffer(DeferredCommandBuffer);
 	return MoveTemp(ExecutionContext);
 }
diff --git a/Engine/Plugins/Runtime/MassEntity/Source/MassEntity/Private/MassEntityQuery.cpp b/Engine/Plugins/Runtime/MassEntity/Source/MassEntity/Private/MassEntityQuery.cpp
--- a/Engine/Plugins/Runtime/MassEntity/Source/MassEntity/Private/MassEntityQuery.cpp
+++ b/Engine/Plugins/Runtime/MassEntity/Source/MassEntity/Private/MassEntityQuery.cpp
@@ -228,7 +228,7 @@ void FMassEntityQuery::ForEachEntityChunk(FMassEntityManager& EntityManager, FMa
 #endif
 
 	ExecutionContext.ClearExecutionData();
-	ExecutionContext.FlushDeferred(EntityManager);
+	ExecutionContext.FlushDeferred();
 }
 
 int32 FMassEntityQuery::GetNumMatchingEntities(FMassEntityManager& InEntityManager)
diff --git a/Engine/Plugins/Runtime/MassEntity/Source/MassEntity/Private/MassExecutionContext.cpp b/Engine/Plugins/Runtime/MassEntity/Source/MassEntity/Private/MassExecutionContext.cpp
--- a/Engine/Plugins/Runtime/MassEntity/Source/MassEntity/Private/MassExecutionContext.cpp
+++ b/Engine/Plugins/Runtime/MassEntity/Source/MassEntity/Private/MassExecutionContext.cpp
@@ -4,11 +4,12 @@
 #include "MassArchetypeData.h"
 
 
-void FMassExecutionContext::FlushDeferred(FMassEntityManager& EntityManager) const
+void FMassExecutionContext::FlushDeferred()
 {
 	if (bFlushDeferredCommands && DeferredCommandBuffer)
 	{
-		EntityManager.FlushCommands(DeferredCommandBuffer);
+		check(EntityManager);
+		EntityManager->FlushCommands(DeferredCommandBuffer);
 	}
 }
 
diff --git a/Engine/Plugins/Runtime/MassEntity/Source/MassEntity/Private/MassExecutor.cpp b/Engine/Plugins/Runtime/MassEntity/Source/MassEntity/Private/MassExecutor.cpp
--- a/Engine/Plugins/Runtime/MassEntity/Source/MassEntity/Private/MassExecutor.cpp
+++ b/Engine/Plugins/Runtime/MassEntity/Source/MassEntity/Private/MassExecutor.cpp
@@ -84,7 +84,7 @@ void RunProcessorsView(TArrayView<UMassProcessor*> Processors, FMassProcessingCo
 
 	TRACE_CPUPROFILER_EVENT_SCOPE_STR("MassExecutor RunProcessorsView")
 
-	FMassExecutionContext ExecutionContext(ProcessingContext.DeltaSeconds);
+	FMassExecutionContext ExecutionContext(ProcessingContext.EntityManager, ProcessingContext.DeltaSeconds);
 	if (EntityCollection)
 	{
 		ExecutionContext.SetEntityCollection(*EntityCollection);
@@ -120,7 +120,7 @@ void RunProcessorsView(TArrayView<UMassProcessor*> Processors, FMassProcessingCo
 		ensure(!EntityManager.Defer().IsFlushing());
 		ensure(!ExecutionContext.Defer().IsFlushing());
 		ExecutionContext.Defer().MoveAppend(EntityManager.Defer());
-		ExecutionContext.FlushDeferred(*ProcessingContext.EntityManager);
+		ExecutionContext.FlushDeferred();
 	}
 	// else make sure we don't just lose the commands. Append to the command buffer requested via
 	// ProcessingContext.CommandBuffer or to the default EntityManager's command buffer.
@@ -139,9 +139,8 @@ void RunProcessorsView(TArrayView<UMassProcessor*> Processors, FMassProcessingCo
 
 struct FMassExecutorDoneTask
 {
-	FMassExecutorDoneTask(const FMassExecutionContext& InExecutionContext, const TSharedPtr<FMassEntityManager>& InEntityManager, TFunction<void()> InOnDoneNotification, const FString& InDebugName)
+	FMassExecutorDoneTask(const FMassExecutionContext& InExecutionContext, TFunction<void()> InOnDoneNotification, const FString& InDebugName)
 		: ExecutionContext(InExecutionContext)
-		, EntityManager(InEntityManager)
 		, OnDoneNotification(InOnDoneNotification)
 		, DebugName(InDebugName)
 	{
@@ -157,8 +156,7 @@ struct FMassExecutorDoneTask
 	{
 		TRACE_CPUPROFILER_EVENT_SCOPE_STR("Flush Deferred Commands Parallel");
 
-		check(EntityManager);
-		FMassEntityManager& EntityManagerRef = *EntityManager.Get();
+		FMassEntityManager& EntityManagerRef = ExecutionContext.GetEntityManagerChecked();
 
 		if (&ExecutionContext.Defer() != &EntityManagerRef.Defer())
 		{
@@ -167,13 +165,12 @@ struct FMassExecutorDoneTask
 
 		UE_LOG(LogMass, Verbose, TEXT("MassExecutor %s tasks DONE"), *DebugName);
 		ExecutionContext.SetFlushDeferredCommands(true);
-		ExecutionContext.FlushDeferred(EntityManagerRef);
+		ExecutionContext.FlushDeferred();
 
 		OnDoneNotification();
 	}
 private:
 	FMassExecutionContext ExecutionContext;
-	TSharedPtr<FMassEntityManager> EntityManager;
 	TFunction<void()> OnDoneNotification;
 	FString DebugName;
 };
@@ -191,7 +188,7 @@ FGraphEventRef TriggerParallelTasks(UMassProcessor& Processor, FMassProcessingCo
 	TRACE_CPUPROFILER_EVENT_SCOPE_STR("MassExecutor RunParallel")
 
 	// not going through FMassEntityManager::CreateExecutionContext on purpose - we do need a separate command buffer
-	FMassExecutionContext ExecutionContext(ProcessingContext.DeltaSeconds);
+	FMassExecutionContext ExecutionContext(ProcessingContext.EntityManager, ProcessingContext.DeltaSeconds);
 	TSharedPtr<FMassCommandBuffer> CommandBuffer = ProcessingContext.CommandBuffer
 		? ProcessingContext.CommandBuffer : MakeShareable(new FMassCommandBuffer());
 	ExecutionContext.SetDeferredCommandBuffer(CommandBuffer);
@@ -209,7 +206,7 @@ FGraphEventRef TriggerParallelTasks(UMassProcessor& Processor, FMassProcessingCo
 	{
 		const FGraphEventArray Prerequisites = { CompletionEvent };
 		CompletionEvent = TGraphTask<FMassExecutorDoneTask>::CreateTask(&Prerequisites)
-			.ConstructAndDispatchWhenReady(ExecutionContext, ProcessingContext.EntityManager, OnDoneNotification, Processor.GetName());
+			.ConstructAndDispatchWhenReady(ExecutionContext, OnDoneNotification, Processor.GetName());
 	}
 
 	return CompletionEvent;
diff --git a/Engine/Plugins/Runtime/MassEntity/Source/MassEntity/Public/MassEntityManager.h b/Engine/Plugins/Runtime/MassEntity/Source/MassEntity/Public/MassEntityManager.h
--- a/Engine/Plugins/Runtime/MassEntity/Source/MassEntity/Public/MassEntityManager.h
+++ b/Engine/Plugins/Runtime/MassEntity/Source/MassEntity/Public/MassEntityManager.h
@@ -309,7 +309,7 @@ public:
 	/**
 	 * Creates and initializes a FMassExecutionContext instance.
 	 */
-	FMassExecutionContext CreateExecutionContext(const float DeltaSeconds) const;
+	FMassExecutionContext CreateExecutionContext(const float DeltaSeconds);
 
 	FScopedProcessing NewProcessingScope() { return FScopedProcessing(ProcessingScopeCount); }
 	bool IsProcessing() const { return ProcessingScopeCount > 0; }
diff --git a/Engine/Plugins/Runtime/MassEntity/Source/MassEntity/Public/MassExecutionContext.h b/Engine/Plugins/Runtime/MassEntity/Source/MassEntity/Public/MassExecutionContext.h
--- a/Engine/Plugins/Runtime/MassEntity/Source/MassEntity/Public/MassExecutionContext.h
+++ b/Engine/Plugins/Runtime/MassEntity/Source/MassEntity/Public/MassExecutionContext.h
@@ -52,6 +52,8 @@ private:
 	int32 ChunkSerialModificationNumber = -1;
 	FMassTagBitSet CurrentArchetypesTagBitSet;
 
+	TSharedPtr<FMassEntityManager> EntityManager;
+
 #if WITH_MASSENTITY_DEBUG
 	FString DebugExecutionDescription;
 #endif
@@ -71,17 +73,20 @@ private:
 	friend FMassEntityQuery;
 
 public:
-	explicit FMassExecutionContext(const float InDeltaTimeSeconds, const bool bInFlushDeferredCommands = true)
+	explicit FMassExecutionContext(const TSharedPtr<FMassEntityManager>& InEntityManager, const float InDeltaTimeSeconds, const bool bInFlushDeferredCommands = true)
 		: DeltaTimeSeconds(InDeltaTimeSeconds)
+		, EntityManager(InEntityManager)
 		, bFlushDeferredCommands(bInFlushDeferredCommands)
 	{
 		Subsystems.AddZeroed(FMassExternalSubsystemBitSet::GetMaxNum());
 	}
 
 	FMassExecutionContext()
-		: FMassExecutionContext(/*InDeltaTimeSeconds=*/0.f)
+		: FMassExecutionContext(nullptr, /*InDeltaTimeSeconds=*/0.f)
 	{}
 
+	FMassEntityManager& GetEntityManagerChecked() { check(EntityManager); return *EntityManager.Get(); }
+
 #if WITH_MASSENTITY_DEBUG
 	const FString& DebugGetExecutionDesc() const { return DebugExecutionDescription; }
 	void DebugSetExecutionDesc(const FString& Description) { DebugExecutionDescription = Description; }
@@ -333,7 +338,7 @@ public:
 		return FragmentType != nullptr && FragmentType == TFragment::StaticStruct();
 	}
 
-	void FlushDeferred(FMassEntityManager& EntityManager) const;
+	void FlushDeferred();
 
 	void ClearExecutionData();
 	void SetCurrentArchetypesTagBitSet(const FMassTagBitSet& BitSet)
diff --git a/Engine/Plugins/Runtime/MassGameplay/Source/MassRepresentation/Private/MassRepresentationProcessor.cpp b/Engine/Plugins/Runtime/MassGameplay/Source/MassRepresentation/Private/MassRepresentationProcessor.cpp
--- a/Engine/Plugins/Runtime/MassGameplay/Source/MassRepresentation/Private/MassRepresentationProcessor.cpp
+++ b/Engine/Plugins/Runtime/MassGameplay/Source/MassRepresentation/Private/MassRepresentationProcessor.cpp
@@ -11,6 +11,7 @@
 #include "Engine/World.h"
 #include "MassRepresentationActorManagement.h"
 #include "MassEntityUtils.h"
+#include "MassExecutionContext.h"
 
 
 namespace UE::MassRepresentation
@@ -41,15 +42,6 @@ void UMassRepresentationProcessor::ConfigureQueries()
 	EntityQuery.AddSubsystemRequirement<UMassActorSubsystem>(EMassFragmentAccess::ReadWrite);
 }
 
-void UMassRepresentationProcessor::Initialize(UObject& Owner)
-{
-	Super::Initialize(Owner);
-
-	World = Owner.GetWorld();
-	check(World);
-	EntityManager = UE::Mass::Utils::GetEntityManagerChecked(*World).AsShared();
-}
-
 void UMassRepresentationProcessor::UpdateRepresentation(FMassExecutionContext& Context)
 {
 	UMassRepresentationSubsystem* RepresentationSubsystem = Context.GetMutableSharedFragment<FMassRepresentationSubsystemSharedFragment>().RepresentationSubsystem;
@@ -60,8 +52,7 @@ void UMassRepresentationProcessor::UpdateRepresentation(FMassExecutionContext& C
 
 	UMassActorSubsystem* MassActorSubsystem = Context.GetMutableSubsystem<UMassActorSubsystem>(RepresentationSubsystem->GetWorld());
 
-	check(EntityManager);
-	FMassEntityManager& CachedEntityManager = *EntityManager.Get();
+	FMassEntityManager& CachedEntityManager = Context.GetEntityManagerChecked();
 
 	const TConstArrayView<FTransformFragment> TransformList = Context.GetFragmentView<FTransformFragment>();
 	const TArrayView<FMassRepresentationFragment> RepresentationList = Context.GetMutableFragmentView<FMassRepresentationFragment>();
@@ -228,10 +219,10 @@ void UMassRepresentationProcessor::UpdateRepresentation(FMassExecutionContext& C
 	}
 }
 
-void UMassRepresentationProcessor::Execute(FMassEntityManager& InEntitySubsystem, FMassExecutionContext& Context)
+void UMassRepresentationProcessor::Execute(FMassEntityManager& InEntityManager, FMassExecutionContext& Context)
 {
 	// Update entities representation
-	EntityQuery.ForEachEntityChunk(InEntitySubsystem, Context, [this](FMassExecutionContext& Context)
+	EntityQuery.ForEachEntityChunk(InEntityManager, Context, [this](FMassExecutionContext& Context)
 	{
 		UpdateRepresentation(Context);
 	});
diff --git a/Engine/Plugins/Runtime/MassGameplay/Source/MassRepresentation/Public/MassRepresentationProcessor.h b/Engine/Plugins/Runtime/MassGameplay/Source/MassRepresentation/Public/MassRepresentationProcessor.h
--- a/Engine/Plugins/Runtime/MassGameplay/Source/MassRepresentation/Public/MassRepresentationProcessor.h
+++ b/Engine/Plugins/Runtime/MassGameplay/Source/MassRepresentation/Public/MassRepresentationProcessor.h
@@ -26,12 +26,6 @@ protected:
 	/** Configure the owned FMassEntityQuery instances to express processor's requirements */
 	virtual void ConfigureQueries() override;
 
-	/**
-	 * Initialize the processor 
-	 * @param Owner of the Processor
-	 */
-	virtual void Initialize(UObject& Owner) override;
-
 	/** 
 	 * Execution method for this processor 
 	 * @param EntitySubsystem is the system to execute the lambdas on each entity chunk
@@ -58,12 +52,6 @@ protected:
 	 */
 	void UpdateRepresentation(FMassExecutionContext& Context);
 
-	/** Caching ptr to our associated world */
-	UPROPERTY(Transient)
-	TObjectPtr<UWorld> World;
-
-	TSharedPtr<FMassEntityManager> EntityManager;
-
 	FMassEntityQuery EntityQuery;
 };
 
