Renamed MassProcessingPhase.h/cpp files to MassProcessingPhaseManager.h/cpp

diff --git a/Engine/Plugins/Runtime/MassEntity/Source/MassEntity/Private/MassEntitySettings.cpp b/Engine/Plugins/Runtime/MassEntity/Source/MassEntity/Private/MassEntitySettings.cpp
--- a/Engine/Plugins/Runtime/MassEntity/Source/MassEntity/Private/MassEntitySettings.cpp
+++ b/Engine/Plugins/Runtime/MassEntity/Source/MassEntity/Private/MassEntitySettings.cpp
@@ -1,7 +1,7 @@
 // Copyright Epic Games, Inc. All Rights Reserved.
 
 #include "MassEntitySettings.h"
-#include "MassProcessingPhase.h"
+#include "MassProcessingPhaseManager.h"
 #include "VisualLogger/VisualLogger.h"
 #include "UObject/UObjectHash.h"
 #include "Misc/CoreDelegates.h"
diff --git a/Engine/Plugins/Runtime/MassEntity/Source/MassEntity/Private/MassProcessingPhase.cpp b/Engine/Plugins/Runtime/MassEntity/Source/MassEntity/Private/MassProcessingPhase.cpp
deleted file mode 100644
diff --git a/Engine/Plugins/Runtime/MassEntity/Source/MassEntity/Private/MassProcessingPhaseManager.cpp b/Engine/Plugins/Runtime/MassEntity/Source/MassEntity/Private/MassProcessingPhaseManager.cpp
new file mode 100644
--- /dev/null
+++ b/Engine/Plugins/Runtime/MassEntity/Source/MassEntity/Private/MassProcessingPhaseManager.cpp
@@ -0,0 +1,355 @@
+// Copyright Epic Games, Inc. All Rights Reserved.
+
+#include "MassProcessingPhaseManager.h"
+#include "MassProcessingTypes.h"
+#include "MassDebugger.h"
+#include "MassEntitySettings.h"
+#include "MassProcessor.h"
+#include "MassExecutor.h"
+#include "MassEntityManager.h"
+#include "MassEntitySubsystem.h"
+#include "VisualLogger/VisualLogger.h"
+#include "Engine/World.h"
+#include "MassCommandBuffer.h"
+
+#define LOCTEXT_NAMESPACE "Mass"
+
+DECLARE_CYCLE_STAT(TEXT("Mass Phase Done"), STAT_MassPhaseDone, STATGROUP_TaskGraphTasks);
+
+namespace UE::Mass::Tweakables
+{
+	bool bFullyParallel = MASS_DO_PARALLEL;
+
+	FAutoConsoleVariableRef CVars[] = {
+		{TEXT("mass.FullyParallel"), bFullyParallel, TEXT("Enables mass processing distribution to all available thread (via the task graph)")},
+	};
+}
+
+namespace UE::Mass::Private
+{
+	ETickingGroup PhaseToTickingGroup[int(EMassProcessingPhase::MAX)]
+	{
+		ETickingGroup::TG_PrePhysics, // EMassProcessingPhase::PrePhysics
+		ETickingGroup::TG_StartPhysics, // EMassProcessingPhase::StartPhysics
+		ETickingGroup::TG_DuringPhysics, // EMassProcessingPhase::DuringPhysics
+		ETickingGroup::TG_EndPhysics,	// EMassProcessingPhase::EndPhysics
+		ETickingGroup::TG_PostPhysics,	// EMassProcessingPhase::PostPhysics
+		ETickingGroup::TG_LastDemotable, // EMassProcessingPhase::FrameEnd
+	};
+} // UE::Mass::Private
+
+//----------------------------------------------------------------------//
+//  FMassProcessingPhase
+//----------------------------------------------------------------------//
+FMassProcessingPhase::FMassProcessingPhase()
+{
+	bCanEverTick = true;
+	bStartWithTickEnabled = false;
+}
+
+void FMassProcessingPhase::ExecuteTick(float DeltaTime, ELevelTick TickType, ENamedThreads::Type CurrentThread, const FGraphEventRef& MyCompletionGraphEvent)
+{
+	if (TickType == LEVELTICK_ViewportsOnly || TickType == LEVELTICK_PauseTick)
+	{
+		return;
+	}
+
+	checkf(PhaseManager, TEXT("Manager is null which is not a supported case. Either this FMassProcessingPhase has not been initialized properly or it's been left dangling after the FMassProcessingPhase owner got destroyed."));
+
+	TRACE_CPUPROFILER_EVENT_SCOPE_STR(*FString::Printf(TEXT("FMassProcessingPhase::ExecuteTick %s"), *UEnum::GetValueAsString(Phase)));
+
+	PhaseManager->OnPhaseStart(*this);
+	{
+		LLM_SCOPE_BYNAME(TEXT("Mass/PhaseStartDelegate"));
+		OnPhaseStart.Broadcast(DeltaTime);
+	}
+
+	check(PhaseProcessor);
+	
+	FMassEntityManager& EntityManager = PhaseManager->GetEntityManagerRef();
+	FMassProcessingContext Context(EntityManager, DeltaTime);
+
+	bIsDuringMassProcessing = true;
+
+	if (bRunInParallelMode)
+	{
+		bool bWorkRequested = false;
+		if (PhaseProcessor->IsEmpty() == false)
+		{
+			const FGraphEventRef PipelineCompletionEvent = UE::Mass::Executor::TriggerParallelTasks(*PhaseProcessor, Context, [this, DeltaTime]()
+				{
+					OnParallelExecutionDone(DeltaTime);
+				});
+
+			if (PipelineCompletionEvent.IsValid())
+			{
+				MyCompletionGraphEvent->DontCompleteUntil(PipelineCompletionEvent);
+				bWorkRequested = true;
+			}
+		}
+		if (bWorkRequested == false)
+		{
+			OnParallelExecutionDone(DeltaTime);
+		}
+	}
+	else
+	{
+		UE::Mass::Executor::Run(*PhaseProcessor, Context);
+
+		{
+			LLM_SCOPE_BYNAME(TEXT("Mass/PhaseEndDelegate"));
+			OnPhaseEnd.Broadcast(DeltaTime);
+		}
+		PhaseManager->OnPhaseEnd(*this);
+		bIsDuringMassProcessing = false;
+	}
+}
+
+void FMassProcessingPhase::OnParallelExecutionDone(const float DeltaTime)
+{
+	bIsDuringMassProcessing = false;
+	{
+		LLM_SCOPE_BYNAME(TEXT("Mass/PhaseEndDelegate"));
+		OnPhaseEnd.Broadcast(DeltaTime);
+	}
+	check(PhaseManager);
+	PhaseManager->OnPhaseEnd(*this);
+}
+
+FString FMassProcessingPhase::DiagnosticMessage()
+{
+	return (PhaseManager ? PhaseManager->GetFullName() : TEXT("NULL-MassProcessingPhaseManager")) + TEXT("[ProcessorTick]");
+}
+
+FName FMassProcessingPhase::DiagnosticContext(bool bDetailed)
+{
+	return PhaseManager ? PhaseManager->GetClass()->GetFName() : TEXT("NULL-MassProcessingPhaseManager");
+}
+
+//----------------------------------------------------------------------//
+// UMassProcessingPhaseManager  
+//----------------------------------------------------------------------//
+void UMassProcessingPhaseManager::PostInitProperties()
+{
+	Super::PostInitProperties();
+	if (HasAnyFlags(RF_ClassDefaultObject) == false)
+	{
+#if WITH_EDITOR
+		UWorld* World = GetWorld();
+		if (World && World->IsGameWorld() == false)
+		{
+			UMassEntitySettings* Settings = GetMutableDefault<UMassEntitySettings>();
+			check(Settings);
+			MassEntitySettingsChangeHandle = Settings->GetOnSettingsChange().AddUObject(this, &UMassProcessingPhaseManager::OnMassEntitySettingsChange);
+		}
+#endif // WITH_EDITOR
+	}
+	CreatePhases();
+}
+
+void UMassProcessingPhaseManager::BeginDestroy()
+{
+	Stop();
+	Super::BeginDestroy();
+
+#if WITH_EDITOR
+	if (UMassEntitySettings* Settings = GetMutableDefault<UMassEntitySettings>())
+	{
+		Settings->GetOnSettingsChange().Remove(MassEntitySettingsChangeHandle);
+	}
+#endif // WITH_EDITOR
+}
+
+void UMassProcessingPhaseManager::InitializePhases(UObject& InProcessorOwner)
+{
+	const FMassProcessingPhaseConfig* ProcessingPhasesConfig = GET_MASS_CONFIG_VALUE(GetProcessingPhasesConfig());
+
+	FString DependencyGraphFileName;
+
+#if WITH_EDITOR
+	const UWorld* World = InProcessorOwner.GetWorld();
+	const UMassEntitySettings* Settings = GetMutableDefault<UMassEntitySettings>();
+	if (World != nullptr && Settings != nullptr && !Settings->DumpDependencyGraphFileName.IsEmpty())
+	{
+		DependencyGraphFileName = FString::Printf(TEXT("%s_%s"), *Settings->DumpDependencyGraphFileName,*ToString(World->GetNetMode()));
+	}
+#endif // WITH_EDITOR
+
+	for (int i = 0; i < int(EMassProcessingPhase::MAX); ++i)
+	{
+		const FMassProcessingPhaseConfig& PhaseConfig = ProcessingPhasesConfig[i];
+		UMassCompositeProcessor* PhaseProcessor = ProcessingPhases[i].PhaseProcessor;
+		FString FileName = !DependencyGraphFileName.IsEmpty() ? FString::Printf(TEXT("%s_%s"), *DependencyGraphFileName, *PhaseConfig.PhaseName.ToString()) : FString();
+		PhaseProcessor->CopyAndSort(PhaseConfig, FileName);
+		PhaseProcessor->Initialize(InProcessorOwner);
+	}
+
+#if WITH_MASSENTITY_DEBUG
+	// print it all out to vislog
+	UE_VLOG_UELOG(this, LogMass, Verbose, TEXT("Phases initialization done. Current composition:"));
+
+	FStringOutputDevice OutDescription;
+	for (int i = 0; i < int(EMassProcessingPhase::MAX); ++i)
+	{
+		if (ProcessingPhases[i].PhaseProcessor)
+		{
+			ProcessingPhases[i].PhaseProcessor->DebugOutputDescription(OutDescription);
+			UE_VLOG_UELOG(this, LogMass, Verbose, TEXT("--- %s"), *OutDescription);
+			OutDescription.Reset();
+		}
+	}	
+#endif // WITH_MASSENTITY_DEBUG
+}
+
+void UMassProcessingPhaseManager::Start(UWorld& World)
+{
+	UMassEntitySubsystem* EntitySubsystem = UWorld::GetSubsystem<UMassEntitySubsystem>(&World);
+
+	if (ensure(EntitySubsystem))
+	{
+		EntityManager = EntitySubsystem->GetMutableEntityManager().AsShared();
+		EnableTickFunctions(World);
+	}
+	else
+	{
+		UE_VLOG_UELOG(this, LogMass, Error, TEXT("Called %s while missing the EntitySubsystem"), ANSI_TO_TCHAR(__FUNCTION__));
+	}
+}
+
+void UMassProcessingPhaseManager::Start(const TSharedPtr<FMassEntityManager>& InEntityManager)
+{
+	UWorld* World = InEntityManager->GetWorld();
+	check(World);
+	EntityManager = InEntityManager;
+	EnableTickFunctions(*World);
+}
+
+void UMassProcessingPhaseManager::EnableTickFunctions(const UWorld& World)
+{
+	check(EntityManager);
+
+	const bool bIsGameWorld = World.IsGameWorld();
+
+	for (FMassProcessingPhase& Phase : ProcessingPhases)
+	{
+		Phase.PhaseManager = this;
+		Phase.RegisterTickFunction(World.PersistentLevel);
+		Phase.SetTickFunctionEnable(true);
+#if WITH_MASSENTITY_DEBUG
+		if (Phase.PhaseProcessor && bIsGameWorld)
+		{
+			// not logging this in the editor mode since it messes up the game-recorded vislog display (with its progressively larger timestamp)
+			FStringOutputDevice Ar;
+			Phase.PhaseProcessor->DebugOutputDescription(Ar);
+			UE_VLOG_UELOG(this, LogMass, Log, TEXT("Enabling phase %s tick:\n%s")
+				, *UEnum::GetValueAsString(Phase.Phase), *Ar);
+		}
+#endif // WITH_MASSENTITY_DEBUG		
+	}
+
+	if (bIsGameWorld)
+	{
+		// not logging this in the editor mode since it messes up the game-recorded vislog display (with its progressively larger timestamp)
+		UE_VLOG_UELOG(this, LogMass, Log, TEXT("MassProcessingPhaseManager %s.%s has been started")
+			, *GetNameSafe(GetOuter()), *GetName());
+	}
+}
+
+void UMassProcessingPhaseManager::Stop()
+{
+	EntityManager.Reset();
+	
+	for (FMassProcessingPhase& Phase : ProcessingPhases)
+	{
+		Phase.SetTickFunctionEnable(false);
+	}
+
+	UWorld* World = GetWorld();
+	if (World && World->IsGameWorld())
+	{
+		// not logging this in editor mode since it messes up the game-recorded vislog display (with its progressively larger timestamp) 
+		UE_VLOG_UELOG(this, LogMass, Log, TEXT("MassProcessingPhaseManager %s.%s has been stopped")
+			, *GetNameSafe(GetOuter()), *GetName());
+	}
+}
+
+void UMassProcessingPhaseManager::SetPhaseProcessor(const EMassProcessingPhase Phase, UMassCompositeProcessor* PhaseProcessor)
+{
+	if (Phase == EMassProcessingPhase::MAX)
+	{
+		UE_VLOG_UELOG(this, LogMass, Error, TEXT("MassProcessingPhaseManager::OverridePhaseProcessor called with Phase == MAX"));
+		return;
+	}
+	
+	// note that it's ok to use PhaseProcessor == nullptr
+	ProcessingPhases[int32(Phase)].PhaseProcessor = PhaseProcessor;
+	
+	if (PhaseProcessor)
+	{
+		REDIRECT_OBJECT_TO_VLOG(PhaseProcessor, this);
+		PhaseProcessor->SetProcessingPhase(Phase);
+		PhaseProcessor->SetGroupName(FName(FString::Printf(TEXT("%s Group"), *UEnum::GetDisplayValueAsText(Phase).ToString())));
+		
+		check(GetOuter());
+		PhaseProcessor->Initialize(*GetOuter());
+
+#if WITH_MASSENTITY_DEBUG
+		FStringOutputDevice Ar;
+		PhaseProcessor->DebugOutputDescription(Ar);
+		UE_VLOG(this, LogMass, Log, TEXT("Setting new group processor for phase %s:\n%s")
+			, *UEnum::GetValueAsString(Phase), *Ar);
+#endif // WITH_MASSENTITY_DEBUG
+	}
+	else
+	{
+		UE_VLOG_UELOG(this, LogMass, Log, TEXT("Nulling-out phase processor for phase %s"), *UEnum::GetValueAsString(Phase));
+	}
+}
+
+void UMassProcessingPhaseManager::CreatePhases() 
+{
+	// @todo copy from settings instead of blindly creating from scratch
+	for (int i = 0; i < int(EMassProcessingPhase::MAX); ++i)
+	{
+		ProcessingPhases[i].Phase = EMassProcessingPhase(i);
+		ProcessingPhases[i].TickGroup = UE::Mass::Private::PhaseToTickingGroup[i];
+		UMassCompositeProcessor* PhaseProcessor = NewObject<UMassCompositeProcessor>(this, UMassCompositeProcessor::StaticClass()
+			, *FString::Printf(TEXT("ProcessingPhase_%s"), *UEnum::GetDisplayValueAsText(EMassProcessingPhase(i)).ToString()));
+		SetPhaseProcessor(EMassProcessingPhase(i), PhaseProcessor);
+	}
+}
+
+void UMassProcessingPhaseManager::OnPhaseStart(const FMassProcessingPhase& Phase)
+{
+	ensure(CurrentPhase == EMassProcessingPhase::MAX);
+	CurrentPhase = Phase.Phase;
+}
+
+void UMassProcessingPhaseManager::OnPhaseEnd(FMassProcessingPhase& Phase)
+{
+	ensure(CurrentPhase == Phase.Phase);
+	CurrentPhase = EMassProcessingPhase::MAX;
+
+	// switch between parallel and single-thread versions only after a given batch of processing has been wrapped up	
+	if (Phase.IsConfiguredForParallelMode() != UE::Mass::Tweakables::bFullyParallel)
+	{
+		if (UE::Mass::Tweakables::bFullyParallel)
+		{
+			Phase.ConfigureForParallelMode();
+		}
+		else
+		{
+			Phase.ConfigureForSingleThreadMode();
+		}
+	}
+}
+
+#if WITH_EDITOR
+void UMassProcessingPhaseManager::OnMassEntitySettingsChange(const FPropertyChangedEvent& PropertyChangedEvent)
+{
+	check(GetOuter());
+	InitializePhases(*GetOuter());
+}
+#endif // WITH_EDITOR
+
+#undef LOCTEXT_NAMESPACE
\ No newline at end of file
diff --git a/Engine/Plugins/Runtime/MassEntity/Source/MassEntity/Public/MassEntitySettings.h b/Engine/Plugins/Runtime/MassEntity/Source/MassEntity/Public/MassEntitySettings.h
--- a/Engine/Plugins/Runtime/MassEntity/Source/MassEntity/Public/MassEntitySettings.h
+++ b/Engine/Plugins/Runtime/MassEntity/Source/MassEntity/Public/MassEntitySettings.h
@@ -3,7 +3,7 @@
 #pragma once
 
 #include "MassSettings.h"
-#include "MassProcessingPhase.h"
+#include "MassProcessingPhaseManager.h"
 #include "MassProcessor.h"
 #include "InstancedStruct.h"
 #include "MassEntitySettings.generated.h"
diff --git a/Engine/Plugins/Runtime/MassEntity/Source/MassEntity/Public/MassProcessingPhase.h b/Engine/Plugins/Runtime/MassEntity/Source/MassEntity/Public/MassProcessingPhase.h
deleted file mode 100644
diff --git a/Engine/Plugins/Runtime/MassEntity/Source/MassEntity/Public/MassProcessingPhaseManager.h b/Engine/Plugins/Runtime/MassEntity/Source/MassEntity/Public/MassProcessingPhaseManager.h
new file mode 100644
--- /dev/null
+++ b/Engine/Plugins/Runtime/MassEntity/Source/MassEntity/Public/MassProcessingPhaseManager.h
@@ -0,0 +1,156 @@
+// Copyright Epic Games, Inc. All Rights Reserved.
+
+#pragma once
+
+#include "CoreMinimal.h"
+#include "UObject/UObjectGlobals.h"
+#include "UObject/Object.h"
+#include "Engine/EngineBaseTypes.h"
+#include "MassProcessingTypes.h"
+#include "MassProcessingPhaseManager.generated.h"
+
+
+class UMassProcessingPhaseManager;
+class UMassProcessor;
+class UMassCompositeProcessor;
+struct FMassEntityManager;
+struct FMassCommandBuffer;
+
+USTRUCT()
+struct FMassProcessingPhase : public FTickFunction
+{
+	GENERATED_BODY()
+
+	DECLARE_MULTICAST_DELEGATE_OneParam(FOnPhaseEvent, const float /*DeltaSeconds*/);
+
+	FMassProcessingPhase();
+
+protected:
+	// FTickFunction interface
+	virtual void ExecuteTick(float DeltaTime, ELevelTick TickType, ENamedThreads::Type CurrentThread, const FGraphEventRef& MyCompletionGraphEvent) override;
+	virtual FString DiagnosticMessage() override;
+	virtual FName DiagnosticContext(bool bDetailed) override;
+	// End of FTickFunction interface
+
+	void OnParallelExecutionDone(const float DeltaTime);
+
+	bool IsConfiguredForParallelMode() const { return bRunInParallelMode; }
+	void ConfigureForParallelMode() { bRunInParallelMode = true; }
+	void ConfigureForSingleThreadMode() { bRunInParallelMode = false; }
+
+public:
+	bool IsDuringMassProcessing() const { return bIsDuringMassProcessing; }
+
+protected:
+	friend UMassProcessingPhaseManager;
+
+	UPROPERTY(EditAnywhere, Category=Mass)
+	TObjectPtr<UMassCompositeProcessor> PhaseProcessor = nullptr;
+
+	UPROPERTY()
+	TObjectPtr<UMassProcessingPhaseManager> PhaseManager = nullptr;
+
+	EMassProcessingPhase Phase = EMassProcessingPhase::MAX;
+	FOnPhaseEvent OnPhaseStart;
+	FOnPhaseEvent OnPhaseEnd;
+
+private:
+	bool bRunInParallelMode = false;
+	std::atomic<bool> bIsDuringMassProcessing = false;
+};
+
+// It is unsafe to copy FTickFunctions and any subclasses of FTickFunction should specify the type trait WithCopy = false
+template<>
+struct TStructOpsTypeTraits<FMassProcessingPhase> : public TStructOpsTypeTraitsBase2<FMassProcessingPhase>
+{
+	enum
+	{
+		WithCopy = false
+	};
+};
+
+/** MassProcessingPhaseManager owns separate FMassProcessingPhase instances for every ETickingGroup. When activated
+ *  via Start function it registers and enables the FMassProcessingPhase instances which themselves are tick functions 
+ *  that host UMassCompositeProcessor which they trigger as part of their Tick function. 
+ *  MassProcessingPhaseManager serves as an interface to said FMassProcessingPhase instances and allows initialization
+ *  with MassSchematics (via InitializePhases function) as well as registering arbitrary functions to be called 
+ *  when a particular phase starts of ends (via GetOnPhaseStart and GetOnPhaseEnd functions). */
+UCLASS(Transient, HideCategories = (Tick))
+class MASSENTITY_API UMassProcessingPhaseManager : public UObject
+{
+	GENERATED_BODY()
+public:
+	FMassEntityManager& GetEntityManagerRef() { check(EntityManager); return *EntityManager.Get(); }
+
+	/** Retrieves OnPhaseStart multicast delegate's reference for a given Phase */
+	FMassProcessingPhase::FOnPhaseEvent& GetOnPhaseStart(const EMassProcessingPhase Phase) { return ProcessingPhases[uint8(Phase)].OnPhaseStart; }
+	/** Retrieves OnPhaseEnd multicast delegate's reference for a given Phase */
+	FMassProcessingPhase::FOnPhaseEvent& GetOnPhaseEnd(const EMassProcessingPhase Phase) { return ProcessingPhases[uint8(Phase)].OnPhaseEnd; }
+
+	/** 
+	 *  Populates hosted FMassProcessingPhase instances with Processors read from MassEntitySettings configuration.
+	 *  Calling this function overrides previous configuration of Phases.
+	 */
+	void InitializePhases(UObject& InProcessorOwner);
+
+	/** 
+	 *  Both flavors of Start function boil down to setting EntitySubsystem and Executor. If the callee has these 
+	 *  at hand it's suggested to use that Start version, otherwise call the World-using one. 
+	 */
+	void Start(UWorld& World);
+	void Start(const TSharedPtr<FMassEntityManager>& InEntityManager);
+	void Stop();
+	bool IsRunning() const { return EntityManager.IsValid(); }
+
+	/** 
+	 *  returns true when called while any of the ProcessingPhases is actively executing its processors. Used to 
+	 *  determine whether it's safe to do entity-related operations like adding fragments.
+	 *  Note that the function will return false while the OnPhaseStart or OnPhaseEnd are being broadcast,
+	 *  the value returned will be `true` only when the entity subsystem is actively engaged 
+	 */
+	bool IsDuringMassProcessing() const { return CurrentPhase != EMassProcessingPhase::MAX && ProcessingPhases[int(CurrentPhase)].IsDuringMassProcessing(); }
+
+	/**
+	 *  Sets the composite processor used for Phase processing phase. Using PhaseProcessor == nullptr results in
+	 *  clearing the Phase's processor out effectively making the Phase perform no calculations. 
+	 */
+	void SetPhaseProcessor(const EMassProcessingPhase Phase, UMassCompositeProcessor* PhaseProcessor);
+
+protected:
+	virtual void PostInitProperties() override;
+	virtual void BeginDestroy() override;
+	void EnableTickFunctions(const UWorld& World);
+
+	/** Creates phase processors instances for each declared phase name, based on MassEntitySettings */
+	virtual void CreatePhases();
+
+#if WITH_EDITOR
+	virtual void OnMassEntitySettingsChange(const FPropertyChangedEvent& PropertyChangedEvent);
+#endif // WITH_EDITOR
+
+	friend FMassProcessingPhase;
+
+	/** 
+	 *  Called by the given Phase at the very start of its execution function (the FMassProcessingPhase::ExecuteTick),
+	 *  even before the FMassProcessingPhase.OnPhaseStart broadcast delegate
+	 */
+	void OnPhaseStart(const FMassProcessingPhase& Phase);
+
+	/**
+	 *  Called by the given Phase at the very end of its execution function (the FMassProcessingPhase::ExecuteTick),
+	 *  after the FMassProcessingPhase.OnPhaseEnd broadcast delegate
+	 */
+	void OnPhaseEnd(FMassProcessingPhase& Phase);
+
+protected:	
+	UPROPERTY(VisibleAnywhere, Category=Mass)
+	FMassProcessingPhase ProcessingPhases[(uint8)EMassProcessingPhase::MAX];
+
+	TSharedPtr<FMassEntityManager> EntityManager;
+
+	EMassProcessingPhase CurrentPhase = EMassProcessingPhase::MAX;
+
+#if WITH_EDITOR
+	FDelegateHandle MassEntitySettingsChangeHandle;
+#endif // WITH_EDITOR
+};
diff --git a/Engine/Plugins/Runtime/MassGameplay/Source/MassSimulation/Public/MassSimulationSubsystem.h b/Engine/Plugins/Runtime/MassGameplay/Source/MassSimulation/Public/MassSimulationSubsystem.h
--- a/Engine/Plugins/Runtime/MassGameplay/Source/MassSimulation/Public/MassSimulationSubsystem.h
+++ b/Engine/Plugins/Runtime/MassGameplay/Source/MassSimulation/Public/MassSimulationSubsystem.h
@@ -4,7 +4,7 @@
 
 #include "Delegates/IDelegateInstance.h"
 #include "MassProcessingTypes.h"
-#include "MassProcessingPhase.h"
+#include "MassProcessingPhaseManager.h"
 #include "Subsystems/WorldSubsystem.h"
 #include "MassSimulationSubsystem.generated.h"
 
