Turned UMassProcessingPhaseManager into a struct since it doesn't really need to be a uobject.
The change is a first step towards FMassProcessingPhaseManager's reusability.

Also:
* simplified phase manager's initialization - there's a single function for that and the caller is responsible to pass in phases configuration
* moved the bits responding to MassEntitySettings changes from FMassProcessingPhaseManager over to UMassSimulationSubsystem (that hosts the default phase manager instance, at least for now).

diff --git a/Engine/Plugins/Runtime/MassEntity/Source/MassEntity/Private/MassEntitySettings.cpp b/Engine/Plugins/Runtime/MassEntity/Source/MassEntity/Private/MassEntitySettings.cpp
--- a/Engine/Plugins/Runtime/MassEntity/Source/MassEntity/Private/MassEntitySettings.cpp
+++ b/Engine/Plugins/Runtime/MassEntity/Source/MassEntity/Private/MassEntitySettings.cpp
@@ -125,10 +125,10 @@ void UMassEntitySettings::AddToActiveProcessorsList(TSubclassOf<UMassProcessor>
 	}
 }
 
-const FMassProcessingPhaseConfig* UMassEntitySettings::GetProcessingPhasesConfig()
+TConstArrayView<FMassProcessingPhaseConfig> UMassEntitySettings::GetProcessingPhasesConfig()
 {
 	BuildProcessorListAndPhases();
-	return ProcessingPhasesConfig;
+	return MakeArrayView(ProcessingPhasesConfig, int32(EMassProcessingPhase::MAX));
 }
 
 #if WITH_EDITOR
diff --git a/Engine/Plugins/Runtime/MassEntity/Source/MassEntity/Private/MassProcessingPhaseManager.cpp b/Engine/Plugins/Runtime/MassEntity/Source/MassEntity/Private/MassProcessingPhaseManager.cpp
--- a/Engine/Plugins/Runtime/MassEntity/Source/MassEntity/Private/MassProcessingPhaseManager.cpp
+++ b/Engine/Plugins/Runtime/MassEntity/Source/MassEntity/Private/MassProcessingPhaseManager.cpp
@@ -3,7 +3,6 @@
 #include "MassProcessingPhaseManager.h"
 #include "MassProcessingTypes.h"
 #include "MassDebugger.h"
-#include "MassEntitySettings.h"
 #include "MassProcessor.h"
 #include "MassExecutor.h"
 #include "MassEntityManager.h"
@@ -118,75 +117,61 @@ void FMassProcessingPhase::OnParallelExecutionDone(const float DeltaTime)
 
 FString FMassProcessingPhase::DiagnosticMessage()
 {
-	return (PhaseManager ? PhaseManager->GetFullName() : TEXT("NULL-MassProcessingPhaseManager")) + TEXT("[ProcessorTick]");
+	return (PhaseManager ? PhaseManager->GetName() : TEXT("NULL-MassProcessingPhaseManager")) + TEXT("[ProcessingPhaseTick]");
 }
 
 FName FMassProcessingPhase::DiagnosticContext(bool bDetailed)
 {
-	return PhaseManager ? PhaseManager->GetClass()->GetFName() : TEXT("NULL-MassProcessingPhaseManager");
+	return TEXT("MassProcessingPhase");
 }
 
-//----------------------------------------------------------------------//
-// UMassProcessingPhaseManager  
-//----------------------------------------------------------------------//
-void UMassProcessingPhaseManager::PostInitProperties()
+void FMassProcessingPhase::Initialize(FMassProcessingPhaseManager& InPhaseManager, const EMassProcessingPhase InPhase, const ETickingGroup InTickGroup, UMassCompositeProcessor& InPhaseProcessor)
 {
-	Super::PostInitProperties();
-	if (HasAnyFlags(RF_ClassDefaultObject) == false)
-	{
-#if WITH_EDITOR
-		UWorld* World = GetWorld();
-		if (World && World->IsGameWorld() == false)
-		{
-			UMassEntitySettings* Settings = GetMutableDefault<UMassEntitySettings>();
-			check(Settings);
-			MassEntitySettingsChangeHandle = Settings->GetOnSettingsChange().AddUObject(this, &UMassProcessingPhaseManager::OnMassEntitySettingsChange);
-		}
-#endif // WITH_EDITOR
-	}
-	CreatePhases();
+	PhaseManager = &InPhaseManager;
+	Phase = InPhase;
+	TickGroup = InTickGroup;
+	PhaseProcessor = &InPhaseProcessor;
 }
 
-void UMassProcessingPhaseManager::BeginDestroy()
+//----------------------------------------------------------------------//
+// FMassProcessingPhaseManager
+//----------------------------------------------------------------------//
+void FMassProcessingPhaseManager::Initialize(UObject& InOwner, TConstArrayView<FMassProcessingPhaseConfig> ProcessingPhasesConfig, const FString& DependencyGraphFileName)
 {
-	Stop();
-	Super::BeginDestroy();
+	Owner = &InOwner;
 
-#if WITH_EDITOR
-	if (UMassEntitySettings* Settings = GetMutableDefault<UMassEntitySettings>())
-	{
-		Settings->GetOnSettingsChange().Remove(MassEntitySettingsChangeHandle);
-	}
-#endif // WITH_EDITOR
-}
+	for (int i = 0; i < int(EMassProcessingPhase::MAX); ++i)
+	{		
+		const EMassProcessingPhase Phase = EMassProcessingPhase(i);
 
-void UMassProcessingPhaseManager::InitializePhases(UObject& InProcessorOwner)
-{
-	const FMassProcessingPhaseConfig* ProcessingPhasesConfig = GET_MASS_CONFIG_VALUE(GetProcessingPhasesConfig());
+		UMassCompositeProcessor* PhaseProcessor = NewObject<UMassCompositeProcessor>(&InOwner, UMassCompositeProcessor::StaticClass()
+			, *FString::Printf(TEXT("ProcessingPhase_%s"), *UEnum::GetDisplayValueAsText(Phase).ToString()));
+	
+		check(PhaseProcessor);
+		ProcessingPhases[i].Initialize(*this, Phase, UE::Mass::Private::PhaseToTickingGroup[i], *PhaseProcessor);
 
-	FString DependencyGraphFileName;
+		REDIRECT_OBJECT_TO_VLOG(PhaseProcessor, &InOwner);
+		PhaseProcessor->SetProcessingPhase(Phase);
+		PhaseProcessor->SetGroupName(FName(FString::Printf(TEXT("%s Group"), *UEnum::GetDisplayValueAsText(Phase).ToString())));
 
-#if WITH_EDITOR
-	const UWorld* World = InProcessorOwner.GetWorld();
-	const UMassEntitySettings* Settings = GetMutableDefault<UMassEntitySettings>();
-	if (World != nullptr && Settings != nullptr && !Settings->DumpDependencyGraphFileName.IsEmpty())
-	{
-		DependencyGraphFileName = FString::Printf(TEXT("%s_%s"), *Settings->DumpDependencyGraphFileName,*ToString(World->GetNetMode()));
-	}
-#endif // WITH_EDITOR
+		check(&InOwner);
+		PhaseProcessor->Initialize(InOwner);
+
+#if WITH_MASSENTITY_DEBUG
+		FStringOutputDevice Ar;
+		PhaseProcessor->DebugOutputDescription(Ar);
+		UE_VLOG(&InOwner, LogMass, Log, TEXT("Setting new group processor for phase %s:\n%s"), *UEnum::GetValueAsString(Phase), *Ar);
+#endif // WITH_MASSENTITY_DEBUG
 
-	for (int i = 0; i < int(EMassProcessingPhase::MAX); ++i)
-	{
 		const FMassProcessingPhaseConfig& PhaseConfig = ProcessingPhasesConfig[i];
-		UMassCompositeProcessor* PhaseProcessor = ProcessingPhases[i].PhaseProcessor;
 		FString FileName = !DependencyGraphFileName.IsEmpty() ? FString::Printf(TEXT("%s_%s"), *DependencyGraphFileName, *PhaseConfig.PhaseName.ToString()) : FString();
 		PhaseProcessor->CopyAndSort(PhaseConfig, FileName);
-		PhaseProcessor->Initialize(InProcessorOwner);
+		PhaseProcessor->Initialize(InOwner);
 	}
 
 #if WITH_MASSENTITY_DEBUG
 	// print it all out to vislog
-	UE_VLOG_UELOG(this, LogMass, Verbose, TEXT("Phases initialization done. Current composition:"));
+	UE_VLOG_UELOG(&InOwner, LogMass, Verbose, TEXT("Phases initialization done. Current composition:"));
 
 	FStringOutputDevice OutDescription;
 	for (int i = 0; i < int(EMassProcessingPhase::MAX); ++i)
@@ -194,14 +179,22 @@ void UMassProcessingPhaseManager::InitializePhases(UObject& InProcessorOwner)
 		if (ProcessingPhases[i].PhaseProcessor)
 		{
 			ProcessingPhases[i].PhaseProcessor->DebugOutputDescription(OutDescription);
-			UE_VLOG_UELOG(this, LogMass, Verbose, TEXT("--- %s"), *OutDescription);
+			UE_VLOG_UELOG(&InOwner, LogMass, Verbose, TEXT("--- %s"), *OutDescription);
 			OutDescription.Reset();
 		}
 	}	
 #endif // WITH_MASSENTITY_DEBUG
 }
 
-void UMassProcessingPhaseManager::Start(UWorld& World)
+void FMassProcessingPhaseManager::Deinitialize()
+{
+	for (FMassProcessingPhase& Phase : ProcessingPhases)
+	{
+		Phase.PhaseProcessor = nullptr;
+	}
+}
+
+void FMassProcessingPhaseManager::Start(UWorld& World)
 {
 	UMassEntitySubsystem* EntitySubsystem = UWorld::GetSubsystem<UMassEntitySubsystem>(&World);
 
@@ -212,11 +205,11 @@ void UMassProcessingPhaseManager::Start(UWorld& World)
 	}
 	else
 	{
-		UE_VLOG_UELOG(this, LogMass, Error, TEXT("Called %s while missing the EntitySubsystem"), ANSI_TO_TCHAR(__FUNCTION__));
+		UE_VLOG_UELOG(Owner.Get(), LogMass, Error, TEXT("Called %s while missing the EntitySubsystem"), ANSI_TO_TCHAR(__FUNCTION__));
 	}
 }
 
-void UMassProcessingPhaseManager::Start(const TSharedPtr<FMassEntityManager>& InEntityManager)
+void FMassProcessingPhaseManager::Start(const TSharedPtr<FMassEntityManager>& InEntityManager)
 {
 	UWorld* World = InEntityManager->GetWorld();
 	check(World);
@@ -224,7 +217,18 @@ void UMassProcessingPhaseManager::Start(const TSharedPtr<FMassEntityManager>& In
 	EnableTickFunctions(*World);
 }
 
-void UMassProcessingPhaseManager::EnableTickFunctions(const UWorld& World)
+void FMassProcessingPhaseManager::AddReferencedObjects(FReferenceCollector& Collector)
+{
+	for (FMassProcessingPhase& Phase : ProcessingPhases)
+	{
+		if (Phase.PhaseProcessor)
+		{
+			Collector.AddReferencedObject(Phase.PhaseProcessor);
+		}
+	}
+}
+
+void FMassProcessingPhaseManager::EnableTickFunctions(const UWorld& World)
 {
 	check(EntityManager);
 
@@ -232,7 +236,6 @@ void UMassProcessingPhaseManager::EnableTickFunctions(const UWorld& World)
 
 	for (FMassProcessingPhase& Phase : ProcessingPhases)
 	{
-		Phase.PhaseManager = this;
 		Phase.RegisterTickFunction(World.PersistentLevel);
 		Phase.SetTickFunctionEnable(true);
 #if WITH_MASSENTITY_DEBUG
@@ -241,21 +244,21 @@ void UMassProcessingPhaseManager::EnableTickFunctions(const UWorld& World)
 			// not logging this in the editor mode since it messes up the game-recorded vislog display (with its progressively larger timestamp)
 			FStringOutputDevice Ar;
 			Phase.PhaseProcessor->DebugOutputDescription(Ar);
-			UE_VLOG_UELOG(this, LogMass, Log, TEXT("Enabling phase %s tick:\n%s")
+			UE_VLOG_UELOG(Owner.Get(), LogMass, Log, TEXT("Enabling phase %s tick:\n%s")
 				, *UEnum::GetValueAsString(Phase.Phase), *Ar);
 		}
-#endif // WITH_MASSENTITY_DEBUG		
+#endif // WITH_MASSENTITY_DEBUG
 	}
 
 	if (bIsGameWorld)
 	{
 		// not logging this in the editor mode since it messes up the game-recorded vislog display (with its progressively larger timestamp)
-		UE_VLOG_UELOG(this, LogMass, Log, TEXT("MassProcessingPhaseManager %s.%s has been started")
-			, *GetNameSafe(GetOuter()), *GetName());
+		UE_VLOG_UELOG(Owner.Get(), LogMass, Log, TEXT("MassProcessingPhaseManager %s.%s has been started")
+			, *GetNameSafe(Owner.Get()), *GetName());
 	}
 }
 
-void UMassProcessingPhaseManager::Stop()
+void FMassProcessingPhaseManager::Stop()
 {
 	EntityManager.Reset();
 	
@@ -264,68 +267,25 @@ void UMassProcessingPhaseManager::Stop()
 		Phase.SetTickFunctionEnable(false);
 	}
 
-	UWorld* World = GetWorld();
-	if (World && World->IsGameWorld())
-	{
-		// not logging this in editor mode since it messes up the game-recorded vislog display (with its progressively larger timestamp) 
-		UE_VLOG_UELOG(this, LogMass, Log, TEXT("MassProcessingPhaseManager %s.%s has been stopped")
-			, *GetNameSafe(GetOuter()), *GetName());
-	}
-}
-
-void UMassProcessingPhaseManager::SetPhaseProcessor(const EMassProcessingPhase Phase, UMassCompositeProcessor* PhaseProcessor)
-{
-	if (Phase == EMassProcessingPhase::MAX)
-	{
-		UE_VLOG_UELOG(this, LogMass, Error, TEXT("MassProcessingPhaseManager::OverridePhaseProcessor called with Phase == MAX"));
-		return;
-	}
-	
-	// note that it's ok to use PhaseProcessor == nullptr
-	ProcessingPhases[int32(Phase)].PhaseProcessor = PhaseProcessor;
-	
-	if (PhaseProcessor)
-	{
-		REDIRECT_OBJECT_TO_VLOG(PhaseProcessor, this);
-		PhaseProcessor->SetProcessingPhase(Phase);
-		PhaseProcessor->SetGroupName(FName(FString::Printf(TEXT("%s Group"), *UEnum::GetDisplayValueAsText(Phase).ToString())));
-		
-		check(GetOuter());
-		PhaseProcessor->Initialize(*GetOuter());
-
-#if WITH_MASSENTITY_DEBUG
-		FStringOutputDevice Ar;
-		PhaseProcessor->DebugOutputDescription(Ar);
-		UE_VLOG(this, LogMass, Log, TEXT("Setting new group processor for phase %s:\n%s")
-			, *UEnum::GetValueAsString(Phase), *Ar);
-#endif // WITH_MASSENTITY_DEBUG
-	}
-	else
-	{
-		UE_VLOG_UELOG(this, LogMass, Log, TEXT("Nulling-out phase processor for phase %s"), *UEnum::GetValueAsString(Phase));
-	}
-}
-
-void UMassProcessingPhaseManager::CreatePhases() 
-{
-	// @todo copy from settings instead of blindly creating from scratch
-	for (int i = 0; i < int(EMassProcessingPhase::MAX); ++i)
+	if (UObject* LocalOwner = Owner.Get())
 	{
-		ProcessingPhases[i].Phase = EMassProcessingPhase(i);
-		ProcessingPhases[i].TickGroup = UE::Mass::Private::PhaseToTickingGroup[i];
-		UMassCompositeProcessor* PhaseProcessor = NewObject<UMassCompositeProcessor>(this, UMassCompositeProcessor::StaticClass()
-			, *FString::Printf(TEXT("ProcessingPhase_%s"), *UEnum::GetDisplayValueAsText(EMassProcessingPhase(i)).ToString()));
-		SetPhaseProcessor(EMassProcessingPhase(i), PhaseProcessor);
+		UWorld* World = LocalOwner->GetWorld();
+		if (World && World->IsGameWorld())
+		{
+			// not logging this in editor mode since it messes up the game-recorded vislog display (with its progressively larger timestamp) 
+			UE_VLOG_UELOG(LocalOwner, LogMass, Log, TEXT("MassProcessingPhaseManager %s.%s has been stopped")
+				, *GetNameSafe(LocalOwner), *GetName());
+		}
 	}
 }
 
-void UMassProcessingPhaseManager::OnPhaseStart(const FMassProcessingPhase& Phase)
+void FMassProcessingPhaseManager::OnPhaseStart(const FMassProcessingPhase& Phase)
 {
 	ensure(CurrentPhase == EMassProcessingPhase::MAX);
 	CurrentPhase = Phase.Phase;
 }
 
-void UMassProcessingPhaseManager::OnPhaseEnd(FMassProcessingPhase& Phase)
+void FMassProcessingPhaseManager::OnPhaseEnd(FMassProcessingPhase& Phase)
 {
 	ensure(CurrentPhase == Phase.Phase);
 	CurrentPhase = EMassProcessingPhase::MAX;
@@ -344,12 +304,9 @@ void UMassProcessingPhaseManager::OnPhaseEnd(FMassProcessingPhase& Phase)
 	}
 }
 
-#if WITH_EDITOR
-void UMassProcessingPhaseManager::OnMassEntitySettingsChange(const FPropertyChangedEvent& PropertyChangedEvent)
+FString FMassProcessingPhaseManager::GetName() const
 {
-	check(GetOuter());
-	InitializePhases(*GetOuter());
+	return GetNameSafe(Owner.Get()) + TEXT("_MassProcessingPhaseManager");
 }
-#endif // WITH_EDITOR
 
-#undef LOCTEXT_NAMESPACE
\ No newline at end of file
+#undef LOCTEXT_NAMESPACE
diff --git a/Engine/Plugins/Runtime/MassEntity/Source/MassEntity/Public/MassEntitySettings.h b/Engine/Plugins/Runtime/MassEntity/Source/MassEntity/Public/MassEntitySettings.h
--- a/Engine/Plugins/Runtime/MassEntity/Source/MassEntity/Public/MassEntitySettings.h
+++ b/Engine/Plugins/Runtime/MassEntity/Source/MassEntity/Public/MassEntitySettings.h
@@ -11,36 +11,9 @@
 
 #define GET_MASS_CONFIG_VALUE(a) (GetMutableDefault<UMassEntitySettings>()->a)
 
-class UMassProcessingPhaseManager;
-struct FMassProcessingPhaseConfig;
 struct FPropertyChangedEvent;
 
 
-USTRUCT()
-struct FMassProcessingPhaseConfig
-{
-	GENERATED_BODY()
-
-	UPROPERTY(EditAnywhere, Category = Mass, config)
-	FName PhaseName;
-
-	UPROPERTY(EditAnywhere, Category = Mass, config, NoClear)
-	TSubclassOf<UMassCompositeProcessor> PhaseGroupClass = UMassCompositeProcessor::StaticClass();
-
-	UPROPERTY(Transient)
-	TArray<TObjectPtr<UMassProcessor>> ProcessorCDOs;
-
-#if WITH_EDITORONLY_DATA
-	// this processor is available only in editor since it's used to present the user the order in which processors
-	// will be executed when given processing phase gets triggered
-	UPROPERTY(Transient)
-	TObjectPtr<UMassCompositeProcessor> PhaseProcessor = nullptr;
-
-	UPROPERTY(VisibleAnywhere, Category = Mass, Transient)
-	FText Description;
-#endif //WITH_EDITORONLY_DATA
-};
-
 /**
  * Implements the settings for MassEntity plugin
  */
@@ -58,7 +31,7 @@ public:
 	void BuildProcessorListAndPhases();
 	void AddToActiveProcessorsList(TSubclassOf<UMassProcessor> ProcessorClass);
 
-	const FMassProcessingPhaseConfig* GetProcessingPhasesConfig();
+	TConstArrayView<FMassProcessingPhaseConfig> GetProcessingPhasesConfig();
 	const FMassProcessingPhaseConfig& GetProcessingPhaseConfig(const EMassProcessingPhase ProcessingPhase) const { check(ProcessingPhase != EMassProcessingPhase::MAX); return ProcessingPhasesConfig[int(ProcessingPhase)]; }
 
 #if WITH_EDITOR
diff --git a/Engine/Plugins/Runtime/MassEntity/Source/MassEntity/Public/MassProcessingPhaseManager.h b/Engine/Plugins/Runtime/MassEntity/Source/MassEntity/Public/MassProcessingPhaseManager.h
--- a/Engine/Plugins/Runtime/MassEntity/Source/MassEntity/Public/MassProcessingPhaseManager.h
+++ b/Engine/Plugins/Runtime/MassEntity/Source/MassEntity/Public/MassProcessingPhaseManager.h
@@ -5,25 +5,54 @@
 #include "CoreMinimal.h"
 #include "UObject/UObjectGlobals.h"
 #include "UObject/Object.h"
+#include "UObject/GCObject.h"
 #include "Engine/EngineBaseTypes.h"
 #include "MassProcessingTypes.h"
+#include "MassProcessor.h"
 #include "MassProcessingPhaseManager.generated.h"
 
 
-class UMassProcessingPhaseManager;
+struct FMassProcessingPhaseManager;
 class UMassProcessor;
 class UMassCompositeProcessor;
 struct FMassEntityManager;
 struct FMassCommandBuffer;
+struct FMassProcessingPhaseConfig;
+
 
 USTRUCT()
-struct FMassProcessingPhase : public FTickFunction
+struct MASSENTITY_API FMassProcessingPhaseConfig
 {
 	GENERATED_BODY()
 
+	UPROPERTY(EditAnywhere, Category = Mass, config)
+	FName PhaseName;
+
+	UPROPERTY(EditAnywhere, Category = Mass, config, NoClear)
+	TSubclassOf<UMassCompositeProcessor> PhaseGroupClass = UMassCompositeProcessor::StaticClass();
+
+	UPROPERTY(Transient)
+	TArray<TObjectPtr<UMassProcessor>> ProcessorCDOs;
+
+#if WITH_EDITORONLY_DATA
+	// this processor is available only in editor since it's used to present the user the order in which processors
+	// will be executed when given processing phase gets triggered
+	UPROPERTY(Transient)
+	TObjectPtr<UMassCompositeProcessor> PhaseProcessor = nullptr;
+
+	UPROPERTY(VisibleAnywhere, Category = Mass, Transient)
+	FText Description;
+#endif //WITH_EDITORONLY_DATA
+};
+
+
+struct MASSENTITY_API FMassProcessingPhase : public FTickFunction
+{
 	DECLARE_MULTICAST_DELEGATE_OneParam(FOnPhaseEvent, const float /*DeltaSeconds*/);
 
 	FMassProcessingPhase();
+	FMassProcessingPhase(const FMassProcessingPhase& Other) = delete;
+	FMassProcessingPhase& operator=(const FMassProcessingPhase& Other) = delete;
 
 protected:
 	// FTickFunction interface
@@ -41,45 +70,40 @@ protected:
 public:
 	bool IsDuringMassProcessing() const { return bIsDuringMassProcessing; }
 
-protected:
-	friend UMassProcessingPhaseManager;
+	void Initialize(FMassProcessingPhaseManager& InPhaseManager, const EMassProcessingPhase InPhase, const ETickingGroup InTickGroup, UMassCompositeProcessor& InPhaseProcessor);
 
-	UPROPERTY(EditAnywhere, Category=Mass)
-	TObjectPtr<UMassCompositeProcessor> PhaseProcessor = nullptr;
+protected:
+	friend FMassProcessingPhaseManager;
 
-	UPROPERTY()
-	TObjectPtr<UMassProcessingPhaseManager> PhaseManager = nullptr;
+	// composite processor representing work to be performed. GC-referenced via AddReferencedObjects
+	UMassCompositeProcessor* PhaseProcessor = nullptr;
 
 	EMassProcessingPhase Phase = EMassProcessingPhase::MAX;
 	FOnPhaseEvent OnPhaseStart;
 	FOnPhaseEvent OnPhaseEnd;
 
 private:
+	FMassProcessingPhaseManager* PhaseManager = nullptr;
 	bool bRunInParallelMode = false;
 	std::atomic<bool> bIsDuringMassProcessing = false;
 };
 
-// It is unsafe to copy FTickFunctions and any subclasses of FTickFunction should specify the type trait WithCopy = false
-template<>
-struct TStructOpsTypeTraits<FMassProcessingPhase> : public TStructOpsTypeTraitsBase2<FMassProcessingPhase>
-{
-	enum
-	{
-		WithCopy = false
-	};
-};
 
-/** MassProcessingPhaseManager owns separate FMassProcessingPhase instances for every ETickingGroup. When activated
- *  via Start function it registers and enables the FMassProcessingPhase instances which themselves are tick functions 
- *  that host UMassCompositeProcessor which they trigger as part of their Tick function. 
- *  MassProcessingPhaseManager serves as an interface to said FMassProcessingPhase instances and allows initialization
- *  with MassSchematics (via InitializePhases function) as well as registering arbitrary functions to be called 
- *  when a particular phase starts of ends (via GetOnPhaseStart and GetOnPhaseEnd functions). */
-UCLASS(Transient, HideCategories = (Tick))
-class MASSENTITY_API UMassProcessingPhaseManager : public UObject
+/** 
+ * MassProcessingPhaseManager owns separate FMassProcessingPhase instances for every ETickingGroup. When activated
+ * via Start function it registers and enables the FMassProcessingPhase instances which themselves are tick functions 
+ * that host UMassCompositeProcessor which they trigger as part of their Tick function. 
+ * MassProcessingPhaseManager serves as an interface to said FMassProcessingPhase instances and allows initialization
+ * with collections of processors (via Initialize function) as well as registering arbitrary functions to be called 
+ * when a particular phase starts or ends (via GetOnPhaseStart and GetOnPhaseEnd functions). 
+ */
+struct MASSENTITY_API FMassProcessingPhaseManager : public FGCObject
 {
-	GENERATED_BODY()
 public:
+	FMassProcessingPhaseManager() = default;
+	FMassProcessingPhaseManager(const FMassProcessingPhaseManager& Other) = delete;
+	FMassProcessingPhaseManager& operator=(const FMassProcessingPhaseManager& Other) = delete;
+
 	FMassEntityManager& GetEntityManagerRef() { check(EntityManager); return *EntityManager.Get(); }
 
 	/** Retrieves OnPhaseStart multicast delegate's reference for a given Phase */
@@ -91,13 +115,19 @@ public:
 	 *  Populates hosted FMassProcessingPhase instances with Processors read from MassEntitySettings configuration.
 	 *  Calling this function overrides previous configuration of Phases.
 	 */
-	void InitializePhases(UObject& InProcessorOwner);
+	void Initialize(UObject& InOwner, TConstArrayView<FMassProcessingPhaseConfig> ProcessingPhasesConfig, const FString& DependencyGraphFileName = TEXT(""));
+
+	/** Needs to be called before destruction, ideally before owner's BeginDestroy (a FGCObject's limitation) */
+	void Deinitialize();
 
 	/** 
-	 *  Both flavors of Start function boil down to setting EntitySubsystem and Executor. If the callee has these 
-	 *  at hand it's suggested to use that Start version, otherwise call the World-using one. 
+	 *  Stores EntityManager associated with given world's MassEntitySubsystem and kicks off phase ticking.
 	 */
 	void Start(UWorld& World);
+	
+	/**
+	 *  Stores InEntityManager as the entity manager and kicks off phase ticking.
+	 */
 	void Start(const TSharedPtr<FMassEntityManager>& InEntityManager);
 	void Stop();
 	bool IsRunning() const { return EntityManager.IsValid(); }
@@ -110,23 +140,21 @@ public:
 	 */
 	bool IsDuringMassProcessing() const { return CurrentPhase != EMassProcessingPhase::MAX && ProcessingPhases[int(CurrentPhase)].IsDuringMassProcessing(); }
 
-	/**
-	 *  Sets the composite processor used for Phase processing phase. Using PhaseProcessor == nullptr results in
-	 *  clearing the Phase's processor out effectively making the Phase perform no calculations. 
-	 */
-	void SetPhaseProcessor(const EMassProcessingPhase Phase, UMassCompositeProcessor* PhaseProcessor);
+	FString GetName() const;
 
 protected:
-	virtual void PostInitProperties() override;
-	virtual void BeginDestroy() override;
+	// FGCObject interface
+	virtual void AddReferencedObjects(FReferenceCollector& Collector) override;
+	virtual FString GetReferencerName() const override
+	{
+		return TEXT("FMassProcessingPhaseManager");
+	}
+	// End of FGCObject interface
+
 	void EnableTickFunctions(const UWorld& World);
 
 	/** Creates phase processors instances for each declared phase name, based on MassEntitySettings */
-	virtual void CreatePhases();
-
-#if WITH_EDITOR
-	virtual void OnMassEntitySettingsChange(const FPropertyChangedEvent& PropertyChangedEvent);
-#endif // WITH_EDITOR
+	void CreatePhases();
 
 	friend FMassProcessingPhase;
 
@@ -143,14 +171,12 @@ protected:
 	void OnPhaseEnd(FMassProcessingPhase& Phase);
 
 protected:	
-	UPROPERTY(VisibleAnywhere, Category=Mass)
+
 	FMassProcessingPhase ProcessingPhases[(uint8)EMassProcessingPhase::MAX];
 
 	TSharedPtr<FMassEntityManager> EntityManager;
 
 	EMassProcessingPhase CurrentPhase = EMassProcessingPhase::MAX;
 
-#if WITH_EDITOR
-	FDelegateHandle MassEntitySettingsChangeHandle;
-#endif // WITH_EDITOR
+	TWeakObjectPtr<UObject> Owner;
 };
diff --git a/Engine/Plugins/Runtime/MassEntity/Source/MassEntity/Public/MassProcessor.h b/Engine/Plugins/Runtime/MassEntity/Source/MassEntity/Public/MassProcessor.h
--- a/Engine/Plugins/Runtime/MassEntity/Source/MassEntity/Public/MassProcessor.h
+++ b/Engine/Plugins/Runtime/MassEntity/Source/MassEntity/Public/MassProcessor.h
@@ -87,7 +87,7 @@ public:
 #endif // WITH_EDITORONLY_DATA
 
 	/** Sets bAutoRegisterWithProcessingPhases. Setting it to true will result in this processor class being always 
-	 * instantiated to be automatically evaluated every frame. @see UMassProcessingPhaseManager
+	 * instantiated to be automatically evaluated every frame. @see FMassProcessingPhaseManager
 	 * Note that calling this function is only valid on CDOs. Calling it on a regular instance will fail an ensure and 
 	 * have no other effect, i.e. CDO's value won't change */
 	void SetShouldAutoRegisterWithGlobalList(const bool bAutoRegister);
diff --git a/Engine/Plugins/Runtime/MassGameplay/Source/MassActors/Private/MassActorSpawnerSubsystem.cpp b/Engine/Plugins/Runtime/MassGameplay/Source/MassActors/Private/MassActorSpawnerSubsystem.cpp
--- a/Engine/Plugins/Runtime/MassGameplay/Source/MassActors/Private/MassActorSpawnerSubsystem.cpp
+++ b/Engine/Plugins/Runtime/MassGameplay/Source/MassActors/Private/MassActorSpawnerSubsystem.cpp
@@ -361,7 +361,7 @@ void UMassActorSpawnerSubsystem::Deinitialize()
 
 	if (const UWorld* World = GetWorld())
 	{
-		if (const UMassSimulationSubsystem* SimSystem = UWorld::GetSubsystem<UMassSimulationSubsystem>(World))
+		if (UMassSimulationSubsystem* SimSystem = UWorld::GetSubsystem<UMassSimulationSubsystem>(World))
 		{
 			SimSystem->GetOnProcessingPhaseStarted(EMassProcessingPhase::PrePhysics).RemoveAll(this);
 			SimSystem->GetOnProcessingPhaseFinished(EMassProcessingPhase::PrePhysics).RemoveAll(this);
diff --git a/Engine/Plugins/Runtime/MassGameplay/Source/MassRepresentation/Private/MassRepresentationSubsystem.cpp b/Engine/Plugins/Runtime/MassGameplay/Source/MassRepresentation/Private/MassRepresentationSubsystem.cpp
--- a/Engine/Plugins/Runtime/MassGameplay/Source/MassRepresentation/Private/MassRepresentationSubsystem.cpp
+++ b/Engine/Plugins/Runtime/MassGameplay/Source/MassRepresentation/Private/MassRepresentationSubsystem.cpp
@@ -343,7 +343,7 @@ void UMassRepresentationSubsystem::Deinitialize()
 {
 	if (const UWorld* World = GetWorld())
 	{
-		if (const UMassSimulationSubsystem* SimSystem = World->GetSubsystem<UMassSimulationSubsystem>())
+		if (UMassSimulationSubsystem* SimSystem = World->GetSubsystem<UMassSimulationSubsystem>())
 		{
 			SimSystem->GetOnProcessingPhaseStarted(EMassProcessingPhase::PrePhysics).RemoveAll(this);
 			SimSystem->GetOnProcessingPhaseFinished(EMassProcessingPhase::PostPhysics).RemoveAll(this);
diff --git a/Engine/Plugins/Runtime/MassGameplay/Source/MassSimulation/Private/MassSimulationSubsystem.cpp b/Engine/Plugins/Runtime/MassGameplay/Source/MassSimulation/Private/MassSimulationSubsystem.cpp
--- a/Engine/Plugins/Runtime/MassGameplay/Source/MassSimulation/Private/MassSimulationSubsystem.cpp
+++ b/Engine/Plugins/Runtime/MassGameplay/Source/MassSimulation/Private/MassSimulationSubsystem.cpp
@@ -6,6 +6,7 @@
 #include "MassExecutor.h"
 #include "MassSimulationSettings.h"
 #include "VisualLogger/VisualLogger.h"
+#include "MassEntitySettings.h"
 #if WITH_EDITOR
 #include "Editor.h"
 #endif // WITH_EDITOR
@@ -24,7 +25,6 @@ namespace UE::MassSimulation
 UMassSimulationSubsystem::UMassSimulationSubsystem(const FObjectInitializer& ObjectInitializer)
 	: Super()
 {
-	PhaseManager = CreateDefaultSubobject<UMassProcessingPhaseManager>(TEXT("PhaseManager"));
 }
 
 void UMassSimulationSubsystem::BeginDestroy()
@@ -38,20 +38,32 @@ void UMassSimulationSubsystem::BeginDestroy()
 	{
 		FEditorDelegates::PrePIEEnded.Remove(PieEndedEventHandle);
 	}
+	if (MassEntitySettingsChangeHandle.IsValid())
+	{
+		if (UMassEntitySettings* Settings = GetMutableDefault<UMassEntitySettings>())
+		{
+			Settings->GetOnSettingsChange().Remove(MassEntitySettingsChangeHandle);
+		}
+	}
 #endif //  WITH_EDITOR
+
+	if (bSimulationStarted)
+	{
+		check(EntityManager);
+		StopSimulation();
+	}
+
 	Super::BeginDestroy();
 }
 
-FMassProcessingPhase::FOnPhaseEvent& UMassSimulationSubsystem::GetOnProcessingPhaseStarted(const EMassProcessingPhase Phase) const
+FMassProcessingPhase::FOnPhaseEvent& UMassSimulationSubsystem::GetOnProcessingPhaseStarted(const EMassProcessingPhase Phase)
 {
-	check(PhaseManager);
-	return PhaseManager->GetOnPhaseStart(Phase);
+	return PhaseManager.GetOnPhaseStart(Phase);
 }
 
-FMassProcessingPhase::FOnPhaseEvent& UMassSimulationSubsystem::GetOnProcessingPhaseFinished(const EMassProcessingPhase Phase) const
+FMassProcessingPhase::FOnPhaseEvent& UMassSimulationSubsystem::GetOnProcessingPhaseFinished(const EMassProcessingPhase Phase)
 {
-	check(PhaseManager);
-	return PhaseManager->GetOnPhaseEnd(Phase);
+	return PhaseManager.GetOnPhaseEnd(Phase);
 }
 
 void UMassSimulationSubsystem::Initialize(FSubsystemCollectionBase& Collection)
@@ -70,6 +82,7 @@ void UMassSimulationSubsystem::Deinitialize()
 	GetOnProcessingPhaseStarted(EMassProcessingPhase::PrePhysics).RemoveAll(this);
 	StopSimulation();
 
+	PhaseManager.Deinitialize();
 	EntityManager.Reset();
 
 	Super::Deinitialize();
@@ -81,7 +94,7 @@ void UMassSimulationSubsystem::PostInitialize()
 
 #if WITH_EDITOR
 	UWorld* World = GetWorld();
-	if (GEditor && World!= nullptr && !World->IsGameWorld())
+	if (GEditor && World != nullptr && !World->IsGameWorld())
 	{
 		// in editor worlds we need to rebuild the pipeline at this point since OnWorldBeginPlay won't be called
 		RebuildTickPipeline();
@@ -89,6 +102,10 @@ void UMassSimulationSubsystem::PostInitialize()
 		PieBeginEventHandle = FEditorDelegates::BeginPIE.AddUObject(this, &UMassSimulationSubsystem::OnPieBegin);
 		PieEndedEventHandle = FEditorDelegates::PrePIEEnded.AddUObject(this, &UMassSimulationSubsystem::OnPieEnded);
 
+		UMassEntitySettings* Settings = GetMutableDefault<UMassEntitySettings>();
+		check(Settings);
+		MassEntitySettingsChangeHandle = Settings->GetOnSettingsChange().AddUObject(this, &UMassSimulationSubsystem::OnMassEntitySettingsChange);
+
 		// note that this starts ticking for the editor world
 		StartSimulation(*World);
 	}
@@ -97,22 +114,35 @@ void UMassSimulationSubsystem::PostInitialize()
 
 void UMassSimulationSubsystem::OnWorldBeginPlay(UWorld& InWorld)
 {
-	// To evaluate the effective processors execution mode, we need to wait on OnWorldBeginPlay before calling RebuildTickPipeline as we are sure by this time the network is setup correctly.
+	// To evaluate the effective processors execution mode, we need to wait on OnWorldBeginPlay before calling
+	// RebuildTickPipeline as we are sure by this time the network is setup correctly.
 	RebuildTickPipeline();
-	// note that we're running for a game world right now. This means the PhaseManager->Start in OnPostWorldInit won't get called
+	// note that since we're in this function we're tied to a game world. This means the StartSimulation in 
+	// PostInitialize haven't been called.
 	StartSimulation(InWorld);
 }
 
 void UMassSimulationSubsystem::RebuildTickPipeline()
 {
-	check(PhaseManager);
-	PhaseManager->InitializePhases(*this);
+	TConstArrayView<FMassProcessingPhaseConfig> ProcessingPhasesConfig = GET_MASS_CONFIG_VALUE(GetProcessingPhasesConfig());
+
+	FString DependencyGraphFileName;
+
+#if WITH_EDITOR
+	const UWorld* World = GetWorld();
+	const UMassEntitySettings* Settings = GetMutableDefault<UMassEntitySettings>();
+	if (World != nullptr && Settings != nullptr && !Settings->DumpDependencyGraphFileName.IsEmpty())
+	{
+		DependencyGraphFileName = FString::Printf(TEXT("%s_%s"), *Settings->DumpDependencyGraphFileName, *ToString(World->GetNetMode()));
+	}
+#endif // WITH_EDITOR
+
+	PhaseManager.Initialize(*this, ProcessingPhasesConfig, DependencyGraphFileName);
 }
 
 void UMassSimulationSubsystem::StartSimulation(UWorld& InWorld)
 {
-	check(PhaseManager);
-	PhaseManager->Start(InWorld);
+	PhaseManager.Start(InWorld);
 
 	bSimulationStarted = true;
 
@@ -121,8 +151,7 @@ void UMassSimulationSubsystem::StartSimulation(UWorld& InWorld)
 
 void UMassSimulationSubsystem::StopSimulation()
 {
-	check(PhaseManager);
-	PhaseManager->Stop();
+	PhaseManager.Stop();
 
 	bSimulationStarted = false;
 }
@@ -151,7 +180,6 @@ void UMassSimulationSubsystem::OnProcessingPhaseStarted(const float DeltaSeconds
 #if WITH_EDITOR
 void UMassSimulationSubsystem::OnPieBegin(const bool bIsSimulation)
 {
-	check(PhaseManager);
 	// called so that we're not processing phases for the editor world while PIE/SIE is running
 	StopSimulation();
 }
@@ -164,4 +192,9 @@ void UMassSimulationSubsystem::OnPieEnded(const bool bIsSimulation)
 		StartSimulation(*World);
 	}
 }
+
+void UMassSimulationSubsystem::OnMassEntitySettingsChange(const FPropertyChangedEvent& PropertyChangedEvent)
+{
+	RebuildTickPipeline();
+}
 #endif // WITH_EDITOR
\ No newline at end of file
diff --git a/Engine/Plugins/Runtime/MassGameplay/Source/MassSimulation/Public/MassSimulationSubsystem.h b/Engine/Plugins/Runtime/MassGameplay/Source/MassSimulation/Public/MassSimulationSubsystem.h
--- a/Engine/Plugins/Runtime/MassGameplay/Source/MassSimulation/Public/MassSimulationSubsystem.h
+++ b/Engine/Plugins/Runtime/MassGameplay/Source/MassSimulation/Public/MassSimulationSubsystem.h
@@ -10,7 +10,6 @@
 
 
 struct FMassEntityManager;
-class UMassProcessingPhaseManager;
 
 DECLARE_LOG_CATEGORY_EXTERN(LogMassSim, Log, All);
 
@@ -23,10 +22,10 @@ public:
 	
 	UMassSimulationSubsystem(const FObjectInitializer& ObjectInitializer = FObjectInitializer::Get());
 
-	const UMassProcessingPhaseManager& GetPhaseManager() const { check(PhaseManager); return *PhaseManager; }
+	const FMassProcessingPhaseManager& GetPhaseManager() const { return PhaseManager; }
 
-	FMassProcessingPhase::FOnPhaseEvent& GetOnProcessingPhaseStarted(const EMassProcessingPhase Phase) const;
-	FMassProcessingPhase::FOnPhaseEvent& GetOnProcessingPhaseFinished(const EMassProcessingPhase Phase) const;
+	FMassProcessingPhase::FOnPhaseEvent& GetOnProcessingPhaseStarted(const EMassProcessingPhase Phase);
+	FMassProcessingPhase::FOnPhaseEvent& GetOnProcessingPhaseFinished(const EMassProcessingPhase Phase);
 	static FOnSimulationStarted& GetOnSimulationStarted() { return OnSimulationStarted; }
 
 	bool IsSimulationStarted() const { return bSimulationStarted; }
@@ -48,14 +47,14 @@ protected:
 #if WITH_EDITOR
 	void OnPieBegin(const bool bIsSimulation);
 	void OnPieEnded(const bool bIsSimulation);
+	void OnMassEntitySettingsChange(const FPropertyChangedEvent& PropertyChangedEvent);
 #endif // WITH_EDITOR
 
 protected:
 
 	TSharedPtr<FMassEntityManager> EntityManager;
 
-	UPROPERTY()
-	TObjectPtr<UMassProcessingPhaseManager> PhaseManager;
+	FMassProcessingPhaseManager PhaseManager;
 
 	inline static FOnSimulationStarted OnSimulationStarted={};
 
@@ -69,5 +68,7 @@ protected:
 #if WITH_EDITOR
 	FDelegateHandle PieBeginEventHandle;
 	FDelegateHandle PieEndedEventHandle;
+
+	FDelegateHandle MassEntitySettingsChangeHandle;
 #endif // WITH_EDITOR
 };
