Made MassDebugger UI tool react to entity selection via Mass GameplayDebugger category. The archetype containing given entity gets selected.

Massaged some of MassDebuggerModel's code while in there.

diff --git a/Engine/Plugins/AI/MassAI/Source/MassAIDebug/Private/GameplayDebuggerCategory_Mass.cpp b/Engine/Plugins/AI/MassAI/Source/MassAIDebug/Private/GameplayDebuggerCategory_Mass.cpp
--- a/Engine/Plugins/AI/MassAI/Source/MassAIDebug/Private/GameplayDebuggerCategory_Mass.cpp
+++ b/Engine/Plugins/AI/MassAI/Source/MassAIDebug/Private/GameplayDebuggerCategory_Mass.cpp
@@ -114,8 +114,11 @@ FGameplayDebuggerCategory_Mass::FGameplayDebuggerCategory_Mass()
 
 void FGameplayDebuggerCategory_Mass::SetCachedEntity(const FMassEntityHandle Entity, const FMassEntityManager& EntityManager)
 {
-	CachedEntity = Entity;
-	FMassDebugger::SelectEntity(EntityManager, Entity);
+	if (CachedEntity != Entity)
+	{
+		CachedEntity = Entity;
+		FMassDebugger::SelectEntity(EntityManager, Entity);
+	}
 }
 
 void FGameplayDebuggerCategory_Mass::PickEntity(const APlayerController& OwnerPC, const UWorld& World, FMassEntityManager& EntityManager, const bool bLimitAngle)
@@ -219,6 +222,10 @@ void FGameplayDebuggerCategory_Mass::CollectData(APlayerController* OwnerPC, AAc
 
 	AddTextLine(FString::Printf(TEXT("{Green}Entities count active{grey}/all: {white}%d{grey}/%d"), EntityManager.DebugGetEntityCount(), EntityManager.DebugGetEntityCount()));
 	AddTextLine(FString::Printf(TEXT("{Green}Registered Archetypes count: {white}%d {green}data ver: {white}%d"), EntityManager.DebugGetArchetypesCount(), EntityManager.GetArchetypeDataVersion()));
+	if (CachedEntity.IsValid())
+	{
+		AddTextLine(FString::Printf(TEXT("{Green}Entity: {White}%s"), *CachedEntity.DebugGetDescription()));
+	}
 
 	if (UE::Mass::Debug::HasDebugEntities())
 	{
@@ -289,7 +296,6 @@ void FGameplayDebuggerCategory_Mass::CollectData(APlayerController* OwnerPC, AAc
 			// (which causes the MassAgentComponent destruction and recreation).
 			if (EntityManager.IsEntityActive(CachedEntity))
 			{
-				AddTextLine(FString::Printf(TEXT("{Green}Entity: {White}%s"), *CachedEntity.DebugGetDescription()));
 				AddTextLine(FString::Printf(TEXT("{Green}Type: {White}%s"), (AgentComp == nullptr) ? TEXT("N/A") : AgentComp->IsPuppet() ? TEXT("PUPPET") : TEXT("AGENT")));
 
 				if (bShowEntityDetails)
diff --git a/Engine/Plugins/Runtime/MassEntity/Source/MassEntityDebugger/Private/MassDebuggerModel.cpp b/Engine/Plugins/Runtime/MassEntity/Source/MassEntityDebugger/Private/MassDebuggerModel.cpp
--- a/Engine/Plugins/Runtime/MassEntity/Source/MassEntityDebugger/Private/MassDebuggerModel.cpp
+++ b/Engine/Plugins/Runtime/MassEntity/Source/MassEntityDebugger/Private/MassDebuggerModel.cpp
@@ -295,6 +295,23 @@ FMassDebuggerProcessingGraph::FMassDebuggerProcessingGraph(const FMassDebuggerMo
 //----------------------------------------------------------------------//
 // FMassDebuggerModel
 //----------------------------------------------------------------------//
+FMassDebuggerModel::FMassDebuggerModel()
+{
+#if WITH_MASSENTITY_DEBUG
+	OnEntitySelectedHandle = FMassDebugger::OnEntitySelectedDelegate.AddRaw(this, &FMassDebuggerModel::OnEntitySelected);
+#endif // WITH_MASSENTITY_DEBUG
+}
+
+FMassDebuggerModel::~FMassDebuggerModel()
+{
+#if WITH_MASSENTITY_DEBUG
+	if (OnEntitySelectedHandle.IsValid())
+	{
+		FMassDebugger::OnEntitySelectedDelegate.Remove(OnEntitySelectedHandle);
+	}
+#endif // WITH_MASSENTITY_DEBUG
+}
+
 void FMassDebuggerModel::SetEnvironment(const TSharedPtr<FMassDebuggerEnvironment>& Item)
 {
 	if (Item)
@@ -362,7 +379,7 @@ void FMassDebuggerModel::ClearProcessorSelection()
 	OnProcessorsSelectedDelegate.Broadcast(SelectedProcessors, ESelectInfo::Direct);
 }
 
-void FMassDebuggerModel::SelectArchetypes(TArray<TSharedPtr<FMassDebuggerArchetypeData>> InSelectedArchetypes, ESelectInfo::Type SelectInfo)
+void FMassDebuggerModel::SelectArchetypes(TArrayView<TSharedPtr<FMassDebuggerArchetypeData>> InSelectedArchetypes, ESelectInfo::Type SelectInfo)
 {
 	ResetSelectedProcessors();
 	ResetSelectedArchetypes();
@@ -455,7 +472,8 @@ void FMassDebuggerModel::CacheProcessingGraphs()
 
 void FMassDebuggerModel::CacheArchetypesData(TMap<FMassArchetypeHandle, TSharedPtr<FMassDebuggerArchetypeData>>& OutTransientArchetypesMap)
 {
-	CachedArchetypes.Reset();
+	CachedAllArchetypes.Reset();
+	CachedArchetypeRepresentatives.Reset();
 
 	if (Environment)
 	{
@@ -481,7 +499,7 @@ void FMassDebuggerModel::StoreArchetypes(const FMassEntityManager& EntityManager
 #if WITH_MASSENTITY_DEBUG
 	TArray<FMassArchetypeHandle> ArchetypeHandles = FMassDebugger::GetAllArchetypes(EntityManager);
 
-	CachedArchetypes.Reset(ArchetypeHandles.Num());
+	CachedAllArchetypes.Reset(ArchetypeHandles.Num());
 
 	int32 MaxBitsUsed = 0;
 
@@ -489,8 +507,8 @@ void FMassDebuggerModel::StoreArchetypes(const FMassEntityManager& EntityManager
 	for (FMassArchetypeHandle& ArchetypeHandle : ArchetypeHandles)
 	{
 		FMassDebuggerArchetypeData* ArchetypeDataPtr = new FMassDebuggerArchetypeData(ArchetypeHandle);
-		ArchetypeDataPtr->Index = CachedArchetypes.Add(MakeShareable(ArchetypeDataPtr));
-		OutTransientArchetypesMap.Add(ArchetypeHandle, CachedArchetypes.Last());
+		ArchetypeDataPtr->Index = CachedAllArchetypes.Add(MakeShareable(ArchetypeDataPtr));
+		OutTransientArchetypesMap.Add(ArchetypeHandle, CachedAllArchetypes.Last());
 
 		MaxBitsUsed = FMath::Max(MaxBitsUsed, ArchetypeDataPtr->GetTotalBitsUsedCount());
 	}
@@ -498,17 +516,17 @@ void FMassDebuggerModel::StoreArchetypes(const FMassEntityManager& EntityManager
 
 	// calculate distances
 	ArchetypeDistances.Reset();
-	ArchetypeDistances.AddDefaulted(CachedArchetypes.Num());
-	for (int i = 0; i < CachedArchetypes.Num(); ++i)
+	ArchetypeDistances.AddDefaulted(CachedAllArchetypes.Num());
+	for (int i = 0; i < CachedAllArchetypes.Num(); ++i)
 	{
-		ArchetypeDistances[i].AddDefaulted(CachedArchetypes.Num());
+		ArchetypeDistances[i].AddDefaulted(CachedAllArchetypes.Num());
 	}
 
-	for (int i = 0; i < CachedArchetypes.Num(); ++i)
+	for (int i = 0; i < CachedAllArchetypes.Num(); ++i)
 	{
-		for (int k = i + 1; k < CachedArchetypes.Num(); ++k)
+		for (int k = i + 1; k < CachedAllArchetypes.Num(); ++k)
 		{
-			const float Distance = UE::Mass::Debugger::Private::CalcArchetypeBitDistance(*CachedArchetypes[i].Get(), *CachedArchetypes[k].Get());
+			const float Distance = UE::Mass::Debugger::Private::CalcArchetypeBitDistance(*CachedAllArchetypes[i].Get(), *CachedAllArchetypes[k].Get());
 			ArchetypeDistances[i][k] = Distance;
 			ArchetypeDistances[k][i] = Distance;
 		}
@@ -516,21 +534,23 @@ void FMassDebuggerModel::StoreArchetypes(const FMassEntityManager& EntityManager
 
 	// Add archetypes that share same primary name under the same entry. 
 	TMap<FString, TSharedPtr<FMassDebuggerArchetypeData>> ArchetypeNameMap;
-	for (int32 Index = 0; Index < CachedArchetypes.Num(); Index++)
+	for (TSharedPtr<FMassDebuggerArchetypeData>& ArchetypeData : CachedAllArchetypes)
 	{
-		TSharedPtr<FMassDebuggerArchetypeData> ArchetypeData = CachedArchetypes[Index];
 		if (const TSharedPtr<FMassDebuggerArchetypeData>* Representative = ArchetypeNameMap.Find(ArchetypeData->PrimaryDebugName))
 		{
 			(*Representative)->Children.Add(ArchetypeData);
-			ArchetypeData->Parent = *Representative; 
-			CachedArchetypes.RemoveAt(Index);
-			Index--;
+			ArchetypeData->Parent = *Representative;
 		}
 		else
 		{
 			ArchetypeNameMap.Add(ArchetypeData->PrimaryDebugName, ArchetypeData);
 		}
 	}
+
+	for (auto& KeyValue : ArchetypeNameMap)
+	{
+		CachedArchetypeRepresentatives.Add(KeyValue.Value);
+	}
 }
 
 FText FMassDebuggerModel::GetDisplayName() const
@@ -597,4 +617,30 @@ void FMassDebuggerModel::ResetSelectedProcessors()
 	SelectedProcessors.Reset();
 }
 
+void FMassDebuggerModel::OnEntitySelected(const FMassEntityManager& EntityManager, const FMassEntityHandle EntityHandle)
+{
+	if (!Environment || Environment->GetEntityManager() != &EntityManager)
+	{
+		// not the entity manager we're debugging right now
+		return;
+	}
+	
+	const FMassArchetypeHandle ArchetypeHandle = EntityManager.GetArchetypeForEntity(EntityHandle);
+	if (ArchetypeHandle.IsValid() == false)
+	{
+		return;
+	}
+	
+	uint32 ArchetypeHash = FMassDebugger::GetArchetypeComposition(ArchetypeHandle).CalculateHash();
+	TSharedPtr<FMassDebuggerArchetypeData>* DebuggerArchetypeData = CachedAllArchetypes.FindByPredicate([ArchetypeHash](const TSharedPtr<FMassDebuggerArchetypeData>& Element)
+		{
+			return Element.IsValid() && Element->CompositionHash == ArchetypeHash;
+		});
+
+	if (DebuggerArchetypeData)
+	{
+		SelectArchetypes(MakeArrayView(DebuggerArchetypeData, 1), ESelectInfo::Direct);
+	}
+}
+
 #undef LOCTEXT_NAMESPACE
diff --git a/Engine/Plugins/Runtime/MassEntity/Source/MassEntityDebugger/Private/MassDebuggerModel.h b/Engine/Plugins/Runtime/MassEntity/Source/MassEntityDebugger/Private/MassDebuggerModel.h
--- a/Engine/Plugins/Runtime/MassEntity/Source/MassEntityDebugger/Private/MassDebuggerModel.h
+++ b/Engine/Plugins/Runtime/MassEntity/Source/MassEntityDebugger/Private/MassDebuggerModel.h
@@ -155,6 +155,9 @@ struct FMassDebuggerModel
 	DECLARE_MULTICAST_DELEGATE_TwoParams(FOnProcessorsSelected, TConstArrayView<TSharedPtr<FMassDebuggerProcessorData>>, ESelectInfo::Type);
 	DECLARE_MULTICAST_DELEGATE_TwoParams(FOnArchetypesSelected, TConstArrayView<TSharedPtr<FMassDebuggerArchetypeData>>, ESelectInfo::Type);
 
+	FMassDebuggerModel();
+	~FMassDebuggerModel();
+
 	void SetEnvironment(const TSharedPtr<FMassDebuggerEnvironment>& Item);
 
 	void RefreshAll();
@@ -163,7 +166,7 @@ struct FMassDebuggerModel
 	void SelectProcessors(TArrayView<TSharedPtr<FMassDebuggerProcessorData>> Processors, ESelectInfo::Type SelectInfo);
 	void ClearProcessorSelection();
 
-	void SelectArchetypes(TArray<TSharedPtr<FMassDebuggerArchetypeData>> SelectedArchetypes, ESelectInfo::Type SelectInfo);
+	void SelectArchetypes(TArrayView<TSharedPtr<FMassDebuggerArchetypeData>> SelectedArchetypes, ESelectInfo::Type SelectInfo);
 	void ClearArchetypeSelection();
 
 	bool IsCurrentEnvironment(const FMassDebuggerEnvironment& InEnvironment) const { return Environment && *Environment.Get() == InEnvironment; }
@@ -188,6 +191,8 @@ protected:
 	void ResetSelectedArchetypes();
 	void ResetSelectedProcessors();
 
+	void OnEntitySelected(const FMassEntityManager& EntityManager, const FMassEntityHandle EntityHandle);
+
 public:
 	FOnRefresh OnRefreshDelegate;
 	FOnProcessorsSelected OnProcessorsSelectedDelegate;
@@ -199,7 +204,8 @@ public:
 	TArray<TSharedPtr<FMassDebuggerProcessorData>> CachedProcessors;
 	TArray<TSharedPtr<FMassDebuggerProcessorData>> SelectedProcessors;
 
-	TArray<TSharedPtr<FMassDebuggerArchetypeData>> CachedArchetypes;
+	TArray<TSharedPtr<FMassDebuggerArchetypeData>> CachedAllArchetypes;
+	TArray<TSharedPtr<FMassDebuggerArchetypeData>> CachedArchetypeRepresentatives;
 	TArray<TSharedPtr<FMassDebuggerArchetypeData>> SelectedArchetypes;
 
 	TArray<TSharedPtr<FMassDebuggerProcessingGraph>> CachedProcessingGraphs;
@@ -209,5 +215,7 @@ public:
 	TArray<TArray<float>> ArchetypeDistances;
 
 	FString EnvironmentDisplayName;
+
+	FDelegateHandle OnEntitySelectedHandle;
 };
 
diff --git a/Engine/Plugins/Runtime/MassEntity/Source/MassEntityDebugger/Private/SMassArchetypesView.cpp b/Engine/Plugins/Runtime/MassEntity/Source/MassEntityDebugger/Private/SMassArchetypesView.cpp
--- a/Engine/Plugins/Runtime/MassEntity/Source/MassEntityDebugger/Private/SMassArchetypesView.cpp
+++ b/Engine/Plugins/Runtime/MassEntity/Source/MassEntityDebugger/Private/SMassArchetypesView.cpp
@@ -199,7 +199,7 @@ void SMassArchetypesView::Construct(const FArguments& InArgs, TSharedRef<FMassDe
 					{
 						int32 ActiveCount = 0;
 						int32 TotalCount = 0;
-						for (const TSharedPtr<FMassDebuggerArchetypeData>& ArchetypeData : DebuggerModel->CachedArchetypes)
+						for (const TSharedPtr<FMassDebuggerArchetypeData>& ArchetypeData : DebuggerModel->CachedArchetypeRepresentatives)
 						{
 							TotalCount++;
 							if (ArchetypeData->ArchetypeStats.EntitiesCount > 0)
@@ -307,7 +307,7 @@ void SMassArchetypesView::OnRefresh()
 	
 	if (DebuggerModel)
 	{
-		TreeViewSource = DebuggerModel->CachedArchetypes;
+		TreeViewSource = DebuggerModel->CachedArchetypeRepresentatives;
 
 		for (const TSharedPtr<FMassDebuggerArchetypeData>& ArchetypeData : TreeViewSource)
 		{
