Fixed Cooked Editor compilation issue caused by using debug-only MassDebugger functions.

diff --git a/Engine/Plugins/Runtime/MassEntity/Source/MassEntityDebugger/Private/MassDebuggerModel.cpp b/Engine/Plugins/Runtime/MassEntity/Source/MassEntityDebugger/Private/MassDebuggerModel.cpp
--- a/Engine/Plugins/Runtime/MassEntity/Source/MassEntityDebugger/Private/MassDebuggerModel.cpp
+++ b/Engine/Plugins/Runtime/MassEntity/Source/MassEntityDebugger/Private/MassDebuggerModel.cpp
@@ -630,8 +630,12 @@ void FMassDebuggerModel::OnEntitySelected(const FMassEntityManager& EntityManage
 	{
 		return;
 	}
-	
-	uint32 ArchetypeHash = FMassDebugger::GetArchetypeComposition(ArchetypeHandle).CalculateHash();
+
+#if WITH_MASSENTITY_DEBUG
+	const uint32 ArchetypeHash = FMassDebugger::GetArchetypeComposition(ArchetypeHandle).CalculateHash();
+#else
+	const uint32 ArchetypeHash = 0;
+#endif // WITH_MASSENTITY_DEBUG
 	TSharedPtr<FMassDebuggerArchetypeData>* DebuggerArchetypeData = CachedAllArchetypes.FindByPredicate([ArchetypeHash](const TSharedPtr<FMassDebuggerArchetypeData>& Element)
 		{
 			return Element.IsValid() && Element->CompositionHash == ArchetypeHash;
