Cleaned up UMassCompositeProcessor.ProcessingFlatGraph bits a bit and extracted the processing graph creation code into a separate funciton.

diff --git a/Engine/Plugins/Runtime/MassEntity/Source/MassEntity/Private/MassDebugger.cpp b/Engine/Plugins/Runtime/MassEntity/Source/MassEntity/Private/MassDebugger.cpp
--- a/Engine/Plugins/Runtime/MassEntity/Source/MassEntity/Private/MassDebugger.cpp
+++ b/Engine/Plugins/Runtime/MassEntity/Source/MassEntity/Private/MassDebugger.cpp
@@ -271,7 +271,7 @@ const TConstArrayView<FName> FMassDebugger::GetArchetypeDebugNames(const FMassAr
 
 TConstArrayView<struct UMassCompositeProcessor::FDependencyNode> FMassDebugger::GetProcessingGraph(const UMassCompositeProcessor& GraphOwner)
 {
-	return GraphOwner.ProcessingFlatGraph;
+	return GraphOwner.FlatProcessingGraph;
 }
 
 FString FMassDebugger::GetRequirementsDescription(const FMassFragmentRequirements& Requirements)
diff --git a/Engine/Plugins/Runtime/MassEntity/Source/MassEntity/Private/MassProcessor.cpp b/Engine/Plugins/Runtime/MassEntity/Source/MassEntity/Private/MassProcessor.cpp
--- a/Engine/Plugins/Runtime/MassEntity/Source/MassEntity/Private/MassProcessor.cpp
+++ b/Engine/Plugins/Runtime/MassEntity/Source/MassEntity/Private/MassProcessor.cpp
@@ -256,9 +256,9 @@ void UMassCompositeProcessor::ConfigureQueries()
 FGraphEventRef UMassCompositeProcessor::DispatchProcessorTasks(const TSharedPtr<FMassEntityManager>& EntityManager, FMassExecutionContext& ExecutionContext, const FGraphEventArray& InPrerequisites)
 {
 	FGraphEventArray Events;
-	Events.Reserve(ProcessingFlatGraph.Num());
+	Events.Reserve(FlatProcessingGraph.Num());
 		
-	for (FDependencyNode& ProcessingNode : ProcessingFlatGraph)
+	for (FDependencyNode& ProcessingNode : FlatProcessingGraph)
 	{
 		FGraphEventArray Prerequisites;
 		for (const int32 DependencyIndex : ProcessingNode.Dependencies)
@@ -277,13 +277,13 @@ FGraphEventRef UMassCompositeProcessor::DispatchProcessorTasks(const TSharedPtr<
 #if WITH_MASSENTITY_DEBUG
 	if (UE::Mass::Debug::bLogProcessingGraph)
 	{
-		for (int i = 0; i < ProcessingFlatGraph.Num(); ++i)
+		for (int i = 0; i < FlatProcessingGraph.Num(); ++i)
 		{
-			FDependencyNode& ProcessingNode = ProcessingFlatGraph[i];
+			FDependencyNode& ProcessingNode = FlatProcessingGraph[i];
 			FString DependenciesDesc;
 			for (const int32 DependencyIndex : ProcessingNode.Dependencies)
 			{
-				DependenciesDesc += FString::Printf(TEXT("%s, "), *ProcessingFlatGraph[DependencyIndex].Name.ToString());
+				DependenciesDesc += FString::Printf(TEXT("%s, "), *FlatProcessingGraph[DependencyIndex].Name.ToString());
 			}
 
 			check(ProcessingNode.Processor);
@@ -341,18 +341,25 @@ void UMassCompositeProcessor::SetProcessors(TArrayView<UMassProcessor*> InProces
 
 	Populate(SortedProcessorsAndGroups);
 
+	BuildFlatProcessingGraph(SortedProcessorsAndGroups);
+}
+
+void UMassCompositeProcessor::BuildFlatProcessingGraph(TConstArrayView<FProcessorDependencySolver::FOrderInfo> SortedProcessorsAndGroups)
+{
+	TRACE_CPUPROFILER_EVENT_SCOPE(Mass_BuildFlatProcessingGraph);
+
 	// this part is creating an ordered, flat list of processors that can be executed in sequence
 	// with subsequent task only depending on the elements prior on the list
 	TMap<FName, int32> NameToDependencyIndex;
 	NameToDependencyIndex.Reserve(SortedProcessorsAndGroups.Num());
 	TArray<int32> SuperGroupDependency;
-	for (FProcessorDependencySolver::FOrderInfo& Element : SortedProcessorsAndGroups)
+	for (const FProcessorDependencySolver::FOrderInfo& Element : SortedProcessorsAndGroups)
 	{
-		NameToDependencyIndex.Add(Element.Name, ProcessingFlatGraph.Num());
+		NameToDependencyIndex.Add(Element.Name, FlatProcessingGraph.Num());
 
 		// we don't expect to get any "group" nodes here. If it happens it indicates a bug in dependency solving
 		checkSlow(Element.Processor);
-		FDependencyNode& Node = ProcessingFlatGraph.Add_GetRef({ Element.Name, Element.Processor });
+		FDependencyNode& Node = FlatProcessingGraph.Add_GetRef({ Element.Name, Element.Processor });
 		Node.Dependencies.Reserve(Element.Dependencies.Num());
 		for (FName DependencyName : Element.Dependencies)
 		{
@@ -366,11 +373,11 @@ void UMassCompositeProcessor::SetProcessors(TArrayView<UMassProcessor*> InProces
 	UE_LOG(LogMass, Log, TEXT("%s flat processing graph:"), *GroupName.ToString());
 
 	TArray<int32> Offset;
-	Offset.Reserve(ProcessingFlatGraph.Num());
+	Offset.Reserve(FlatProcessingGraph.Num());
 
-	for (int i = 0; i < ProcessingFlatGraph.Num(); ++i)
+	for (int i = 0; i < FlatProcessingGraph.Num(); ++i)
 	{
-		FDependencyNode& ProcessingNode = ProcessingFlatGraph[i];
+		FDependencyNode& ProcessingNode = FlatProcessingGraph[i];
 		FString DependenciesDesc;
 		int32 MyOffset = 0;
 		for (const int32 DependencyIndex : ProcessingNode.Dependencies)
diff --git a/Engine/Plugins/Runtime/MassEntity/Source/MassEntity/Public/MassProcessor.h b/Engine/Plugins/Runtime/MassEntity/Source/MassEntity/Public/MassProcessor.h
--- a/Engine/Plugins/Runtime/MassEntity/Source/MassEntity/Public/MassProcessor.h
+++ b/Engine/Plugins/Runtime/MassEntity/Source/MassEntity/Public/MassProcessor.h
@@ -173,7 +173,8 @@ public:
 	FName GetGroupName() const { return GroupName; }
 
 	virtual void CopyAndSort(const FMassProcessingPhaseConfig& PhaseConfig, const FString& DependencyGraphFileName = FString());
-	void SetProcessors(TArrayView<UMassProcessor*> InProcessorInstances, const FString& DependencyGraphFileName = FString());
+	virtual void SetProcessors(TArrayView<UMassProcessor*> InProcessorInstances, const FString& DependencyGraphFileName = FString());
+	virtual void BuildFlatProcessingGraph(TConstArrayView<FProcessorDependencySolver::FOrderInfo> SortedProcessorsAndGroups);
 
 	/** adds SubProcessor to an appropriately named group. If RequestedGroupName == None then SubProcessor
 	 *  will be added directly to ChildPipeline. If not then the indicated group will be searched for in ChildPipeline 
@@ -211,7 +212,7 @@ protected:
 	UPROPERTY()
 	FName GroupName;
 
-	TArray<FDependencyNode> ProcessingFlatGraph;
+	TArray<FDependencyNode> FlatProcessingGraph;
 
 	struct FProcessorCompletion
 	{
