Renamed FProcessorDependencySolver::FOrderInfo to FMassProcessorOrderInfo and moved it to MassProcessingTypes.h to avoid including MassProcessorDependencySolver.h in MassProcessor.h for that one type.

diff --git a/Engine/Plugins/Runtime/MassEntity/Source/MassEntity/Private/MassProcessor.cpp b/Engine/Plugins/Runtime/MassEntity/Source/MassEntity/Private/MassProcessor.cpp
--- a/Engine/Plugins/Runtime/MassEntity/Source/MassEntity/Private/MassProcessor.cpp
+++ b/Engine/Plugins/Runtime/MassEntity/Source/MassEntity/Private/MassProcessor.cpp
@@ -336,7 +336,7 @@ void UMassCompositeProcessor::SetProcessors(TArrayView<UMassProcessor*> InProces
 {
 	// figure out dependencies
 	FProcessorDependencySolver Solver(InProcessorInstances, GroupName, DependencyGraphFileName);
-	TArray<FProcessorDependencySolver::FOrderInfo> SortedProcessorsAndGroups;
+	TArray<FMassProcessorOrderInfo> SortedProcessorsAndGroups;
 	Solver.ResolveDependencies(SortedProcessorsAndGroups);
 
 	Populate(SortedProcessorsAndGroups);
@@ -344,7 +344,7 @@ void UMassCompositeProcessor::SetProcessors(TArrayView<UMassProcessor*> InProces
 	BuildFlatProcessingGraph(SortedProcessorsAndGroups);
 }
 
-void UMassCompositeProcessor::BuildFlatProcessingGraph(TConstArrayView<FProcessorDependencySolver::FOrderInfo> SortedProcessorsAndGroups)
+void UMassCompositeProcessor::BuildFlatProcessingGraph(TConstArrayView<FMassProcessorOrderInfo> SortedProcessorsAndGroups)
 {
 	TRACE_CPUPROFILER_EVENT_SCOPE(Mass_BuildFlatProcessingGraph);
 
@@ -353,7 +353,7 @@ void UMassCompositeProcessor::BuildFlatProcessingGraph(TConstArrayView<FProcesso
 	TMap<FName, int32> NameToDependencyIndex;
 	NameToDependencyIndex.Reserve(SortedProcessorsAndGroups.Num());
 	TArray<int32> SuperGroupDependency;
-	for (const FProcessorDependencySolver::FOrderInfo& Element : SortedProcessorsAndGroups)
+	for (const FMassProcessorOrderInfo& Element : SortedProcessorsAndGroups)
 	{
 		NameToDependencyIndex.Add(Element.Name, FlatProcessingGraph.Num());
 
@@ -397,15 +397,15 @@ void UMassCompositeProcessor::BuildFlatProcessingGraph(TConstArrayView<FProcesso
 #endif // WITH_MASSENTITY_DEBUG
 }
 
-void UMassCompositeProcessor::Populate(TConstArrayView<FProcessorDependencySolver::FOrderInfo> OrderedProcessors)
+void UMassCompositeProcessor::Populate(TConstArrayView<FMassProcessorOrderInfo> OrderedProcessors)
 {
 	ChildPipeline.Processors.Reset();
 
 	const FMassProcessingPhaseConfig& PhaseConfig = GET_MASS_CONFIG_VALUE(GetProcessingPhaseConfig(ProcessingPhase));
 
-	for (const FProcessorDependencySolver::FOrderInfo& ProcessorInfo : OrderedProcessors)
+	for (const FMassProcessorOrderInfo& ProcessorInfo : OrderedProcessors)
 	{
-		if (ensureMsgf(ProcessorInfo.NodeType == EDependencyNodeType::Processor, TEXT("Encountered unexpected EDependencyNodeType while populating %s"), *GetGroupName().ToString()))
+		if (ensureMsgf(ProcessorInfo.NodeType == FMassProcessorOrderInfo::EDependencyNodeType::Processor, TEXT("Encountered unexpected FMassProcessorOrderInfo::EDependencyNodeType while populating %s"), *GetGroupName().ToString()))
 		{
 			checkSlow(ProcessorInfo.Processor);
 			
diff --git a/Engine/Plugins/Runtime/MassEntity/Source/MassEntity/Private/MassProcessorDependencySolver.cpp b/Engine/Plugins/Runtime/MassEntity/Source/MassEntity/Private/MassProcessorDependencySolver.cpp
--- a/Engine/Plugins/Runtime/MassEntity/Source/MassEntity/Private/MassProcessorDependencySolver.cpp
+++ b/Engine/Plugins/Runtime/MassEntity/Source/MassEntity/Private/MassProcessorDependencySolver.cpp
@@ -843,7 +843,7 @@ void FProcessorDependencySolver::DumpGraph(FArchive& LogFile) const
 
 #endif // 0; disabled the graph building for now, leaving the old code here for reference
 
-void FProcessorDependencySolver::Solve(TArray<FProcessorDependencySolver::FOrderInfo>& OutResult)
+void FProcessorDependencySolver::Solve(TArray<FMassProcessorOrderInfo>& OutResult)
 {
 	using UE::Mass::Private::NameViewToString;
 
@@ -929,12 +929,12 @@ void FProcessorDependencySolver::Solve(TArray<FProcessorDependencySolver::FOrder
 		// at this point we expect SortedNodeIndices to only point to regular processors (i.e. no groups)
 		if (ensure(AllNodes[NodeIndex].Processor != nullptr))
 		{
-			OutResult.Add({ AllNodes[NodeIndex].Name, AllNodes[NodeIndex].Processor, EDependencyNodeType::Processor, DependencyNames });
+			OutResult.Add({ AllNodes[NodeIndex].Name, AllNodes[NodeIndex].Processor, FMassProcessorOrderInfo::EDependencyNodeType::Processor, DependencyNames });
 		}
 	}
 }
 
-void FProcessorDependencySolver::ResolveDependencies(TArray<FProcessorDependencySolver::FOrderInfo>& OutResult)
+void FProcessorDependencySolver::ResolveDependencies(TArray<FMassProcessorOrderInfo>& OutResult)
 {
 	if (Processors.Num() == 0)
 	{
@@ -978,7 +978,7 @@ void FProcessorDependencySolver::ResolveDependencies(TArray<FProcessorDependency
 	Solve(OutResult);
 
 	UE_LOG(LogMass, Verbose, TEXT("Dependency order:"));
-	for (const FProcessorDependencySolver::FOrderInfo& Info : OutResult)
+	for (const FMassProcessorOrderInfo& Info : OutResult)
 	{
 		UE_LOG(LogMass, Verbose, TEXT("\t%s"), *Info.Name.ToString());
 	}
diff --git a/Engine/Plugins/Runtime/MassEntity/Source/MassEntity/Public/MassProcessingTypes.h b/Engine/Plugins/Runtime/MassEntity/Source/MassEntity/Public/MassProcessingTypes.h
--- a/Engine/Plugins/Runtime/MassEntity/Source/MassEntity/Public/MassProcessingTypes.h
+++ b/Engine/Plugins/Runtime/MassEntity/Source/MassEntity/Public/MassProcessingTypes.h
@@ -125,3 +125,19 @@ enum class EMassProcessingPhase : uint8
 	FrameEnd,
 	MAX,
 };
+
+struct FMassProcessorOrderInfo
+{
+	enum class EDependencyNodeType : uint8
+	{
+		Invalid,
+		Processor,
+		GroupStart,
+		GroupEnd
+	};
+
+	FName Name = TEXT("");
+	UMassProcessor* Processor = nullptr;
+	EDependencyNodeType NodeType = EDependencyNodeType::Invalid;
+	TArray<FName> Dependencies;
+};
diff --git a/Engine/Plugins/Runtime/MassEntity/Source/MassEntity/Public/MassProcessor.h b/Engine/Plugins/Runtime/MassEntity/Source/MassEntity/Public/MassProcessor.h
--- a/Engine/Plugins/Runtime/MassEntity/Source/MassEntity/Public/MassProcessor.h
+++ b/Engine/Plugins/Runtime/MassEntity/Source/MassEntity/Public/MassProcessor.h
@@ -7,7 +7,6 @@
 #include "MassProcessingTypes.h"
 #include "Async/TaskGraphInterfaces.h"
 #include "MassCommandBuffer.h"
-#include "MassProcessorDependencySolver.h"
 #include "MassProcessor.generated.h"
 
 
@@ -174,7 +173,7 @@ public:
 
 	virtual void CopyAndSort(const FMassProcessingPhaseConfig& PhaseConfig, const FString& DependencyGraphFileName = FString());
 	virtual void SetProcessors(TArrayView<UMassProcessor*> InProcessorInstances, const FString& DependencyGraphFileName = FString());
-	virtual void BuildFlatProcessingGraph(TConstArrayView<FProcessorDependencySolver::FOrderInfo> SortedProcessorsAndGroups);
+	virtual void BuildFlatProcessingGraph(TConstArrayView<FMassProcessorOrderInfo> SortedProcessorsAndGroups);
 
 	/** adds SubProcessor to an appropriately named group. If RequestedGroupName == None then SubProcessor
 	 *  will be added directly to ChildPipeline. If not then the indicated group will be searched for in ChildPipeline 
@@ -194,7 +193,7 @@ protected:
 	/**
 	 *  Adds processors in OrderedProcessors to ChildPipeline and builds flat processing graph that's being used in multithreaded mode.
 	 */
-	void Populate(TConstArrayView<FProcessorDependencySolver::FOrderInfo> OrderedProcessors);
+	void Populate(TConstArrayView<FMassProcessorOrderInfo> OrderedProcessors);
 
 	/** RequestedGroupName can indicate a multi-level group name, like so: A.B.C
 	 *  We need to extract the highest-level group name ('A' in the example), and see if it already exists. 
diff --git a/Engine/Plugins/Runtime/MassEntity/Source/MassEntity/Public/MassProcessorDependencySolver.h b/Engine/Plugins/Runtime/MassEntity/Source/MassEntity/Public/MassProcessorDependencySolver.h
--- a/Engine/Plugins/Runtime/MassEntity/Source/MassEntity/Public/MassProcessorDependencySolver.h
+++ b/Engine/Plugins/Runtime/MassEntity/Source/MassEntity/Public/MassProcessorDependencySolver.h
@@ -9,14 +9,6 @@
 
 class UMassProcessor;
 
-enum class EDependencyNodeType : uint8
-{
-	Invalid,
-	Processor,
-	GroupStart,
-	GroupEnd
-};
-
 namespace EMassAccessOperation
 {
 	constexpr uint32 Read = 0;
@@ -119,16 +111,9 @@ private:
 	};
 
 public:
-	struct FOrderInfo
-	{
-		FName Name = TEXT("");
-		UMassProcessor* Processor = nullptr;
-		EDependencyNodeType NodeType = EDependencyNodeType::Invalid;
-		TArray<FName> Dependencies;
-	};
 
 	MASSENTITY_API FProcessorDependencySolver(TArrayView<UMassProcessor*> InProcessors, const FName Name, const FString& InDependencyGraphFileName = FString());
-	MASSENTITY_API void ResolveDependencies(TArray<FOrderInfo>& OutResult);
+	MASSENTITY_API void ResolveDependencies(TArray<FMassProcessorOrderInfo>& OutResult);
 
 	MASSENTITY_API static void CreateSubGroupNames(FName InGroupName, TArray<FString>& SubGroupNames);
 
@@ -145,7 +130,7 @@ protected:
 	
 	void CreateNodes(UMassProcessor& Processor);
 	void BuildDependencies();
-	void Solve(TArray<FProcessorDependencySolver::FOrderInfo>& OutResult);
+	void Solve(TArray<FMassProcessorOrderInfo>& OutResult);
 	void LogNode(const FNode& Node, int Indent = 0);
 	
 	// @todo due to fundamental change to how nodes are organized the graph generation needs reimplementation
diff --git a/Engine/Plugins/Runtime/MassEntity/Source/MassEntityTestSuite/Private/MassDependencyTest.cpp b/Engine/Plugins/Runtime/MassEntity/Source/MassEntityTestSuite/Private/MassDependencyTest.cpp
--- a/Engine/Plugins/Runtime/MassEntity/Source/MassEntityTestSuite/Private/MassDependencyTest.cpp
+++ b/Engine/Plugins/Runtime/MassEntity/Source/MassEntityTestSuite/Private/MassDependencyTest.cpp
@@ -22,7 +22,7 @@ static FName GetProcessorName()
 struct FDependencySolverBase : FAITestBase
 {
 	TArray<UMassTestProcessorBase*> Processors;
-	TArray<FProcessorDependencySolver::FOrderInfo> Result;
+	TArray<FMassProcessorOrderInfo> Result;
 	
 	virtual bool SetUp() override
 	{
@@ -153,7 +153,7 @@ struct FDeepGroup : FDependencySolverBase
 		// dump all the group information from the Result collection for easier ordering testing
 		for (int32 i = 0; i < Result.Num(); ++i)
 		{
-			if (Result[i].NodeType != EDependencyNodeType::Processor)
+			if (Result[i].NodeType != FMassProcessorOrderInfo::EDependencyNodeType::Processor)
 			{
 				Result.RemoveAt(i--, 1, /*bAllowShrinking=*/false);
 			}
@@ -212,7 +212,7 @@ struct FComplexScenario : FDependencySolverBase
 
 		for (int32 i = 0; i < Result.Num(); ++i)
 		{
-			AITEST_EQUAL("We expect only processor nodes in the results", Result[i].NodeType, EDependencyNodeType::Processor);
+			AITEST_EQUAL("We expect only processor nodes in the results", Result[i].NodeType, FMassProcessorOrderInfo::EDependencyNodeType::Processor);
 		}
 
 		AITEST_TRUE("D is the only fully dependency-less processor so should be first", Result[0].Name == GetProcessorName<UMassTestProcessor_D>());		
diff --git a/Engine/Plugins/Runtime/MassEntity/Source/MassEntityTestSuite/Private/MassMultiThreadTest.cpp b/Engine/Plugins/Runtime/MassEntity/Source/MassEntityTestSuite/Private/MassMultiThreadTest.cpp
--- a/Engine/Plugins/Runtime/MassEntity/Source/MassEntityTestSuite/Private/MassMultiThreadTest.cpp
+++ b/Engine/Plugins/Runtime/MassEntity/Source/MassEntityTestSuite/Private/MassMultiThreadTest.cpp
@@ -28,7 +28,7 @@ struct FMTTestBase : FEntityTestBase
 	using Super = FExecutionTestBase;
 	UMassCompositeProcessor* CompositeProcessor = nullptr;
 	TArray<UMassTestProcessorBase*> Processors;
-	TArray<FProcessorDependencySolver::FOrderInfo> Result;
+	TArray<FMassProcessorOrderInfo> Result;
 	FGraphEventRef FinishEvent;
 	TArray<FMassEntityHandle> Entities;
 
