Fix to MassEntityDebugger Slate Styles.

diff --git a/Engine/Plugins/Runtime/MassEntity/Source/MassEntityDebugger/Private/MassDebuggerStyle.cpp b/Engine/Plugins/Runtime/MassEntity/Source/MassEntityDebugger/Private/MassDebuggerStyle.cpp
--- a/Engine/Plugins/Runtime/MassEntity/Source/MassEntityDebugger/Private/MassDebuggerStyle.cpp
+++ b/Engine/Plugins/Runtime/MassEntity/Source/MassEntityDebugger/Private/MassDebuggerStyle.cpp
@@ -45,11 +45,11 @@ public:
 			.SetColorAndOpacity(FLinearColor(230.0f / 255.0f, 230.0f / 255.0f, 230.0f / 255.0f, 0.9f));
 		Set("MassDebug.Font.Large", StateTitle);
 
-		Set("MassDebug.Fragment", new BOX_BRUSH("Common/LightGroupBorder", FMargin(4.0f / 16.0f), FLinearColor(.2f, .2f, .9f)));
-		Set("MassDebug.Fragment.Added", new BOX_BRUSH("Common/LightGroupBorder", FMargin(4.0f / 16.0f), FLinearColor(.2f, .2f, .9f)));
-		Set("MassDebug.Fragment.Removed", new BOX_BRUSH("Common/LightGroupBorder", FMargin(4.0f / 16.0f), FLinearColor(.9f, .2f, .2f)));
-		Set("MassDebug.Fragment.ReadOnly", new BOX_BRUSH("Common/LightGroupBorder", FMargin(4.0f / 16.0f), FLinearColor(.5f, .5f, .5f)));
-		Set("MassDebug.Fragment.ReadWrite", new BOX_BRUSH("Common/LightGroupBorder", FMargin(4.0f / 16.0f), FLinearColor(0.f, .7f, 0.f)));
+		Set("MassDebug.Fragment", new CORE_BOX_BRUSH("Common/LightGroupBorder", FMargin(4.0f / 16.0f), FLinearColor(.2f, .2f, .9f)));
+		Set("MassDebug.Fragment.Added", new CORE_BOX_BRUSH("Common/LightGroupBorder", FMargin(4.0f / 16.0f), FLinearColor(.2f, .2f, .9f)));
+		Set("MassDebug.Fragment.Removed", new CORE_BOX_BRUSH("Common/LightGroupBorder", FMargin(4.0f / 16.0f), FLinearColor(.9f, .2f, .2f)));
+		Set("MassDebug.Fragment.ReadOnly", new CORE_BOX_BRUSH("Common/LightGroupBorder", FMargin(4.0f / 16.0f), FLinearColor(.5f, .5f, .5f)));
+		Set("MassDebug.Fragment.ReadWrite", new CORE_BOX_BRUSH("Common/LightGroupBorder", FMargin(4.0f / 16.0f), FLinearColor(0.f, .7f, 0.f)));
 
 		Set("MassDebug.Label.Background", new FSlateRoundedBoxBrush(FStyleColors::Foreground, 4.0f));
 
