Removed wrong TRACE_CPUPROFILER_EVENT_SCOPE from FMassExecutor::TriggerParallelTasks

diff --git a/Engine/Plugins/Runtime/MassEntity/Source/MassEntity/Private/MassExecutor.cpp b/Engine/Plugins/Runtime/MassEntity/Source/MassEntity/Private/MassExecutor.cpp
--- a/Engine/Plugins/Runtime/MassEntity/Source/MassEntity/Private/MassExecutor.cpp
+++ b/Engine/Plugins/Runtime/MassEntity/Source/MassEntity/Private/MassExecutor.cpp
@@ -177,8 +177,6 @@ private:
 
 FGraphEventRef TriggerParallelTasks(UMassProcessor& Processor, FMassProcessingContext& ProcessingContext, TFunction<void()> OnDoneNotification)
 {
-	TRACE_CPUPROFILER_EVENT_SCOPE(RunProcessorsView);
-
 	if (!ProcessingContext.EntityManager)
 	{
 		UE_LOG(LogMass, Error, TEXT("%s ProcessingContext.EntityManager is null. Baling out."), ANSI_TO_TCHAR(__FUNCTION__));
