Added a multicast delegate to FMassEntityManager that gets triggered upon new Archetype creation.

diff --git a/Engine/Plugins/Runtime/MassEntity/Source/MassEntity/Private/MassEntityManager.cpp b/Engine/Plugins/Runtime/MassEntity/Source/MassEntity/Private/MassEntityManager.cpp
--- a/Engine/Plugins/Runtime/MassEntity/Source/MassEntity/Private/MassEntityManager.cpp
+++ b/Engine/Plugins/Runtime/MassEntity/Source/MassEntity/Private/MassEntityManager.cpp
@@ -260,6 +260,8 @@ FMassArchetypeHandle FMassEntityManager::CreateArchetype(const FMassArchetypeCom
 			checkSlow(FragmentConfig.FragmentType)
 			FragmentTypeToArchetypeMap.FindOrAdd(FragmentConfig.FragmentType).Add(ArchetypeDataPtr);
 		}
+
+		OnNewArchetypeEvent.Broadcast(FMassArchetypeHandle(ArchetypeDataPtr));
 	}
 
 	return FMassArchetypeHelper::ArchetypeHandleFromData(ArchetypeDataPtr);
@@ -314,6 +316,8 @@ FMassArchetypeHandle FMassEntityManager::InternalCreateSimilarArchetype(const FM
 			checkSlow(FragmentConfig.FragmentType)
 			FragmentTypeToArchetypeMap.FindOrAdd(FragmentConfig.FragmentType).Add(ArchetypeDataPtr);
 		}
+
+		OnNewArchetypeEvent.Broadcast(FMassArchetypeHandle(ArchetypeDataPtr));
 	}
 
 	return FMassArchetypeHelper::ArchetypeHandleFromData(ArchetypeDataPtr);
diff --git a/Engine/Plugins/Runtime/MassEntity/Source/MassEntity/Public/MassEntityManager.h b/Engine/Plugins/Runtime/MassEntity/Source/MassEntity/Public/MassEntityManager.h
--- a/Engine/Plugins/Runtime/MassEntity/Source/MassEntity/Public/MassEntityManager.h
+++ b/Engine/Plugins/Runtime/MassEntity/Source/MassEntity/Public/MassEntityManager.h
@@ -43,6 +43,8 @@ struct MASSENTITY_API FMassEntityManager : public TSharedFromThis<FMassEntityMan
 	friend FMassEntityQuery;
 	friend FMassDebugger;
 
+	DECLARE_MULTICAST_DELEGATE_OneParam(FOnNewArchetypeDelegate, const FMassArchetypeHandle&);
+
 private:
 	// Index 0 is reserved so we can treat that index as an invalid entity handle
 	constexpr static int32 NumReservedEntities = 1;
@@ -379,6 +381,7 @@ public:
 
 	FMassObserverManager& GetObserverManager() { return ObserverManager; }
 
+	FOnNewArchetypeDelegate& GetOnNewArchetypeEvent() { return OnNewArchetypeEvent; }
 	/** 
 	 * Fetches the world associated with the Owner. 
 	 * @note that it's ok for a given EntityManager to not have an owner or the owner not being part of a UWorld, depending on the use case
@@ -471,5 +474,7 @@ private:
 
 	TWeakObjectPtr<UObject> Owner;
 
+	FOnNewArchetypeDelegate OnNewArchetypeEvent;
+
 	bool bInitialized = false;
 };
