Two fold change to Mass processor dependency solver:
1. Before the actual dependency solving the processor collection is sorted in such a way that lets us first process the elements that are blocking the largest number of other nodes.
2. When storing and comparing information regarding used resources we also compare sets of matching archetypes to fine-tune resource access. Depending on the parameters the archetypes can be created specially for the purpose of dependency solving (called "virtual archetypes" in code) or queried from the given EntityManager instance.

The change results in shortest "longest sequence" yet which in turn significantly improved processor distribution across threads when running mass in multithreaded mode.

The new approach is utilized by the PhaseManager to rebuild Phases' processing graph whenever a new archetype is created. There is some testing performed to make sure the newly created archetypes are relevant for the current processing graph and if there's any need to rebuild it (by testing the processors that have been pruned the last time around the algorithm has been run).
This part resulted in cleaning up the processing setting and sorting code, dropping UMassCompositeProcessor::CopyAndSort and updating MassEntitySettings code that was using it.

diff --git a/Engine/Plugins/Runtime/MassEntity/Source/MassEntity/Private/MassEntitySettings.cpp b/Engine/Plugins/Runtime/MassEntity/Source/MassEntity/Private/MassEntitySettings.cpp
--- a/Engine/Plugins/Runtime/MassEntity/Source/MassEntity/Private/MassEntitySettings.cpp
+++ b/Engine/Plugins/Runtime/MassEntity/Source/MassEntity/Private/MassEntitySettings.cpp
@@ -7,6 +7,7 @@
 #include "Misc/CoreDelegates.h"
 #if WITH_EDITOR
 #include "ObjectEditorUtils.h"
+#include "CoreGlobals.h"
 #endif // WITH_EDITOR
 
 //----------------------------------------------------------------------//
@@ -44,19 +45,39 @@ void UMassEntitySettings::BuildProcessorListAndPhases()
 void UMassEntitySettings::BuildPhases()
 {
 #if WITH_EDITOR
-	for (int i = 0; i < int(EMassProcessingPhase::MAX); ++i)
+	if (GIsEditor)
 	{
-		FMassProcessingPhaseConfig& PhaseConfig = ProcessingPhasesConfig[i];
-		PhaseConfig.PhaseProcessor = NewObject<UMassCompositeProcessor>(this, PhaseConfig.PhaseGroupClass
-			, *FString::Printf(TEXT("ProcessingPhase_%s"), *PhaseConfig.PhaseName.ToString()));
-		PhaseConfig.PhaseProcessor->SetGroupName(PhaseConfig.PhaseName);
-		PhaseConfig.PhaseProcessor->SetProcessingPhase(EMassProcessingPhase(i));
-		const FString PhaseDumpDependencyGraphFileName = !DumpDependencyGraphFileName.IsEmpty() ? DumpDependencyGraphFileName + TEXT("_") + PhaseConfig.PhaseName.ToString() : FString();
-		PhaseConfig.PhaseProcessor->CopyAndSort(PhaseConfig, PhaseDumpDependencyGraphFileName);
-		
-		FStringOutputDevice Ar;
-		PhaseConfig.PhaseProcessor->DebugOutputDescription(Ar);
-		PhaseConfig.Description = FText::FromString(Ar);
+		for (int i = 0; i < int(EMassProcessingPhase::MAX); ++i)
+		{
+			FMassProcessingPhaseConfig& PhaseConfig = ProcessingPhasesConfig[i];
+			PhaseConfig.PhaseProcessor = NewObject<UMassCompositeProcessor>(this, PhaseConfig.PhaseGroupClass
+				, *FString::Printf(TEXT("ProcessingPhase_%s"), *PhaseConfig.PhaseName.ToString()));
+			PhaseConfig.PhaseProcessor->SetGroupName(PhaseConfig.PhaseName);
+			PhaseConfig.PhaseProcessor->SetProcessingPhase(EMassProcessingPhase(i));
+			const FString PhaseDumpDependencyGraphFileName = !DumpDependencyGraphFileName.IsEmpty() ? DumpDependencyGraphFileName + TEXT("_") + PhaseConfig.PhaseName.ToString() : FString();
+
+			FProcessorDependencySolver::FResult Result;
+			Result.DependencyGraphFileName = PhaseDumpDependencyGraphFileName;
+			FMassPhaseProcessorConfigurationHelper Configurator(*PhaseConfig.PhaseProcessor, PhaseConfig, *this);
+			Configurator.bInitializeCreatedProcessors = false;
+			Configurator.bIsGameRuntime = false;
+			Configurator.Configure(/*EntityManager=*/nullptr, &Result);
+
+			if (Result.PrunedProcessorClasses.Num() > 0)
+			{
+				UE_VLOG_UELOG(this, LogMass, Warning, TEXT("Calculating dependencies for phase %s resulted in pruned processors:")
+					, *PhaseConfig.PhaseName.ToString());
+				for (const TSubclassOf<UMassProcessor>& ProcessorClass : Result.PrunedProcessorClasses)
+				{
+					UE_VLOG_UELOG(this, LogMass, Warning, TEXT("\t%s"), *GetNameSafe(ProcessorClass));
+				}
+				UE_VLOG_UELOG(this, LogMass, Warning, TEXT("Make sure EntityQueries owned by these processors got registered or override UMassProcessor::ShouldAllowQueryBasedPruning to return the appropriate result"));
+			}
+
+			FStringOutputDevice Ar;
+			PhaseConfig.PhaseProcessor->DebugOutputDescription(Ar);
+			PhaseConfig.Description = FText::FromString(Ar);
+		}
 	}
 #endif // WITH_EDITOR
 }
diff --git a/Engine/Plugins/Runtime/MassEntity/Source/MassEntity/Private/MassProcessingPhaseManager.cpp b/Engine/Plugins/Runtime/MassEntity/Source/MassEntity/Private/MassProcessingPhaseManager.cpp
--- a/Engine/Plugins/Runtime/MassEntity/Source/MassEntity/Private/MassProcessingPhaseManager.cpp
+++ b/Engine/Plugins/Runtime/MassEntity/Source/MassEntity/Private/MassProcessingPhaseManager.cpp
@@ -133,57 +133,75 @@ void FMassProcessingPhase::Initialize(FMassProcessingPhaseManager& InPhaseManage
 	PhaseProcessor = &InPhaseProcessor;
 }
 
+//----------------------------------------------------------------------//
+// FPhaseProcessorConfigurator
+//----------------------------------------------------------------------//
+void FMassPhaseProcessorConfigurationHelper::Configure(const TSharedPtr<FMassEntityManager>& EntityManager, FProcessorDependencySolver::FResult* OutOptionalResult)
+{
+	FMassRuntimePipeline TmpPipeline;
+	TmpPipeline.CreateFromArray(PhaseConfig.ProcessorCDOs, ProcessorOuter);
+
+	TArray<FMassProcessorOrderInfo> SortedProcessors;
+	FProcessorDependencySolver Solver(TmpPipeline.Processors, bIsGameRuntime);
+
+	Solver.ResolveDependencies(SortedProcessors, EntityManager, OutOptionalResult);
+
+	PhaseProcessor.Populate(SortedProcessors);
+
+	PhaseProcessor.BuildFlatProcessingGraph(SortedProcessors);
+
+	if (bInitializeCreatedProcessors)
+	{
+		PhaseProcessor.Initialize(ProcessorOuter);
+	}
+}
+
 //----------------------------------------------------------------------//
 // FMassProcessingPhaseManager
 //----------------------------------------------------------------------//
-void FMassProcessingPhaseManager::Initialize(UObject& InOwner, TConstArrayView<FMassProcessingPhaseConfig> ProcessingPhasesConfig, const FString& DependencyGraphFileName)
+void FMassProcessingPhaseManager::Initialize(UObject& InOwner, TConstArrayView<FMassProcessingPhaseConfig> InProcessingPhasesConfig, const FString& DependencyGraphFileName)
 {
+#if WITH_EDITOR
+	UWorld* World = InOwner.GetWorld();
+	const bool bIsEditorWorld = (World != 0) && (World->IsEditorWorld() && !World->IsGameWorld());
+#endif // WITH_EDITOR
 	Owner = &InOwner;
+	ProcessingPhasesConfig = InProcessingPhasesConfig;
 
-	for (int i = 0; i < int(EMassProcessingPhase::MAX); ++i)
+	for (int PhaseAsInt = 0; PhaseAsInt < int(EMassProcessingPhase::MAX); ++PhaseAsInt)
 	{		
-		const EMassProcessingPhase Phase = EMassProcessingPhase(i);
+		const EMassProcessingPhase Phase = EMassProcessingPhase(PhaseAsInt);
 
 		UMassCompositeProcessor* PhaseProcessor = NewObject<UMassCompositeProcessor>(&InOwner, UMassCompositeProcessor::StaticClass()
 			, *FString::Printf(TEXT("ProcessingPhase_%s"), *UEnum::GetDisplayValueAsText(Phase).ToString()));
 	
 		check(PhaseProcessor);
-		ProcessingPhases[i].Initialize(*this, Phase, UE::Mass::Private::PhaseToTickingGroup[i], *PhaseProcessor);
+		ProcessingPhases[PhaseAsInt].Initialize(*this, Phase, UE::Mass::Private::PhaseToTickingGroup[PhaseAsInt], *PhaseProcessor);
 
 		REDIRECT_OBJECT_TO_VLOG(PhaseProcessor, &InOwner);
 		PhaseProcessor->SetProcessingPhase(Phase);
 		PhaseProcessor->SetGroupName(FName(FString::Printf(TEXT("%s Group"), *UEnum::GetDisplayValueAsText(Phase).ToString())));
 
-		check(&InOwner);
-		PhaseProcessor->Initialize(InOwner);
-
 #if WITH_MASSENTITY_DEBUG
 		FStringOutputDevice Ar;
 		PhaseProcessor->DebugOutputDescription(Ar);
 		UE_VLOG(&InOwner, LogMass, Log, TEXT("Setting new group processor for phase %s:\n%s"), *UEnum::GetValueAsString(Phase), *Ar);
 #endif // WITH_MASSENTITY_DEBUG
 
-		const FMassProcessingPhaseConfig& PhaseConfig = ProcessingPhasesConfig[i];
-		FString FileName = !DependencyGraphFileName.IsEmpty() ? FString::Printf(TEXT("%s_%s"), *DependencyGraphFileName, *PhaseConfig.PhaseName.ToString()) : FString();
-		PhaseProcessor->CopyAndSort(PhaseConfig, FileName);
-		PhaseProcessor->Initialize(InOwner);
-	}
-
-#if WITH_MASSENTITY_DEBUG
-	// print it all out to vislog
-	UE_VLOG_UELOG(&InOwner, LogMass, Verbose, TEXT("Phases initialization done. Current composition:"));
-
-	FStringOutputDevice OutDescription;
-	for (int i = 0; i < int(EMassProcessingPhase::MAX); ++i)
-	{
-		if (ProcessingPhases[i].PhaseProcessor)
+#if WITH_EDITOR
+		if (bIsEditorWorld)
 		{
-			ProcessingPhases[i].PhaseProcessor->DebugOutputDescription(OutDescription);
-			UE_VLOG_UELOG(&InOwner, LogMass, Verbose, TEXT("--- %s"), *OutDescription);
-			OutDescription.Reset();
+			// populating the phase processor with initial data for editor world so that the default processing graph
+			// can be investigated in the editor without running PIE.
+			// Runtime processing graph are initialized at runtime base on the actual archetypes instantiated at call time.
+			FProcessorDependencySolver::FResult Result;
+			Result.DependencyGraphFileName = DependencyGraphFileName;
+			FMassPhaseProcessorConfigurationHelper Configurator(*PhaseProcessor, ProcessingPhasesConfig[PhaseAsInt], InOwner);
+			Configurator.bIsGameRuntime = false;
+			Configurator.Configure(/*EntityManager=*/nullptr, &Result);
 		}
-	}	
-#endif // WITH_MASSENTITY_DEBUG
+#endif // WITH_EDITOR
+	}
 }
 
 void FMassProcessingPhaseManager::Deinitialize()
@@ -200,8 +218,7 @@ void FMassProcessingPhaseManager::Start(UWorld& World)
 
 	if (ensure(EntitySubsystem))
 	{
-		EntityManager = EntitySubsystem->GetMutableEntityManager().AsShared();
-		EnableTickFunctions(World);
+		Start(EntitySubsystem->GetMutableEntityManager().AsShared());
 	}
 	else
 	{
@@ -214,6 +231,9 @@ void FMassProcessingPhaseManager::Start(const TSharedPtr<FMassEntityManager>& In
 	UWorld* World = InEntityManager->GetWorld();
 	check(World);
 	EntityManager = InEntityManager;
+
+	OnNewArchetypeHandle = EntityManager->GetOnNewArchetypeEvent().AddRaw(this, &FMassProcessingPhaseManager::OnNewArchetype);
+
 	EnableTickFunctions(*World);
 }
 
@@ -260,7 +280,11 @@ void FMassProcessingPhaseManager::EnableTickFunctions(const UWorld& World)
 
 void FMassProcessingPhaseManager::Stop()
 {
-	EntityManager.Reset();
+	if (EntityManager)
+	{
+		EntityManager->GetOnNewArchetypeEvent().Remove(OnNewArchetypeHandle);
+		EntityManager.Reset();
+	}
 	
 	for (FMassProcessingPhase& Phase : ProcessingPhases)
 	{
@@ -283,6 +307,43 @@ void FMassProcessingPhaseManager::OnPhaseStart(const FMassProcessingPhase& Phase
 {
 	ensure(CurrentPhase == EMassProcessingPhase::MAX);
 	CurrentPhase = Phase.Phase;
+
+	const int32 PhaseAsInt = int32(Phase.Phase);
+	if (Owner.IsValid()
+		&& ensure(Phase.Phase != EMassProcessingPhase::MAX)
+		&& ProcessingGraphBuildStates[PhaseAsInt].bNewArchetypes
+		// if not a valid index then we're not able to recalculate dependencies 
+		&& ensure(ProcessingPhasesConfig.IsValidIndex(PhaseAsInt)))
+	{
+		TRACE_CPUPROFILER_EVENT_SCOPE_STR("Mass Rebuild Phase Graph");
+
+		FPhaseGraphBuildState& GraphBuildState = ProcessingGraphBuildStates[PhaseAsInt];
+		if (GraphBuildState.bInitialized == false || FProcessorDependencySolver::IsResultUpToDate(GraphBuildState.LastResult, EntityManager) == false)
+		{
+			UMassCompositeProcessor* PhaseProcessor = ProcessingPhases[PhaseAsInt].PhaseProcessor;
+			check(PhaseProcessor);
+
+			GraphBuildState.LastResult.Reset();
+
+			FMassPhaseProcessorConfigurationHelper Configurator(*PhaseProcessor, ProcessingPhasesConfig[PhaseAsInt], *Owner.Get());
+			Configurator.Configure(EntityManager, &GraphBuildState.LastResult);
+
+			GraphBuildState.bInitialized = true;
+
+#if WITH_MASSENTITY_DEBUG
+			UObject* OwnerPtr = Owner.Get();
+			// print it all out to vislog
+			UE_VLOG_UELOG(OwnerPtr, LogMass, Verbose, TEXT("Phases initialization done. Current composition:"));
+
+			FStringOutputDevice OutDescription;
+			PhaseProcessor->DebugOutputDescription(OutDescription);
+			UE_VLOG_UELOG(OwnerPtr, LogMass, Verbose, TEXT("--- %s"), *OutDescription);
+			OutDescription.Reset();
+#endif // WITH_MASSENTITY_DEBUG
+		}
+
+		ProcessingGraphBuildStates[PhaseAsInt].bNewArchetypes = false;
+	}
 }
 
 void FMassProcessingPhaseManager::OnPhaseEnd(FMassProcessingPhase& Phase)
@@ -309,4 +370,12 @@ FString FMassProcessingPhaseManager::GetName() const
 	return GetNameSafe(Owner.Get()) + TEXT("_MassProcessingPhaseManager");
 }
 
+void FMassProcessingPhaseManager::OnNewArchetype(const FMassArchetypeHandle& NewArchetype)
+{
+	for (FPhaseGraphBuildState& GraphBuildState : ProcessingGraphBuildStates)
+	{
+		GraphBuildState.bNewArchetypes = true;
+	}
+}
+
 #undef LOCTEXT_NAMESPACE
diff --git a/Engine/Plugins/Runtime/MassEntity/Source/MassEntity/Private/MassProcessor.cpp b/Engine/Plugins/Runtime/MassEntity/Source/MassEntity/Private/MassProcessor.cpp
--- a/Engine/Plugins/Runtime/MassEntity/Source/MassEntity/Private/MassProcessor.cpp
+++ b/Engine/Plugins/Runtime/MassEntity/Source/MassEntity/Private/MassProcessor.cpp
@@ -133,10 +133,43 @@ void UMassProcessor::SetShouldAutoRegisterWithGlobalList(const bool bAutoRegiste
 	}
 }
 
+void UMassProcessor::GetArchetypesMatchingOwnedQueries(const FMassEntityManager& EntityManager, TArray<FMassArchetypeHandle>& OutArchetype)
+{
+	for (FMassEntityQuery* QueryPtr : OwnedQueries)
+	{
+		CA_ASSUME(QueryPtr);
+		QueryPtr->CacheArchetypes(EntityManager);
+
+		for (const FMassArchetypeHandle& ArchetypeHandle : QueryPtr->GetArchetypes())
+		{
+			OutArchetype.AddUnique(ArchetypeHandle);
+		}
+	}
+}
+
+bool UMassProcessor::DoesAnyArchetypeMatchOwnedQueries(const FMassEntityManager& EntityManager)
+{
+	for (FMassEntityQuery* QueryPtr : OwnedQueries)
+	{
+		CA_ASSUME(QueryPtr);
+		QueryPtr->CacheArchetypes(EntityManager);
+
+		if (QueryPtr->GetArchetypes().Num() > 0)
+		{
+			return true;
+		}
+	}
+	return false;
+}
+
 void UMassProcessor::PostInitProperties()
 {
 	Super::PostInitProperties();
-	if (HasAnyFlags(RF_ClassDefaultObject) == false) 
+
+	// We want even the CDO processors to be fully set up so that we can reason about processors based on CDOs, without 
+	// needing to instantiate processors based on a given class.
+	// Note that we skip abstract processors because we're not going to instantiate them at runtime anyway.
+	if (GetClass()->HasAnyClassFlags(CLASS_Abstract) == false)
 	{
 		ConfigureQueries();
 
@@ -146,9 +179,10 @@ void UMassProcessor::PostInitProperties()
 			CA_ASSUME(QueryPtr);
 			bNeedsGameThread = (bNeedsGameThread || QueryPtr->DoesRequireGameThreadExecution());
 		}
-		// @todo warning temporarily disabled until details cleared
-		/*UE_CLOG(bRequiresGameThreadExecution != bNeedsGameThread, LogMass, Warning, TEXT("%s is marked bRequiresGameThreadExecution = %s, while the registered quries' requirement indicate the opposite")
-			, *GetProcessorName(), bRequiresGameThreadExecution ? TEXT("TRUE") : TEXT("FALSE"));*/
+		
+		UE_CLOG(bRequiresGameThreadExecution != bNeedsGameThread, LogMass, Verbose
+			, TEXT("%s is marked bRequiresGameThreadExecution = %s, while the registered quries' requirement indicate the opposite")
+			, *GetProcessorName(), bRequiresGameThreadExecution ? TEXT("TRUE") : TEXT("FALSE"));
 
 		// better safe than sorry - if queries indicate the game thread execution is required then we marked the whole processor as such
 		bRequiresGameThreadExecution = bRequiresGameThreadExecution || bNeedsGameThread;
@@ -321,39 +355,30 @@ void UMassCompositeProcessor::Initialize(UObject& Owner)
 	}
 }
 
-void UMassCompositeProcessor::CopyAndSort(const FMassProcessingPhaseConfig& PhaseConfig, const FString& DependencyGraphFileName)
-{
-	check(GetOuter());
-	
-	// create processors via a temporary pipeline. This ensures consistency with other places we use FRuntimePipelines
-	FMassRuntimePipeline TmpPipeline;
-	TmpPipeline.CreateFromArray(PhaseConfig.ProcessorCDOs, *GetOuter());
-
-	SetProcessors(TmpPipeline.Processors, DependencyGraphFileName);
-}
-
-void UMassCompositeProcessor::SetProcessors(TArrayView<UMassProcessor*> InProcessorInstances, const FString& DependencyGraphFileName)
+void UMassCompositeProcessor::SetProcessors(TArrayView<UMassProcessor*> InProcessorInstances, const TSharedPtr<FMassEntityManager>& EntityManager)
 {
 	// figure out dependencies
-	FProcessorDependencySolver Solver(InProcessorInstances, GroupName, DependencyGraphFileName);
-	TArray<FMassProcessorOrderInfo> SortedProcessorsAndGroups;
-	Solver.ResolveDependencies(SortedProcessorsAndGroups);
+	FProcessorDependencySolver Solver(InProcessorInstances);
+	TArray<FMassProcessorOrderInfo> SortedProcessors;
+	Solver.ResolveDependencies(SortedProcessors, EntityManager);
 
-	Populate(SortedProcessorsAndGroups);
+	Populate(SortedProcessors);
 
-	BuildFlatProcessingGraph(SortedProcessorsAndGroups);
+	BuildFlatProcessingGraph(SortedProcessors);
 }
 
-void UMassCompositeProcessor::BuildFlatProcessingGraph(TConstArrayView<FMassProcessorOrderInfo> SortedProcessorsAndGroups)
+void UMassCompositeProcessor::BuildFlatProcessingGraph(TConstArrayView<FMassProcessorOrderInfo> SortedProcessors)
 {
 	TRACE_CPUPROFILER_EVENT_SCOPE(Mass_BuildFlatProcessingGraph);
 
+	FlatProcessingGraph.Reset();
+
 	// this part is creating an ordered, flat list of processors that can be executed in sequence
 	// with subsequent task only depending on the elements prior on the list
 	TMap<FName, int32> NameToDependencyIndex;
-	NameToDependencyIndex.Reserve(SortedProcessorsAndGroups.Num());
+	NameToDependencyIndex.Reserve(SortedProcessors.Num());
 	TArray<int32> SuperGroupDependency;
-	for (const FMassProcessorOrderInfo& Element : SortedProcessorsAndGroups)
+	for (const FMassProcessorOrderInfo& Element : SortedProcessors)
 	{
 		NameToDependencyIndex.Add(Element.Name, FlatProcessingGraph.Num());
 
@@ -366,33 +391,29 @@ void UMassCompositeProcessor::BuildFlatProcessingGraph(TConstArrayView<FMassProc
 			checkSlow(DependencyName.IsNone() == false);
 			Node.Dependencies.Add(NameToDependencyIndex.FindChecked(DependencyName));
 		}
+#if WITH_MASSENTITY_DEBUG
+		Node.SequenceIndex = Element.SequenceIndex;
+#endif // WITH_MASSENTITY_DEBUG
 	}
 
 #if WITH_MASSENTITY_DEBUG
 	FScopedCategoryAndVerbosityOverride LogOverride(TEXT("LogMass"), ELogVerbosity::Log);
 	UE_LOG(LogMass, Log, TEXT("%s flat processing graph:"), *GroupName.ToString());
 
-	TArray<int32> Offset;
-	Offset.Reserve(FlatProcessingGraph.Num());
-
-	for (int i = 0; i < FlatProcessingGraph.Num(); ++i)
+	int32 Index = 0;
+	for (const FDependencyNode& ProcessingNode : FlatProcessingGraph)
 	{
-		FDependencyNode& ProcessingNode = FlatProcessingGraph[i];
 		FString DependenciesDesc;
-		int32 MyOffset = 0;
 		for (const int32 DependencyIndex : ProcessingNode.Dependencies)
 		{
 			DependenciesDesc += FString::Printf(TEXT("%d, "), DependencyIndex);
-			MyOffset = FMath::Max(MyOffset, Offset[DependencyIndex]);
 		}
-		MyOffset += 2;
-		Offset.Add(MyOffset);
-
 		if (ProcessingNode.Processor)
 		{
-			UE_LOG(LogMass, Log, TEXT("[%2d]%*s%s%s%s"), i, MyOffset, TEXT(""), *ProcessingNode.Processor->GetProcessorName()
+			UE_LOG(LogMass, Log, TEXT("[%2d]%*s%s%s%s"), Index, ProcessingNode.SequenceIndex * 2, TEXT(""), *ProcessingNode.Processor->GetProcessorName()
 				, DependenciesDesc.Len() > 0 ? TEXT(" depends on ") : TEXT(""), *DependenciesDesc);
 		}
+		++Index;
 	}
 #endif // WITH_MASSENTITY_DEBUG
 }
diff --git a/Engine/Plugins/Runtime/MassEntity/Source/MassEntity/Private/MassProcessorDependencySolver.cpp b/Engine/Plugins/Runtime/MassEntity/Source/MassEntity/Private/MassProcessorDependencySolver.cpp
--- a/Engine/Plugins/Runtime/MassEntity/Source/MassEntity/Private/MassProcessorDependencySolver.cpp
+++ b/Engine/Plugins/Runtime/MassEntity/Source/MassEntity/Private/MassProcessorDependencySolver.cpp
@@ -25,6 +25,18 @@ namespace UE::Mass::Private
 		}
 		return ReturnVal + TEXT("]");
 	}
+
+	bool DoArchetypeContainersOverlap(TConstArrayView<FMassArchetypeHandle> A, const TArray<FMassArchetypeHandle>& B)
+	{
+		for (const FMassArchetypeHandle& HandleA : A)
+		{
+			if (B.Contains(HandleA))
+			{
+				return true;
+			}
+		}
+		return false;
+	}
 }
 
 //----------------------------------------------------------------------//
@@ -75,10 +87,19 @@ bool FMassExecutionRequirements::IsEmpty() const
 		&& RequiredAllTags.IsEmpty() && RequiredAnyTags.IsEmpty() && RequiredNoneTags.IsEmpty();
 }
 
+FMassArchetypeCompositionDescriptor FMassExecutionRequirements::AsCompositionDescriptor() const
+{
+	return FMassArchetypeCompositionDescriptor(Fragments.Read + Fragments.Write
+		, RequiredAllTags + RequiredAnyTags
+		, ChunkFragments.Read + ChunkFragments.Write
+		, SharedFragments.Read + SharedFragments.Write);
+}
+
 //----------------------------------------------------------------------//
 //  FProcessorDependencySolver::FResourceUsage
 //----------------------------------------------------------------------//
-FProcessorDependencySolver::FResourceUsage::FResourceUsage()
+FProcessorDependencySolver::FResourceUsage::FResourceUsage(const TArray<FNode>& InAllNodes)
+	: AllNodesView(InAllNodes)
 {
 	for (int i = 0; i < EMassAccessOperation::MAX; ++i)
 	{
@@ -93,31 +114,61 @@ template<typename TBitSet>
 void FProcessorDependencySolver::FResourceUsage::HandleElementType(TMassExecutionAccess<FResourceAccess>& ElementAccess
 	, const TMassExecutionAccess<TBitSet>& TestedRequirements, FProcessorDependencySolver::FNode& InOutNode, const int32 NodeIndex)
 {
-	// for every bit set in TestedRequirements we do the following
-	// 1. For every read only requirement we make InOutNode depend on the currently stored Writer of this resource
+	using UE::Mass::Private::DoArchetypeContainersOverlap;
+
+	// for every bit set in TestedRequirements we do the following:
+	// 1. For every read-only requirement we make InOutNode depend on the currently stored Writer of this resource
 	//    - note that this operation is not destructive, meaning we don't destructively consume the data, since all 
 	//      subsequent read access to the given resource will also depend on the Writer
+	//    - note 2: we also fine tune what we store as a dependency for InOutNode by checking if InOutNode's archetype
+	//      overlap with whoever the current Writer is 
+	//    - this will result in InOutNode wait for the current Writer to finish before starting its own work and 
+	//      that's exactly what we need to do to avoid accessing data while it's potentially being written
 	// 2. For every read-write requirement we make InOutNode depend on all the readers and writers currently stored. 
 	//    - once that's done we clean currently stored Readers and Writers since every subsequent operation on this 
 	//      resource will be blocked by currently considered InOutNode (as the new Writer)
+	//    - again, we do check corresponding archetype collections overlap
+	//    - similarly to the read operation waiting on write operations in pt 1. we want to hold off the write 
+	//      operations to be performed by InOutNode until all currently registered (and conflicting) writers and readers 
+	//      are done with their operations 
 	// 3. For all accessed resources we store information that InOutNode is accessing it
+	//    - we do this so that the following nodes know that they'll have to wait for InOutNode if an access 
+	//      conflict arises. 
 
 	// 1. For every read only requirement we make InOutNode depend on the currently stored Writer of this resource
 	for (auto It = TestedRequirements.Read.GetIndexIterator(); It; ++It)
 	{
-		InOutNode.OriginalDependencies.Append(ElementAccess.Write.Access[*It].Users);
+		for (int32 UserIndex : ElementAccess.Write.Access[*It].Users)
+		{
+			if (DoArchetypeContainersOverlap(AllNodesView[UserIndex].ValidArchetypes, InOutNode.ValidArchetypes))
+			{
+				InOutNode.OriginalDependencies.Add(UserIndex);
+			}
+		}
 	}
 
 	// 2. For every read-write requirement we make InOutNode depend on all the readers and writers currently stored. 
 	for (auto It = TestedRequirements.Write.GetIndexIterator(); It; ++It)
 	{
-		InOutNode.OriginalDependencies.Append(ElementAccess.Read.Access[*It].Users);
-		ElementAccess.Read.Access[*It].Users.Reset();
+		for (int32 i = ElementAccess.Read.Access[*It].Users.Num() - 1; i >= 0; --i)
+		{
+			const int32 UserIndex = ElementAccess.Read.Access[*It].Users[i];
+			if (DoArchetypeContainersOverlap(AllNodesView[UserIndex].ValidArchetypes, InOutNode.ValidArchetypes))
+			{	
+				InOutNode.OriginalDependencies.Add(UserIndex);
+				ElementAccess.Read.Access[*It].Users.RemoveAtSwap(i);
+			}
+		}
 
-		// with the algorithm described above we expect there to only ever be at most one Writer
-		ensure(ElementAccess.Write.Access[*It].Users.Num() <= 1);
-		InOutNode.OriginalDependencies.Append(ElementAccess.Write.Access[*It].Users);
-		ElementAccess.Write.Access[*It].Users.Reset();
+		for (int32 i = ElementAccess.Write.Access[*It].Users.Num() - 1; i >= 0; --i)
+		{
+			const int32 UserIndex = ElementAccess.Write.Access[*It].Users[i];
+			if (DoArchetypeContainersOverlap(AllNodesView[UserIndex].ValidArchetypes, InOutNode.ValidArchetypes))
+			{
+				InOutNode.OriginalDependencies.Add(UserIndex);
+				ElementAccess.Write.Access[*It].Users.RemoveAtSwap(i);
+			}
+		}
 	}
 
 	// 3. For all accessed resources we store information that InOutNode is accessing it
@@ -149,11 +200,38 @@ bool FProcessorDependencySolver::FResourceUsage::CanAccess(const TMassExecutionA
 	);
 }
 
-bool FProcessorDependencySolver::FResourceUsage::CanAccessRequirements(const FMassExecutionRequirements& TestedRequirements) const
+bool FProcessorDependencySolver::FResourceUsage::HasArchetypeConflict(TMassExecutionAccess<FResourceAccess> ElementAccess, const TArray<FMassArchetypeHandle>& InArchetypes) const
+{
+	using UE::Mass::Private::DoArchetypeContainersOverlap;
+
+	// this function is being run when we've already determined there's an access conflict on given ElementsAccess,
+	// meaning whoever's asking is trying to access Elements that are already being used. We can still grant access 
+	// though provided that none of the current users of Element access the same archetypes the querier does (as provided 
+	// by InArchetypes).
+	// @todo this operation could be even more efficient and precise if we tracked which operation (read/write) and which
+	// specific Element were conflicting and the we could limit the check to that. That would however significantly 
+	// complicate the code and would require a major refactor to keep things clean.
+	for (int i = 0; i < EMassAccessOperation::MAX; ++i)
+	{
+		for (const FResourceUsers& Resource : ElementAccess[i].Access)
+		{
+			for (const int32 UserIndex : Resource.Users)
+			{
+				if (DoArchetypeContainersOverlap(AllNodesView[UserIndex].ValidArchetypes, InArchetypes))
+				{
+					return true;
+				}
+			}
+		}
+	}
+	return false;
+}
+
+bool FProcessorDependencySolver::FResourceUsage::CanAccessRequirements(const FMassExecutionRequirements& TestedRequirements, const TArray<FMassArchetypeHandle>& InArchetypes) const
 {
-	bool bCanAccess = CanAccess<FMassFragmentBitSet>(Requirements.Fragments, TestedRequirements.Fragments)
-		&& CanAccess<FMassChunkFragmentBitSet>(Requirements.ChunkFragments, TestedRequirements.ChunkFragments)
-		&& CanAccess<FMassSharedFragmentBitSet>(Requirements.SharedFragments, TestedRequirements.SharedFragments)
+	bool bCanAccess = (CanAccess<FMassFragmentBitSet>(Requirements.Fragments, TestedRequirements.Fragments) || !HasArchetypeConflict(FragmentsAccess, InArchetypes))
+		&& (CanAccess<FMassChunkFragmentBitSet>(Requirements.ChunkFragments, TestedRequirements.ChunkFragments) || !HasArchetypeConflict(ChunkFragmentsAccess, InArchetypes))
+		&& (CanAccess<FMassSharedFragmentBitSet>(Requirements.SharedFragments, TestedRequirements.SharedFragments) || !HasArchetypeConflict(SharedFragmentsAccess, InArchetypes))
 		&& CanAccess<FMassExternalSubsystemBitSet>(Requirements.RequiredSubsystems, TestedRequirements.RequiredSubsystems);
 
 	return bCanAccess;
@@ -169,13 +247,31 @@ void FProcessorDependencySolver::FResourceUsage::SubmitNode(const int32 NodeInde
 	Requirements.Append(InOutNode.Requirements);
 }
 
+//----------------------------------------------------------------------//
+//  FProcessorDependencySolver::FNode
+//----------------------------------------------------------------------//
+void FProcessorDependencySolver::FNode::IncreaseWaitingNodesCount(TArrayView<FProcessorDependencySolver::FNode> InAllNodes)
+{
+	// cycle-protection check. If true it means we have a cycle and the whole algorithm result will be unreliable 
+	if (TotalWaitingNodes >= FMath::Square(InAllNodes.Num()))
+	{
+		return;
+	}
+
+	++TotalWaitingNodes;
+
+	for (const int32 DependencyIndex : OriginalDependencies)
+	{
+		check(&InAllNodes[DependencyIndex] != this);
+		InAllNodes[DependencyIndex].IncreaseWaitingNodesCount(InAllNodes);
+	}
+}
 //----------------------------------------------------------------------//
 //  FProcessorDependencySolver
 //----------------------------------------------------------------------//
-FProcessorDependencySolver::FProcessorDependencySolver(TArrayView<UMassProcessor*> InProcessors, const FName Name, const FString& InDependencyGraphFileName /*= FString()*/)
+FProcessorDependencySolver::FProcessorDependencySolver(TArrayView<UMassProcessor*> InProcessors, const bool bIsGameRuntime)
 	: Processors(InProcessors)
-	, DependencyGraphFileName(InDependencyGraphFileName)
-	, CollectionName(Name)
+	, bGameRuntime(bIsGameRuntime)
 {}
 
 bool FProcessorDependencySolver::PerformSolverStep(FResourceUsage& ResourceUsage, TArray<int32>& InOutIndicesRemaining, TArray<int32>& OutNodeIndices)
@@ -188,7 +284,7 @@ bool FProcessorDependencySolver::PerformSolverStep(FResourceUsage& ResourceUsage
 		const int32 NodeIndex = InOutIndicesRemaining[i];
 		if (AllNodes[NodeIndex].TransientDependencies.Num() == 0)
 		{
-			if (ResourceUsage.CanAccessRequirements(AllNodes[NodeIndex].Requirements))
+			if (ResourceUsage.CanAccessRequirements(AllNodes[NodeIndex].Requirements, AllNodes[NodeIndex].ValidArchetypes))
 			{
 				AcceptedNodeIndex = NodeIndex;
 				break;
@@ -206,17 +302,24 @@ bool FProcessorDependencySolver::PerformSolverStep(FResourceUsage& ResourceUsage
 	{
 		const int32 NodeIndex = AcceptedNodeIndex != INDEX_NONE ? AcceptedNodeIndex : FallbackAcceptedNodeIndex;
 
+		FNode& Node = AllNodes[NodeIndex];
+
 		// Note that this is not an unexpected event and will happen during every dependency solving. It's a part 
 		// of the algorithm. We initially look for all the things we can run without conflicting with anything else. 
 		// But that can't last forever, at some point we'll end up in a situation where every node left waits for 
 		// something that has been submitted already. Then we just pick one of the waiting ones (the one indicated by 
 		// FallbackAcceptedNodeIndex), "run it" and proceed.
 		UE_CLOG(AcceptedNodeIndex == INDEX_NONE, LogMass, Verbose, TEXT("No dependency-free node can be picked, due to resource requirements. Picking %s as the next node.")
-			, *AllNodes[NodeIndex].Name.ToString());
+			, *Node.Name.ToString());
 
-		ResourceUsage.SubmitNode(NodeIndex, AllNodes[NodeIndex]);
+		ResourceUsage.SubmitNode(NodeIndex, Node);
 		InOutIndicesRemaining.RemoveSingle(NodeIndex);
 		OutNodeIndices.Add(NodeIndex);
+		for (const int32 DependencyIndex : Node.OriginalDependencies)
+		{
+			Node.SequencePositionIndex = FMath::Max(Node.SequencePositionIndex, AllNodes[DependencyIndex].SequencePositionIndex);
+		}
+		++Node.SequencePositionIndex;
 
 		for (const int32 RemainingNodeIndex : InOutIndicesRemaining)
 		{
@@ -250,16 +353,16 @@ void FProcessorDependencySolver::CreateSubGroupNames(FName InGroupName, TArray<F
 	}
 }
 
-void FProcessorDependencySolver::CreateNodes(UMassProcessor& Processor)
+int32 FProcessorDependencySolver::CreateNodes(UMassProcessor& Processor)
 {
 	check(Processor.GetClass());
 	const FName ProcName = Processor.GetClass()->GetFName();
 
-	if (NodeIndexMap.Find(ProcName) != nullptr)
+	if (const int32* NodeIndexPtr = NodeIndexMap.Find(ProcName))
 	{
 		UE_LOG(LogMass, Warning, TEXT("%s Processor %s already registered. Duplicates are not supported.")
 			, ANSI_TO_TCHAR(__FUNCTION__), *ProcName.ToString());
-		return;
+		return *NodeIndexPtr;
 	}
 
 	const FMassProcessorExecutionOrder& ExecutionOrder = Processor.GetExecutionOrder();
@@ -314,6 +417,8 @@ void FProcessorDependencySolver::CreateNodes(UMassProcessor& Processor)
 	{
 		AllNodes[ParentGroupNodeIndex].SubNodeIndices.Add(NodeIndex);
 	}
+
+	return NodeIndex;
 }
 
 void FProcessorDependencySolver::BuildDependencies()
@@ -855,6 +960,7 @@ void FProcessorDependencySolver::Solve(TArray<FMassProcessorOrderInfo>& OutResul
 	for (FNode& Node : AllNodes)
 	{
 		Node.TransientDependencies = Node.OriginalDependencies;
+		Node.TotalWaitingNodes = 0;
 	}
 
 	TArray<int32> IndicesRemaining;
@@ -865,15 +971,16 @@ void FProcessorDependencySolver::Solve(TArray<FMassProcessorOrderInfo>& OutResul
 		if (AllNodes[i].IsGroup() == false)
 		{
 			IndicesRemaining.Add(i);
+			AllNodes[i].IncreaseWaitingNodesCount(AllNodes);
 		}
 	}
-	// order the indices to process according to calculated resource hunger
+
 	IndicesRemaining.Sort([this](const int32 IndexA, const int32 IndexB){
-		return AllNodes[IndexA].Requirements.ResourcesUsedCount < AllNodes[IndexB].Requirements.ResourcesUsedCount;
+		return AllNodes[IndexA].TotalWaitingNodes > AllNodes[IndexB].TotalWaitingNodes;
 	});
 
 	// this is where we'll be tracking what's being accessed by whom
-	FResourceUsage ResourceUsage;
+	FResourceUsage ResourceUsage(AllNodes);
 
 	TArray<int32> SortedNodeIndices;
 	SortedNodeIndices.Reserve(AllNodes.Num());
@@ -929,13 +1036,15 @@ void FProcessorDependencySolver::Solve(TArray<FMassProcessorOrderInfo>& OutResul
 		// at this point we expect SortedNodeIndices to only point to regular processors (i.e. no groups)
 		if (ensure(AllNodes[NodeIndex].Processor != nullptr))
 		{
-			OutResult.Add({ AllNodes[NodeIndex].Name, AllNodes[NodeIndex].Processor, FMassProcessorOrderInfo::EDependencyNodeType::Processor, DependencyNames });
+			OutResult.Add({ AllNodes[NodeIndex].Name, AllNodes[NodeIndex].Processor, FMassProcessorOrderInfo::EDependencyNodeType::Processor, DependencyNames, AllNodes[NodeIndex].SequencePositionIndex });
 		}
 	}
 }
 
-void FProcessorDependencySolver::ResolveDependencies(TArray<FMassProcessorOrderInfo>& OutResult)
+void FProcessorDependencySolver::ResolveDependencies(TArray<FMassProcessorOrderInfo>& OutResult, TSharedPtr<FMassEntityManager> EntityManager, FProcessorDependencySolver::FResult* InOutOptionalResult)
 {
+	TRACE_CPUPROFILER_EVENT_SCOPE_STR("Mass ResolveDependencies");
+
 	if (Processors.Num() == 0)
 	{
 		return;
@@ -943,6 +1052,11 @@ void FProcessorDependencySolver::ResolveDependencies(TArray<FMassProcessorOrderI
 
 	FScopedCategoryAndVerbosityOverride LogOverride(TEXT("LogMass"), ELogVerbosity::Log);
 
+	if (InOutOptionalResult)
+	{
+		DependencyGraphFileName = InOutOptionalResult->DependencyGraphFileName;
+	}
+
 	bAnyCyclesDetected = false;
 
 	UE_LOG(LogMass, Log, TEXT("Gathering dependencies data:"));
@@ -952,9 +1066,20 @@ void FProcessorDependencySolver::ResolveDependencies(TArray<FMassProcessorOrderI
 	// as the very first node we add a "root" node that represents the "top level group" and also simplifies the rest
 	// of the lookup code - if a processor declares it's in group None or depends on Node it we don't need to check that 
 	// explicitly. 
-	AllNodes.Add(FNode(CollectionName, nullptr, 0));
+	AllNodes.Add(FNode(FName(), nullptr, 0));
 	NodeIndexMap.Add(FName(), 0);
 
+	const bool bCreateVirtualArchetypes = (!EntityManager);
+	if (bCreateVirtualArchetypes)
+	{
+		// create FMassEntityManager instance that we'll use to sort out processors' overlaps
+		// the idea for this is that for every processor we have we create an archetype matching given processor's requirements. 
+		// Once that's done we have a collection of "virtual" archetypes our processors expect. Then we ask every processor 
+		// to cache the archetypes they'd accept, using processors' owned queries. The idea is that some of the nodes will 
+		// end up with more than just the virtual archetype created for that specific node. The practice proved the idea correct. 
+		EntityManager = MakeShareable(new FMassEntityManager());
+	}
+
 	// gather the processors information first
 	for (UMassProcessor* Processor : Processors)
 	{
@@ -963,9 +1088,49 @@ void FProcessorDependencySolver::ResolveDependencies(TArray<FMassProcessorOrderI
 			UE_LOG(LogMass, Warning, TEXT("%s nullptr found in Processors collection being processed"), ANSI_TO_TCHAR(__FUNCTION__));
 			continue;
 		}
-		CreateNodes(*Processor);
+
+		const int32 ProcessorNodeIndex = CreateNodes(*Processor);
+
+		if (bCreateVirtualArchetypes)
+		{
+			// this line is a part of a nice trick we're doing here utilizing EntityManager's archetype creation based on 
+			// what each processor expects, and EntityQuery's capability to cache archetypes matching its requirements (used below)
+			EntityManager->CreateArchetype(AllNodes[ProcessorNodeIndex].Requirements.AsCompositionDescriptor());
+		}
+	}
+
+	UE_LOG(LogMass, Verbose, TEXT("Pruning processors..."));
+
+	int32 PrunedProcessorsCount = 0;
+	for (FNode& Node : AllNodes)
+	{
+		if (Node.IsGroup() == false)
+		{
+			// for each processor-representing node we cache information on which archetypes among the once we've created 
+			// above (see the EntityManager.CreateArchetype call in the previous loop) match this processor. 
+			Node.Processor->GetArchetypesMatchingOwnedQueries(*EntityManager.Get(), Node.ValidArchetypes);
+
+			// prune the archetype-less processors
+			if (Node.ValidArchetypes.Num() == 0 && Node.Processor->ShouldAllowQueryBasedPruning(bGameRuntime))
+			{
+				CA_ASSUME(Node.Processor);
+				UE_LOG(LogMass, Verbose, TEXT("\t%s"), *Node.Processor->GetName());
+
+				if (InOutOptionalResult)
+				{
+					InOutOptionalResult->PrunedProcessorClasses.Add(Node.Processor->GetClass());
+				}
+
+				// clearing out the processor will result in the rest of the algorithm to treat it as a group - we still 
+				// want to preserve the configured ExecuteBefore and ExecuteAfter dependencies
+				Node.Processor = nullptr;
+				++PrunedProcessorsCount;
+			}
+		}
 	}
 
+	UE_LOG(LogMass, Verbose, TEXT("Number of processors pruned: %d"), PrunedProcessorsCount);
+
 	check(AllNodes.Num());
 	LogNode(AllNodes[0]);
 
@@ -982,6 +1147,42 @@ void FProcessorDependencySolver::ResolveDependencies(TArray<FMassProcessorOrderI
 	{
 		UE_LOG(LogMass, Verbose, TEXT("\t%s"), *Info.Name.ToString());
 	}
+
+	int32 MaxSequenceLength = 0;
+	for (FNode& Node : AllNodes)
+	{
+		MaxSequenceLength = FMath::Max(MaxSequenceLength, Node.SequencePositionIndex);
+	}
+
+	UE_LOG(LogMass, Verbose, TEXT("Max sequence length: %d"), MaxSequenceLength);
+
+	if (InOutOptionalResult)
+	{
+		InOutOptionalResult->MaxSequenceLength = MaxSequenceLength;
+		InOutOptionalResult->ArchetypeDataVersion = EntityManager->GetArchetypeDataVersion();
+	}
+}
+
+bool FProcessorDependencySolver::IsResultUpToDate(const FProcessorDependencySolver::FResult& InResult, TSharedPtr<FMassEntityManager> EntityManager)
+{
+	if (InResult.PrunedProcessorClasses.Num() == 0 || !EntityManager || InResult.ArchetypeDataVersion == EntityManager->GetArchetypeDataVersion())
+	{
+		return true;
+	}
+
+	// this is inefficient right now since we're using CDOs and need to check all archetypes every time.
+	// Would be more efficient if we had a common place where all processors live, both active and inactive, so that we can utilize those. 
+	for (const TSubclassOf<UMassProcessor>& ProcessorClass : InResult.PrunedProcessorClasses)
+	{
+		if (UMassProcessor* ProcessorCDO = ProcessorClass.GetDefaultObject())
+		{
+			if (ProcessorCDO->DoesAnyArchetypeMatchOwnedQueries(*EntityManager.Get()))
+			{
+				return false;
+			}
+		}
+	}
+	return true;
 }
 
 #undef LOCTEXT_NAMESPACE 
\ No newline at end of file
diff --git a/Engine/Plugins/Runtime/MassEntity/Source/MassEntity/Public/MassProcessingPhaseManager.h b/Engine/Plugins/Runtime/MassEntity/Source/MassEntity/Public/MassProcessingPhaseManager.h
--- a/Engine/Plugins/Runtime/MassEntity/Source/MassEntity/Public/MassProcessingPhaseManager.h
+++ b/Engine/Plugins/Runtime/MassEntity/Source/MassEntity/Public/MassProcessingPhaseManager.h
@@ -9,6 +9,7 @@
 #include "Engine/EngineBaseTypes.h"
 #include "MassProcessingTypes.h"
 #include "MassProcessor.h"
+#include "MassProcessorDependencySolver.h"
 #include "MassProcessingPhaseManager.generated.h"
 
 
@@ -89,6 +90,22 @@ private:
 };
 
 
+struct MASSENTITY_API FMassPhaseProcessorConfigurationHelper
+{
+	FMassPhaseProcessorConfigurationHelper(UMassCompositeProcessor& InOutPhaseProcessor, const FMassProcessingPhaseConfig& InPhaseConfig, UObject& InProcessorOuter)
+		: PhaseProcessor(InOutPhaseProcessor), PhaseConfig(InPhaseConfig), ProcessorOuter(InProcessorOuter)
+	{
+	}
+
+	void Configure(const TSharedPtr<FMassEntityManager>& EntityManager = TSharedPtr<FMassEntityManager>(), FProcessorDependencySolver::FResult* OutOptionalResult = nullptr);
+
+	UMassCompositeProcessor& PhaseProcessor;
+	const FMassProcessingPhaseConfig& PhaseConfig;
+	UObject& ProcessorOuter;
+	bool bInitializeCreatedProcessors = true;
+	bool bIsGameRuntime = true;
+};
+
 /** 
  * MassProcessingPhaseManager owns separate FMassProcessingPhase instances for every ETickingGroup. When activated
  * via Start function it registers and enables the FMassProcessingPhase instances which themselves are tick functions 
@@ -104,6 +121,7 @@ public:
 	FMassProcessingPhaseManager(const FMassProcessingPhaseManager& Other) = delete;
 	FMassProcessingPhaseManager& operator=(const FMassProcessingPhaseManager& Other) = delete;
 
+	const TSharedPtr<FMassEntityManager>& GetEntityManager() { return EntityManager; }
 	FMassEntityManager& GetEntityManagerRef() { check(EntityManager); return *EntityManager.Get(); }
 
 	/** Retrieves OnPhaseStart multicast delegate's reference for a given Phase */
@@ -170,13 +188,25 @@ protected:
 	 */
 	void OnPhaseEnd(FMassProcessingPhase& Phase);
 
+	void OnNewArchetype(const FMassArchetypeHandle& NewArchetype);
+
 protected:	
+	struct FPhaseGraphBuildState
+	{
+		FProcessorDependencySolver::FResult LastResult;
+		bool bNewArchetypes = true;
+		bool bInitialized = false;
+	};
 
 	FMassProcessingPhase ProcessingPhases[(uint8)EMassProcessingPhase::MAX];
+	FPhaseGraphBuildState ProcessingGraphBuildStates[(uint8)EMassProcessingPhase::MAX];
+	TArray<FMassProcessingPhaseConfig> ProcessingPhasesConfig;
 
 	TSharedPtr<FMassEntityManager> EntityManager;
 
 	EMassProcessingPhase CurrentPhase = EMassProcessingPhase::MAX;
 
 	TWeakObjectPtr<UObject> Owner;
+
+	FDelegateHandle OnNewArchetypeHandle;
 };
diff --git a/Engine/Plugins/Runtime/MassEntity/Source/MassEntity/Public/MassProcessingTypes.h b/Engine/Plugins/Runtime/MassEntity/Source/MassEntity/Public/MassProcessingTypes.h
--- a/Engine/Plugins/Runtime/MassEntity/Source/MassEntity/Public/MassProcessingTypes.h
+++ b/Engine/Plugins/Runtime/MassEntity/Source/MassEntity/Public/MassProcessingTypes.h
@@ -140,4 +140,5 @@ struct FMassProcessorOrderInfo
 	UMassProcessor* Processor = nullptr;
 	EDependencyNodeType NodeType = EDependencyNodeType::Invalid;
 	TArray<FName> Dependencies;
+	int32 SequenceIndex = INDEX_NONE;
 };
diff --git a/Engine/Plugins/Runtime/MassEntity/Source/MassEntity/Public/MassProcessor.h b/Engine/Plugins/Runtime/MassEntity/Source/MassEntity/Public/MassProcessor.h
--- a/Engine/Plugins/Runtime/MassEntity/Source/MassEntity/Public/MassProcessor.h
+++ b/Engine/Plugins/Runtime/MassEntity/Source/MassEntity/Public/MassProcessor.h
@@ -64,6 +64,13 @@ public:
 	//----------------------------------------------------------------------//
 	// Ordering functions 
 	//----------------------------------------------------------------------//
+	/** Indicates whether this processor can ever be pruned while considered for a phase processing graph. A processor
+	 *  can get pruned if none of its registered queries interact with archetypes instantiated at the moment of graph
+	 *  building. This can also happen for special processors that don't register any queries - if that's the case override 
+	 *  this function to return an appropriate value
+	 *  @param bRuntimeMode indicates whether the pruning is being done for game runtime (true) or editor-time presentation (false) */
+	virtual bool ShouldAllowQueryBasedPruning(const bool bRuntimeMode = true) const { return bRuntimeMode; }
+
 	virtual EMassProcessingPhase GetProcessingPhase() const { return ProcessingPhase; }
 	virtual void SetProcessingPhase(EMassProcessingPhase Phase) { ProcessingPhase = Phase; }
 	bool DoesRequireGameThreadExecution() const { return bRequiresGameThreadExecution; }
@@ -91,6 +98,9 @@ public:
 	 * have no other effect, i.e. CDO's value won't change */
 	void SetShouldAutoRegisterWithGlobalList(const bool bAutoRegister);
 
+	void GetArchetypesMatchingOwnedQueries(const FMassEntityManager& EntityManager, TArray<FMassArchetypeHandle>& OutArchetype);
+	bool DoesAnyArchetypeMatchOwnedQueries(const FMassEntityManager& EntityManager);
+	
 #if CPUPROFILERTRACE_ENABLED
 	FString StatId;
 #endif
@@ -157,6 +167,9 @@ public:
 		FName Name;
 		UMassProcessor* Processor = nullptr;
 		TArray<int32> Dependencies;
+#if WITH_MASSENTITY_DEBUG
+		int32 SequenceIndex = INDEX_NONE;
+#endif // WITH_MASSENTITY_DEBUG
 	};
 
 public:
@@ -171,9 +184,17 @@ public:
 	void SetGroupName(FName NewName);
 	FName GetGroupName() const { return GroupName; }
 
-	virtual void CopyAndSort(const FMassProcessingPhaseConfig& PhaseConfig, const FString& DependencyGraphFileName = FString());
-	virtual void SetProcessors(TArrayView<UMassProcessor*> InProcessorInstances, const FString& DependencyGraphFileName = FString());
-	virtual void BuildFlatProcessingGraph(TConstArrayView<FMassProcessorOrderInfo> SortedProcessorsAndGroups);
+	virtual void SetProcessors(TArrayView<UMassProcessor*> InProcessorInstances, const TSharedPtr<FMassEntityManager>& EntityManager = nullptr);
+
+	/** 
+	 * Builds flat processing graph that's being used for multithreading execution of hosted processors.
+	 */
+	virtual void BuildFlatProcessingGraph(TConstArrayView<FMassProcessorOrderInfo> SortedProcessors);
+
+	/**
+	 *  Adds processors in OrderedProcessors to ChildPipeline
+	 */
+	void Populate(TConstArrayView<FMassProcessorOrderInfo> OrderedProcessors);
 
 	/** adds SubProcessor to an appropriately named group. If RequestedGroupName == None then SubProcessor
 	 *  will be added directly to ChildPipeline. If not then the indicated group will be searched for in ChildPipeline 
@@ -190,11 +211,6 @@ protected:
 	virtual void ConfigureQueries() override;
 	virtual void Execute(FMassEntityManager& EntityManager, FMassExecutionContext& Context) override;
 
-	/**
-	 *  Adds processors in OrderedProcessors to ChildPipeline and builds flat processing graph that's being used in multithreaded mode.
-	 */
-	void Populate(TConstArrayView<FMassProcessorOrderInfo> OrderedProcessors);
-
 	/** RequestedGroupName can indicate a multi-level group name, like so: A.B.C
 	 *  We need to extract the highest-level group name ('A' in the example), and see if it already exists. 
 	 *  If not, create it. 
diff --git a/Engine/Plugins/Runtime/MassEntity/Source/MassEntity/Public/MassProcessorDependencySolver.h b/Engine/Plugins/Runtime/MassEntity/Source/MassEntity/Public/MassProcessorDependencySolver.h
--- a/Engine/Plugins/Runtime/MassEntity/Source/MassEntity/Public/MassProcessorDependencySolver.h
+++ b/Engine/Plugins/Runtime/MassEntity/Source/MassEntity/Public/MassProcessorDependencySolver.h
@@ -4,6 +4,7 @@
 
 #include "MassProcessingTypes.h"
 #include "MassEntityTypes.h"
+#include "MassArchetypeTypes.h"
 #include "Containers/StaticArray.h"
 
 
@@ -45,6 +46,7 @@ struct MASSENTITY_API FMassExecutionRequirements
 	void CountResourcesUsed();
 	int32 GetTotalBitsUsedCount();
 	bool IsEmpty() const;
+	FMassArchetypeCompositionDescriptor AsCompositionDescriptor() const;
 
 	TMassExecutionAccess<FMassFragmentBitSet> Fragments;
 	TMassExecutionAccess<FMassChunkFragmentBitSet> ChunkFragments;
@@ -65,6 +67,9 @@ private:
 			: Name(InName), Processor(InProcessor), NodeIndex(InNodeIndex)
 		{}
 
+		bool IsGroup() const { return Processor == nullptr; }
+		void IncreaseWaitingNodesCount(TArrayView<FNode> InAllNodes);
+
 		FName Name = TEXT("");
 		UMassProcessor* Processor = nullptr;
 		TArray<int32> OriginalDependencies;
@@ -73,16 +78,22 @@ private:
 		TArray<FName> ExecuteAfter;
 		FMassExecutionRequirements Requirements;
 		int32 NodeIndex = INDEX_NONE;
+		/** indicates how often given node can be found in dependencies sequence for other nodes  */
+		int32 TotalWaitingNodes = 0;
+		/** 
+		 * indicates how deep within dependencies graph this give node is, or in other words, what's the longest sequence 
+		 * from this node to a dependency-less "parent" node 
+		 */
+		int32 SequencePositionIndex = 0;
 		TArray<int32> SubNodeIndices;
-
-		bool IsGroup() const { return Processor == nullptr; }
+		TArray<FMassArchetypeHandle> ValidArchetypes;
 	};
 
 	struct FResourceUsage
 	{
-		FResourceUsage();
+		FResourceUsage(const TArray<FNode>& InAllNodes);
 
-		bool CanAccessRequirements(const FMassExecutionRequirements& TestedRequirements) const;
+		bool CanAccessRequirements(const FMassExecutionRequirements& TestedRequirements, const TArray<FMassArchetypeHandle>& InArchetypes) const;
 		void SubmitNode(const int32 NodeIndex, FNode& InOutNode);
 
 	private:
@@ -101,22 +112,45 @@ private:
 		TMassExecutionAccess<FResourceAccess> ChunkFragmentsAccess;
 		TMassExecutionAccess<FResourceAccess> SharedFragmentsAccess;
 		TMassExecutionAccess<FResourceAccess> RequiredSubsystemsAccess;
+		TConstArrayView<FNode> AllNodesView;
 
 		template<typename TBitSet>
-		static void HandleElementType(TMassExecutionAccess<FResourceAccess>& ElementAccess
+		void HandleElementType(TMassExecutionAccess<FResourceAccess>& ElementAccess
 			, const TMassExecutionAccess<TBitSet>& TestedRequirements, FProcessorDependencySolver::FNode& InOutNode, const int32 NodeIndex);
 
 		template<typename TBitSet>
 		static bool CanAccess(const TMassExecutionAccess<TBitSet>& StoredElements, const TMassExecutionAccess<TBitSet>& TestedElements);
+
+		/** Determines whether any of the Elements' (i.e. Fragment, Tag,...) users operate on any of the archetypes given via InArchetypes */
+		bool HasArchetypeConflict(TMassExecutionAccess<FResourceAccess> ElementAccess, const TArray<FMassArchetypeHandle>& InArchetypes) const;
 	};
 
 public:
+	/** Optionally returned by ResolveDependencies and contains information about processors that have been pruned and 
+	 *  other potentially useful bits. To be used in a transient fashion. */
+	struct FResult
+	{
+		FString DependencyGraphFileName;
+		TArray<TSubclassOf<UMassProcessor>> PrunedProcessorClasses;
+		int32 MaxSequenceLength = 0;
+		uint32 ArchetypeDataVersion = 0;
 
-	MASSENTITY_API FProcessorDependencySolver(TArrayView<UMassProcessor*> InProcessors, const FName Name, const FString& InDependencyGraphFileName = FString());
-	MASSENTITY_API void ResolveDependencies(TArray<FMassProcessorOrderInfo>& OutResult);
+		void Reset()
+		{
+			PrunedProcessorClasses.Reset();
+			MaxSequenceLength = 0;
+			ArchetypeDataVersion = 0;
+		}
+	};
+
+	MASSENTITY_API FProcessorDependencySolver(TArrayView<UMassProcessor*> InProcessors, const bool bIsGameRuntime = true);
+	MASSENTITY_API void ResolveDependencies(TArray<FMassProcessorOrderInfo>& OutResult, TSharedPtr<FMassEntityManager> EntityManager = nullptr, FResult* InOutOptionalResult = nullptr);
 
 	MASSENTITY_API static void CreateSubGroupNames(FName InGroupName, TArray<FString>& SubGroupNames);
 
+	/** Determines whether the dependency solving that produced InResult will produce different results if run with a given EntityManager */
+	static bool IsResultUpToDate(const FResult& InResult, TSharedPtr<FMassEntityManager> EntityManager);
+
 protected:
 	// note that internals are protected rather than private to support unit testing
 
@@ -128,7 +162,7 @@ protected:
 	 */
 	bool PerformSolverStep(FResourceUsage& ResourceUsage, TArray<int32>& InOutIndicesRemaining, TArray<int32>& OutNodeIndices);
 	
-	void CreateNodes(UMassProcessor& Processor);
+	int32 CreateNodes(UMassProcessor& Processor);
 	void BuildDependencies();
 	void Solve(TArray<FMassProcessorOrderInfo>& OutResult);
 	void LogNode(const FNode& Node, int Indent = 0);
@@ -139,8 +173,8 @@ protected:
 
 	TArrayView<UMassProcessor*> Processors;
 	bool bAnyCyclesDetected = false;
+	const bool bGameRuntime = true;
 	FString DependencyGraphFileName;
-	FName CollectionName;
 	TArray<FNode> AllNodes;
 	TMap<FName, int32> NodeIndexMap;
 };
\ No newline at end of file
diff --git a/Engine/Plugins/Runtime/MassEntity/Source/MassEntityTestSuite/Private/MassDependencyTest.cpp b/Engine/Plugins/Runtime/MassEntity/Source/MassEntityTestSuite/Private/MassDependencyTest.cpp
--- a/Engine/Plugins/Runtime/MassEntity/Source/MassEntityTestSuite/Private/MassDependencyTest.cpp
+++ b/Engine/Plugins/Runtime/MassEntity/Source/MassEntityTestSuite/Private/MassDependencyTest.cpp
@@ -33,7 +33,7 @@ struct FDependencySolverBase : FAITestBase
 	void Solve()
 	{
 		Result.Reset();
-		FProcessorDependencySolver Solver(MakeArrayView((UMassProcessor**)Processors.GetData(), Processors.Num()), TEXT("Test") );
+		FProcessorDependencySolver Solver(MakeArrayView((UMassProcessor**)Processors.GetData(), Processors.Num()));
 		Solver.ResolveDependencies(Result);
 	}
 };	
@@ -210,6 +210,8 @@ struct FComplexScenario : FDependencySolverBase
 	{
 		Solve();
 
+		AITEST_TRUE("None of the processors should have been pruned", Result.Num() == Processors.Num());
+
 		for (int32 i = 0; i < Result.Num(); ++i)
 		{
 			AITEST_EQUAL("We expect only processor nodes in the results", Result[i].NodeType, FMassProcessorOrderInfo::EDependencyNodeType::Processor);
diff --git a/Engine/Plugins/Runtime/MassEntity/Source/MassEntityTestSuite/Public/MassEntityTestTypes.h b/Engine/Plugins/Runtime/MassEntity/Source/MassEntityTestSuite/Public/MassEntityTestTypes.h
--- a/Engine/Plugins/Runtime/MassEntity/Source/MassEntityTestSuite/Public/MassEntityTestTypes.h
+++ b/Engine/Plugins/Runtime/MassEntity/Source/MassEntityTestSuite/Public/MassEntityTestTypes.h
@@ -114,6 +114,7 @@ public:
 	{
 		RequirementsFunction(EntityQuery);
 	}
+	virtual bool ShouldAllowQueryBasedPruning(const bool bRuntimeMode = true) const { return false; }
 
 	TFunction<void(FMassEntityManager& EntityManager, FMassExecutionContext& Context)> ExecutionFunction;
 	TFunction<void(FMassEntityQuery& Query)> RequirementsFunction;
