FProcessorDependencySolver renamed to FMassProcessorDependencySolver

diff --git a/Engine/Plugins/Runtime/MassEntity/Source/MassEntity/Private/MassEntitySettings.cpp b/Engine/Plugins/Runtime/MassEntity/Source/MassEntity/Private/MassEntitySettings.cpp
--- a/Engine/Plugins/Runtime/MassEntity/Source/MassEntity/Private/MassEntitySettings.cpp
+++ b/Engine/Plugins/Runtime/MassEntity/Source/MassEntity/Private/MassEntitySettings.cpp
@@ -56,7 +56,7 @@ void UMassEntitySettings::BuildPhases()
 			PhaseConfig.PhaseProcessor->SetProcessingPhase(EMassProcessingPhase(i));
 			const FString PhaseDumpDependencyGraphFileName = !DumpDependencyGraphFileName.IsEmpty() ? DumpDependencyGraphFileName + TEXT("_") + PhaseConfig.PhaseName.ToString() : FString();
 
-			FProcessorDependencySolver::FResult Result;
+			FMassProcessorDependencySolver::FResult Result;
 			Result.DependencyGraphFileName = PhaseDumpDependencyGraphFileName;
 			FMassPhaseProcessorConfigurationHelper Configurator(*PhaseConfig.PhaseProcessor, PhaseConfig, *this);
 			Configurator.bInitializeCreatedProcessors = false;
diff --git a/Engine/Plugins/Runtime/MassEntity/Source/MassEntity/Private/MassProcessingPhaseManager.cpp b/Engine/Plugins/Runtime/MassEntity/Source/MassEntity/Private/MassProcessingPhaseManager.cpp
--- a/Engine/Plugins/Runtime/MassEntity/Source/MassEntity/Private/MassProcessingPhaseManager.cpp
+++ b/Engine/Plugins/Runtime/MassEntity/Source/MassEntity/Private/MassProcessingPhaseManager.cpp
@@ -136,13 +136,13 @@ void FMassProcessingPhase::Initialize(FMassProcessingPhaseManager& InPhaseManage
 //----------------------------------------------------------------------//
 // FPhaseProcessorConfigurator
 //----------------------------------------------------------------------//
-void FMassPhaseProcessorConfigurationHelper::Configure(const TSharedPtr<FMassEntityManager>& EntityManager, FProcessorDependencySolver::FResult* OutOptionalResult)
+void FMassPhaseProcessorConfigurationHelper::Configure(const TSharedPtr<FMassEntityManager>& EntityManager, FMassProcessorDependencySolver::FResult* OutOptionalResult)
 {
 	FMassRuntimePipeline TmpPipeline;
 	TmpPipeline.CreateFromArray(PhaseConfig.ProcessorCDOs, ProcessorOuter);
 
 	TArray<FMassProcessorOrderInfo> SortedProcessors;
-	FProcessorDependencySolver Solver(TmpPipeline.Processors, bIsGameRuntime);
+	FMassProcessorDependencySolver Solver(TmpPipeline.Processors, bIsGameRuntime);
 
 	Solver.ResolveDependencies(SortedProcessors, EntityManager, OutOptionalResult);
 
@@ -194,7 +194,7 @@ void FMassProcessingPhaseManager::Initialize(UObject& InOwner, TConstArrayView<F
 			// populating the phase processor with initial data for editor world so that the default processing graph
 			// can be investigated in the editor without running PIE.
 			// Runtime processing graph are initialized at runtime base on the actual archetypes instantiated at call time.
-			FProcessorDependencySolver::FResult Result;
+			FMassProcessorDependencySolver::FResult Result;
 			Result.DependencyGraphFileName = DependencyGraphFileName;
 			FMassPhaseProcessorConfigurationHelper Configurator(*PhaseProcessor, ProcessingPhasesConfig[PhaseAsInt], InOwner);
 			Configurator.bIsGameRuntime = false;
@@ -318,7 +318,7 @@ void FMassProcessingPhaseManager::OnPhaseStart(const FMassProcessingPhase& Phase
 		TRACE_CPUPROFILER_EVENT_SCOPE_STR("Mass Rebuild Phase Graph");
 
 		FPhaseGraphBuildState& GraphBuildState = ProcessingGraphBuildStates[PhaseAsInt];
-		if (GraphBuildState.bInitialized == false || FProcessorDependencySolver::IsResultUpToDate(GraphBuildState.LastResult, EntityManager) == false)
+		if (GraphBuildState.bInitialized == false || FMassProcessorDependencySolver::IsResultUpToDate(GraphBuildState.LastResult, EntityManager) == false)
 		{
 			UMassCompositeProcessor* PhaseProcessor = ProcessingPhases[PhaseAsInt].PhaseProcessor;
 			check(PhaseProcessor);
diff --git a/Engine/Plugins/Runtime/MassEntity/Source/MassEntity/Private/MassProcessor.cpp b/Engine/Plugins/Runtime/MassEntity/Source/MassEntity/Private/MassProcessor.cpp
--- a/Engine/Plugins/Runtime/MassEntity/Source/MassEntity/Private/MassProcessor.cpp
+++ b/Engine/Plugins/Runtime/MassEntity/Source/MassEntity/Private/MassProcessor.cpp
@@ -358,7 +358,7 @@ void UMassCompositeProcessor::Initialize(UObject& Owner)
 void UMassCompositeProcessor::SetProcessors(TArrayView<UMassProcessor*> InProcessorInstances, const TSharedPtr<FMassEntityManager>& EntityManager)
 {
 	// figure out dependencies
-	FProcessorDependencySolver Solver(InProcessorInstances);
+	FMassProcessorDependencySolver Solver(InProcessorInstances);
 	TArray<FMassProcessorOrderInfo> SortedProcessors;
 	Solver.ResolveDependencies(SortedProcessors, EntityManager);
 
diff --git a/Engine/Plugins/Runtime/MassEntity/Source/MassEntity/Private/MassProcessorDependencySolver.cpp b/Engine/Plugins/Runtime/MassEntity/Source/MassEntity/Private/MassProcessorDependencySolver.cpp
--- a/Engine/Plugins/Runtime/MassEntity/Source/MassEntity/Private/MassProcessorDependencySolver.cpp
+++ b/Engine/Plugins/Runtime/MassEntity/Source/MassEntity/Private/MassProcessorDependencySolver.cpp
@@ -98,7 +98,7 @@ FMassArchetypeCompositionDescriptor FMassExecutionRequirements::AsCompositionDes
 //----------------------------------------------------------------------//
 //  FProcessorDependencySolver::FResourceUsage
 //----------------------------------------------------------------------//
-FProcessorDependencySolver::FResourceUsage::FResourceUsage(const TArray<FNode>& InAllNodes)
+FMassProcessorDependencySolver::FResourceUsage::FResourceUsage(const TArray<FNode>& InAllNodes)
 	: AllNodesView(InAllNodes)
 {
 	for (int i = 0; i < EMassAccessOperation::MAX; ++i)
@@ -111,8 +111,8 @@ FProcessorDependencySolver::FResourceUsage::FResourceUsage(const TArray<FNode>&
 }
 
 template<typename TBitSet>
-void FProcessorDependencySolver::FResourceUsage::HandleElementType(TMassExecutionAccess<FResourceAccess>& ElementAccess
-	, const TMassExecutionAccess<TBitSet>& TestedRequirements, FProcessorDependencySolver::FNode& InOutNode, const int32 NodeIndex)
+void FMassProcessorDependencySolver::FResourceUsage::HandleElementType(TMassExecutionAccess<FResourceAccess>& ElementAccess
+	, const TMassExecutionAccess<TBitSet>& TestedRequirements, FMassProcessorDependencySolver::FNode& InOutNode, const int32 NodeIndex)
 {
 	using UE::Mass::Private::DoArchetypeContainersOverlap;
 
@@ -185,7 +185,7 @@ void FProcessorDependencySolver::FResourceUsage::HandleElementType(TMassExecutio
 }
 
 template<typename TBitSet>
-bool FProcessorDependencySolver::FResourceUsage::CanAccess(const TMassExecutionAccess<TBitSet>& StoredElements, const TMassExecutionAccess<TBitSet>& TestedElements)
+bool FMassProcessorDependencySolver::FResourceUsage::CanAccess(const TMassExecutionAccess<TBitSet>& StoredElements, const TMassExecutionAccess<TBitSet>& TestedElements)
 {
 	// see if there's an overlap of tested write operations with existing read & write operations, as well as 
 	// tested read operations with existing write operations
@@ -200,7 +200,7 @@ bool FProcessorDependencySolver::FResourceUsage::CanAccess(const TMassExecutionA
 	);
 }
 
-bool FProcessorDependencySolver::FResourceUsage::HasArchetypeConflict(TMassExecutionAccess<FResourceAccess> ElementAccess, const TArray<FMassArchetypeHandle>& InArchetypes) const
+bool FMassProcessorDependencySolver::FResourceUsage::HasArchetypeConflict(TMassExecutionAccess<FResourceAccess> ElementAccess, const TArray<FMassArchetypeHandle>& InArchetypes) const
 {
 	using UE::Mass::Private::DoArchetypeContainersOverlap;
 
@@ -227,7 +227,7 @@ bool FProcessorDependencySolver::FResourceUsage::HasArchetypeConflict(TMassExecu
 	return false;
 }
 
-bool FProcessorDependencySolver::FResourceUsage::CanAccessRequirements(const FMassExecutionRequirements& TestedRequirements, const TArray<FMassArchetypeHandle>& InArchetypes) const
+bool FMassProcessorDependencySolver::FResourceUsage::CanAccessRequirements(const FMassExecutionRequirements& TestedRequirements, const TArray<FMassArchetypeHandle>& InArchetypes) const
 {
 	bool bCanAccess = (CanAccess<FMassFragmentBitSet>(Requirements.Fragments, TestedRequirements.Fragments) || !HasArchetypeConflict(FragmentsAccess, InArchetypes))
 		&& (CanAccess<FMassChunkFragmentBitSet>(Requirements.ChunkFragments, TestedRequirements.ChunkFragments) || !HasArchetypeConflict(ChunkFragmentsAccess, InArchetypes))
@@ -237,7 +237,7 @@ bool FProcessorDependencySolver::FResourceUsage::CanAccessRequirements(const FMa
 	return bCanAccess;
 }
 
-void FProcessorDependencySolver::FResourceUsage::SubmitNode(const int32 NodeIndex, FNode& InOutNode)
+void FMassProcessorDependencySolver::FResourceUsage::SubmitNode(const int32 NodeIndex, FNode& InOutNode)
 {
 	HandleElementType<FMassFragmentBitSet>(FragmentsAccess, InOutNode.Requirements.Fragments, InOutNode, NodeIndex);
 	HandleElementType<FMassChunkFragmentBitSet>(ChunkFragmentsAccess, InOutNode.Requirements.ChunkFragments, InOutNode, NodeIndex);
@@ -250,7 +250,7 @@ void FProcessorDependencySolver::FResourceUsage::SubmitNode(const int32 NodeInde
 //----------------------------------------------------------------------//
 //  FProcessorDependencySolver::FNode
 //----------------------------------------------------------------------//
-void FProcessorDependencySolver::FNode::IncreaseWaitingNodesCount(TArrayView<FProcessorDependencySolver::FNode> InAllNodes)
+void FMassProcessorDependencySolver::FNode::IncreaseWaitingNodesCount(TArrayView<FMassProcessorDependencySolver::FNode> InAllNodes)
 {
 	// cycle-protection check. If true it means we have a cycle and the whole algorithm result will be unreliable 
 	if (TotalWaitingNodes >= FMath::Square(InAllNodes.Num()))
@@ -269,12 +269,12 @@ void FProcessorDependencySolver::FNode::IncreaseWaitingNodesCount(TArrayView<FPr
 //----------------------------------------------------------------------//
 //  FProcessorDependencySolver
 //----------------------------------------------------------------------//
-FProcessorDependencySolver::FProcessorDependencySolver(TArrayView<UMassProcessor*> InProcessors, const bool bIsGameRuntime)
+FMassProcessorDependencySolver::FMassProcessorDependencySolver(TArrayView<UMassProcessor*> InProcessors, const bool bIsGameRuntime)
 	: Processors(InProcessors)
 	, bGameRuntime(bIsGameRuntime)
 {}
 
-bool FProcessorDependencySolver::PerformSolverStep(FResourceUsage& ResourceUsage, TArray<int32>& InOutIndicesRemaining, TArray<int32>& OutNodeIndices)
+bool FMassProcessorDependencySolver::PerformSolverStep(FResourceUsage& ResourceUsage, TArray<int32>& InOutIndicesRemaining, TArray<int32>& OutNodeIndices)
 {
 	int32 AcceptedNodeIndex = INDEX_NONE;
 	int32 FallbackAcceptedNodeIndex = INDEX_NONE;
@@ -332,7 +332,7 @@ bool FProcessorDependencySolver::PerformSolverStep(FResourceUsage& ResourceUsage
 	return false;
 }
 
-void FProcessorDependencySolver::CreateSubGroupNames(FName InGroupName, TArray<FString>& SubGroupNames)
+void FMassProcessorDependencySolver::CreateSubGroupNames(FName InGroupName, TArray<FString>& SubGroupNames)
 {
 	// the function will convert composite group name into a series of progressively more precise group names
 	// so "A.B.C" will result in ["A", "A.B", "A.B.C"]
@@ -353,7 +353,7 @@ void FProcessorDependencySolver::CreateSubGroupNames(FName InGroupName, TArray<F
 	}
 }
 
-int32 FProcessorDependencySolver::CreateNodes(UMassProcessor& Processor)
+int32 FMassProcessorDependencySolver::CreateNodes(UMassProcessor& Processor)
 {
 	check(Processor.GetClass());
 	const FName ProcName = Processor.GetClass()->GetFName();
@@ -421,7 +421,7 @@ int32 FProcessorDependencySolver::CreateNodes(UMassProcessor& Processor)
 	return NodeIndex;
 }
 
-void FProcessorDependencySolver::BuildDependencies()
+void FMassProcessorDependencySolver::BuildDependencies()
 {
 	// at this point we have collected all the known processors and groups in AllNodes so we can transpose 
 	// A.ExecuteBefore(B) type of dependencies into B.ExecuteAfter(A)
@@ -535,7 +535,7 @@ void FProcessorDependencySolver::BuildDependencies()
 	}
 }
 
-void FProcessorDependencySolver::LogNode(const FNode& Node, int Indent)
+void FMassProcessorDependencySolver::LogNode(const FNode& Node, int Indent)
 {
 	using UE::Mass::Private::NameViewToString;
 
@@ -948,7 +948,7 @@ void FProcessorDependencySolver::DumpGraph(FArchive& LogFile) const
 
 #endif // 0; disabled the graph building for now, leaving the old code here for reference
 
-void FProcessorDependencySolver::Solve(TArray<FMassProcessorOrderInfo>& OutResult)
+void FMassProcessorDependencySolver::Solve(TArray<FMassProcessorOrderInfo>& OutResult)
 {
 	using UE::Mass::Private::NameViewToString;
 
@@ -1041,7 +1041,7 @@ void FProcessorDependencySolver::Solve(TArray<FMassProcessorOrderInfo>& OutResul
 	}
 }
 
-void FProcessorDependencySolver::ResolveDependencies(TArray<FMassProcessorOrderInfo>& OutResult, TSharedPtr<FMassEntityManager> EntityManager, FProcessorDependencySolver::FResult* InOutOptionalResult)
+void FMassProcessorDependencySolver::ResolveDependencies(TArray<FMassProcessorOrderInfo>& OutResult, TSharedPtr<FMassEntityManager> EntityManager, FMassProcessorDependencySolver::FResult* InOutOptionalResult)
 {
 	TRACE_CPUPROFILER_EVENT_SCOPE_STR("Mass ResolveDependencies");
 
@@ -1163,7 +1163,7 @@ void FProcessorDependencySolver::ResolveDependencies(TArray<FMassProcessorOrderI
 	}
 }
 
-bool FProcessorDependencySolver::IsResultUpToDate(const FProcessorDependencySolver::FResult& InResult, TSharedPtr<FMassEntityManager> EntityManager)
+bool FMassProcessorDependencySolver::IsResultUpToDate(const FMassProcessorDependencySolver::FResult& InResult, TSharedPtr<FMassEntityManager> EntityManager)
 {
 	if (InResult.PrunedProcessorClasses.Num() == 0 || !EntityManager || InResult.ArchetypeDataVersion == EntityManager->GetArchetypeDataVersion())
 	{
diff --git a/Engine/Plugins/Runtime/MassEntity/Source/MassEntity/Public/MassProcessingPhaseManager.h b/Engine/Plugins/Runtime/MassEntity/Source/MassEntity/Public/MassProcessingPhaseManager.h
--- a/Engine/Plugins/Runtime/MassEntity/Source/MassEntity/Public/MassProcessingPhaseManager.h
+++ b/Engine/Plugins/Runtime/MassEntity/Source/MassEntity/Public/MassProcessingPhaseManager.h
@@ -97,7 +97,7 @@ struct MASSENTITY_API FMassPhaseProcessorConfigurationHelper
 	{
 	}
 
-	void Configure(const TSharedPtr<FMassEntityManager>& EntityManager = TSharedPtr<FMassEntityManager>(), FProcessorDependencySolver::FResult* OutOptionalResult = nullptr);
+	void Configure(const TSharedPtr<FMassEntityManager>& EntityManager = TSharedPtr<FMassEntityManager>(), FMassProcessorDependencySolver::FResult* OutOptionalResult = nullptr);
 
 	UMassCompositeProcessor& PhaseProcessor;
 	const FMassProcessingPhaseConfig& PhaseConfig;
@@ -193,7 +193,7 @@ protected:
 protected:	
 	struct FPhaseGraphBuildState
 	{
-		FProcessorDependencySolver::FResult LastResult;
+		FMassProcessorDependencySolver::FResult LastResult;
 		bool bNewArchetypes = true;
 		bool bInitialized = false;
 	};
diff --git a/Engine/Plugins/Runtime/MassEntity/Source/MassEntity/Public/MassProcessorDependencySolver.h b/Engine/Plugins/Runtime/MassEntity/Source/MassEntity/Public/MassProcessorDependencySolver.h
--- a/Engine/Plugins/Runtime/MassEntity/Source/MassEntity/Public/MassProcessorDependencySolver.h
+++ b/Engine/Plugins/Runtime/MassEntity/Source/MassEntity/Public/MassProcessorDependencySolver.h
@@ -58,7 +58,7 @@ struct MASSENTITY_API FMassExecutionRequirements
 	int32 ResourcesUsedCount = INDEX_NONE;
 };
 
-struct FProcessorDependencySolver
+struct FMassProcessorDependencySolver
 {
 private:
 	struct FNode
@@ -116,7 +116,7 @@ private:
 
 		template<typename TBitSet>
 		void HandleElementType(TMassExecutionAccess<FResourceAccess>& ElementAccess
-			, const TMassExecutionAccess<TBitSet>& TestedRequirements, FProcessorDependencySolver::FNode& InOutNode, const int32 NodeIndex);
+			, const TMassExecutionAccess<TBitSet>& TestedRequirements, FMassProcessorDependencySolver::FNode& InOutNode, const int32 NodeIndex);
 
 		template<typename TBitSet>
 		static bool CanAccess(const TMassExecutionAccess<TBitSet>& StoredElements, const TMassExecutionAccess<TBitSet>& TestedElements);
@@ -143,7 +143,7 @@ public:
 		}
 	};
 
-	MASSENTITY_API FProcessorDependencySolver(TArrayView<UMassProcessor*> InProcessors, const bool bIsGameRuntime = true);
+	MASSENTITY_API FMassProcessorDependencySolver(TArrayView<UMassProcessor*> InProcessors, const bool bIsGameRuntime = true);
 	MASSENTITY_API void ResolveDependencies(TArray<FMassProcessorOrderInfo>& OutResult, TSharedPtr<FMassEntityManager> EntityManager = nullptr, FResult* InOutOptionalResult = nullptr);
 
 	MASSENTITY_API static void CreateSubGroupNames(FName InGroupName, TArray<FString>& SubGroupNames);
diff --git a/Engine/Plugins/Runtime/MassEntity/Source/MassEntityTestSuite/Private/MassDependencyTest.cpp b/Engine/Plugins/Runtime/MassEntity/Source/MassEntityTestSuite/Private/MassDependencyTest.cpp
--- a/Engine/Plugins/Runtime/MassEntity/Source/MassEntityTestSuite/Private/MassDependencyTest.cpp
+++ b/Engine/Plugins/Runtime/MassEntity/Source/MassEntityTestSuite/Private/MassDependencyTest.cpp
@@ -33,7 +33,7 @@ struct FDependencySolverBase : FAITestBase
 	void Solve()
 	{
 		Result.Reset();
-		FProcessorDependencySolver Solver(MakeArrayView((UMassProcessor**)Processors.GetData(), Processors.Num()));
+		FMassProcessorDependencySolver Solver(MakeArrayView((UMassProcessor**)Processors.GetData(), Processors.Num()));
 		Solver.ResolveDependencies(Result);
 	}
 };	
@@ -309,16 +309,16 @@ struct FGroupNamesGeneration : FAITestBase
 	{
 		TArray<FString> SubGroupNames;
 		FName EmptyName;
-		FProcessorDependencySolver::CreateSubGroupNames(EmptyName, SubGroupNames);
+		FMassProcessorDependencySolver::CreateSubGroupNames(EmptyName, SubGroupNames);
 
 		AITEST_TRUE("Empty group name is supported", SubGroupNames.Num() > 0);
 		AITEST_TRUE("Empty group name handled like any other name", SubGroupNames[0] == EmptyName.ToString());
 		
-		FProcessorDependencySolver::CreateSubGroupNames(TEXT("X"), SubGroupNames);
+		FMassProcessorDependencySolver::CreateSubGroupNames(TEXT("X"), SubGroupNames);
 		AITEST_TRUE("Trivial group name is supported", SubGroupNames.Num() > 0);
 		AITEST_TRUE("Trivial group name shouldn\'t get decorated", SubGroupNames[0] == TEXT("X"));
 
-		FProcessorDependencySolver::CreateSubGroupNames(TEXT("W.X.Y.Z"), SubGroupNames);
+		FMassProcessorDependencySolver::CreateSubGroupNames(TEXT("W.X.Y.Z"), SubGroupNames);
 		AITEST_TRUE("Complex group name should be result in a number of group names equal to group name\'s depth", SubGroupNames.Num() == 4);
 		AITEST_TRUE("Group name W.X.Y.Z should contain subgroup W", SubGroupNames.Find(TEXT("W")) != INDEX_NONE);
 		AITEST_TRUE("Group name W.X.Y.Z should contain subgroup W.X", SubGroupNames.Find(TEXT("W.X")) != INDEX_NONE);
