Addinionally marked MassEntityDebugger module as "Editor" type module for the MassEntity plugin (via MassEntity.uplugin), in addition to it already being marked as "DeveloperTool" (as suggested in #build slack channel)

diff --git a/Engine/Plugins/Runtime/MassEntity/MassEntity.uplugin b/Engine/Plugins/Runtime/MassEntity/MassEntity.uplugin
--- a/Engine/Plugins/Runtime/MassEntity/MassEntity.uplugin
+++ b/Engine/Plugins/Runtime/MassEntity/MassEntity.uplugin
@@ -33,6 +33,10 @@
 		{
 			"Name": "MassEntityDebugger",
 			"Type": "DeveloperTool"
+		},
+		{
+			"Name": "MassEntityDebugger",
+			"Type": "Editor"
 		}
 	],
 	"Plugins": [
