Modified Mass dependency solver to cut some of the fine-tunning tests when MASS_DO_PARALLEL == 0 and no parallel execution of processors will take place.

diff --git a/Engine/Plugins/Runtime/MassEntity/Source/MassEntity/Private/MassProcessingPhaseManager.cpp b/Engine/Plugins/Runtime/MassEntity/Source/MassEntity/Private/MassProcessingPhaseManager.cpp
--- a/Engine/Plugins/Runtime/MassEntity/Source/MassEntity/Private/MassProcessingPhaseManager.cpp
+++ b/Engine/Plugins/Runtime/MassEntity/Source/MassEntity/Private/MassProcessingPhaseManager.cpp
@@ -148,7 +148,10 @@ void FMassPhaseProcessorConfigurationHelper::Configure(const TSharedPtr<FMassEnt
 
 	PhaseProcessor.Populate(SortedProcessors);
 
-	PhaseProcessor.BuildFlatProcessingGraph(SortedProcessors);
+	if (Solver.IsSolvingForSingleThread() == false)
+	{
+		PhaseProcessor.BuildFlatProcessingGraph(SortedProcessors);
+	}
 
 	if (bInitializeCreatedProcessors)
 	{
diff --git a/Engine/Plugins/Runtime/MassEntity/Source/MassEntity/Private/MassProcessor.cpp b/Engine/Plugins/Runtime/MassEntity/Source/MassEntity/Private/MassProcessor.cpp
--- a/Engine/Plugins/Runtime/MassEntity/Source/MassEntity/Private/MassProcessor.cpp
+++ b/Engine/Plugins/Runtime/MassEntity/Source/MassEntity/Private/MassProcessor.cpp
@@ -364,12 +364,19 @@ void UMassCompositeProcessor::SetProcessors(TArrayView<UMassProcessor*> InProces
 
 	Populate(SortedProcessors);
 
-	BuildFlatProcessingGraph(SortedProcessors);
+	if (Solver.IsSolvingForSingleThread() == false)
+	{
+		BuildFlatProcessingGraph(SortedProcessors);
+	}
 }
 
 void UMassCompositeProcessor::BuildFlatProcessingGraph(TConstArrayView<FMassProcessorOrderInfo> SortedProcessors)
 {
 	TRACE_CPUPROFILER_EVENT_SCOPE(Mass_BuildFlatProcessingGraph);
+#if !MASS_DO_PARALLEL
+	UE_LOG(LogMass, Warning
+		, TEXT("MassCompositeProcessor::BuildFlatProcessingGraph is not expected to run in a single-threaded Mass setup. The flat graph will not be used at runtime."));
+#endif // MASS_DO_PARALLEL
 
 	FlatProcessingGraph.Reset();
 
diff --git a/Engine/Plugins/Runtime/MassEntity/Source/MassEntity/Private/MassProcessorDependencySolver.cpp b/Engine/Plugins/Runtime/MassEntity/Source/MassEntity/Private/MassProcessorDependencySolver.cpp
--- a/Engine/Plugins/Runtime/MassEntity/Source/MassEntity/Private/MassProcessorDependencySolver.cpp
+++ b/Engine/Plugins/Runtime/MassEntity/Source/MassEntity/Private/MassProcessorDependencySolver.cpp
@@ -284,7 +284,8 @@ bool FMassProcessorDependencySolver::PerformSolverStep(FResourceUsage& ResourceU
 		const int32 NodeIndex = InOutIndicesRemaining[i];
 		if (AllNodes[NodeIndex].TransientDependencies.Num() == 0)
 		{
-			if (ResourceUsage.CanAccessRequirements(AllNodes[NodeIndex].Requirements, AllNodes[NodeIndex].ValidArchetypes))
+			// if we're solving dependencies for a single thread use we don't need to fine-tune the order based on resources nor archetypes
+			if (bSingleThreadTarget || ResourceUsage.CanAccessRequirements(AllNodes[NodeIndex].Requirements, AllNodes[NodeIndex].ValidArchetypes))
 			{
 				AcceptedNodeIndex = NodeIndex;
 				break;
diff --git a/Engine/Plugins/Runtime/MassEntity/Source/MassEntity/Public/MassProcessorDependencySolver.h b/Engine/Plugins/Runtime/MassEntity/Source/MassEntity/Public/MassProcessorDependencySolver.h
--- a/Engine/Plugins/Runtime/MassEntity/Source/MassEntity/Public/MassProcessorDependencySolver.h
+++ b/Engine/Plugins/Runtime/MassEntity/Source/MassEntity/Public/MassProcessorDependencySolver.h
@@ -151,6 +151,8 @@ public:
 	/** Determines whether the dependency solving that produced InResult will produce different results if run with a given EntityManager */
 	static bool IsResultUpToDate(const FResult& InResult, TSharedPtr<FMassEntityManager> EntityManager);
 
+	bool IsSolvingForSingleThread() const { return bSingleThreadTarget; }
+
 protected:
 	// note that internals are protected rather than private to support unit testing
 
@@ -173,6 +175,14 @@ protected:
 
 	TArrayView<UMassProcessor*> Processors;
 	bool bAnyCyclesDetected = false;
+	/**
+	 * indicates whether we're generating processor order to be run in single- or multi-threaded environment (usually
+	 * this meas Dedicated Server vs Any other configuration). In Single-Threaded mode we can skip a bunch of expensive, 
+	 * fine tunning tests.
+	 * @Note: currently the value depends on MASS_DO_PARALLEL and there's no way to configure it otherwise, but there's 
+	 * nothing inherently stopping us from letting users configure it.
+	 */
+	const bool bSingleThreadTarget = bool(!MASS_DO_PARALLEL);
 	const bool bGameRuntime = true;
 	FString DependencyGraphFileName;
 	TArray<FNode> AllNodes;
