Updated MassDebugger processing graph display to work with single-threaded phases as well.

diff --git a/Engine/Plugins/Runtime/MassEntity/Source/MassEntity/Private/MassDebugger.cpp b/Engine/Plugins/Runtime/MassEntity/Source/MassEntity/Private/MassDebugger.cpp
--- a/Engine/Plugins/Runtime/MassEntity/Source/MassEntity/Private/MassDebugger.cpp
+++ b/Engine/Plugins/Runtime/MassEntity/Source/MassEntity/Private/MassDebugger.cpp
@@ -274,6 +274,11 @@ TConstArrayView<struct UMassCompositeProcessor::FDependencyNode> FMassDebugger::
 	return GraphOwner.FlatProcessingGraph;
 }
 
+TConstArrayView<TObjectPtr<UMassProcessor>> FMassDebugger::GetHostedProcessors(const UMassCompositeProcessor& GraphOwner)
+{
+	return GraphOwner.ChildPipeline.Processors;
+}
+
 FString FMassDebugger::GetRequirementsDescription(const FMassFragmentRequirements& Requirements)
 {
 	TStringBuilder<256> StringBuilder;
diff --git a/Engine/Plugins/Runtime/MassEntity/Source/MassEntity/Public/MassDebugger.h b/Engine/Plugins/Runtime/MassEntity/Source/MassEntity/Public/MassDebugger.h
--- a/Engine/Plugins/Runtime/MassEntity/Source/MassEntity/Public/MassDebugger.h
+++ b/Engine/Plugins/Runtime/MassEntity/Source/MassEntity/Public/MassDebugger.h
@@ -76,6 +76,7 @@ struct MASSENTITY_API FMassDebugger
 	static const TConstArrayView<FName> GetArchetypeDebugNames(const FMassArchetypeHandle& ArchetypeHandle);
 
 	static TConstArrayView<UMassCompositeProcessor::FDependencyNode> GetProcessingGraph(const UMassCompositeProcessor& GraphOwner);
+	static TConstArrayView<TObjectPtr<UMassProcessor>> GetHostedProcessors(const UMassCompositeProcessor& GraphOwner);
 	
 	static FString GetSingleRequirementDescription(const FMassFragmentRequirementDescription& Requirement);
 	static FString GetRequirementsDescription(const FMassFragmentRequirements& Requirements);
diff --git a/Engine/Plugins/Runtime/MassEntity/Source/MassEntityDebugger/Private/MassDebuggerModel.cpp b/Engine/Plugins/Runtime/MassEntity/Source/MassEntityDebugger/Private/MassDebuggerModel.cpp
--- a/Engine/Plugins/Runtime/MassEntity/Source/MassEntityDebugger/Private/MassDebuggerModel.cpp
+++ b/Engine/Plugins/Runtime/MassEntity/Source/MassEntityDebugger/Private/MassDebuggerModel.cpp
@@ -281,13 +281,32 @@ FMassDebuggerProcessingGraph::FMassDebuggerProcessingGraph(const FMassDebuggerMo
 #if WITH_MASSENTITY_DEBUG
 	TConstArrayView<UMassCompositeProcessor::FDependencyNode> ProcessingGraph = FMassDebugger::GetProcessingGraph(InGraphOwner);
 
-	GraphNodes.Reserve(ProcessingGraph.Num());
-	for (const UMassCompositeProcessor::FDependencyNode& Node : ProcessingGraph)
+	if (ProcessingGraph.Num() > 0)
 	{
-		check(Node.Processor);
-		const TSharedPtr<FMassDebuggerProcessorData>& ProcessorData = DebuggerModel.GetProcessorDataChecked(*Node.Processor);
-		check(ProcessorData.IsValid());
-		GraphNodes.Add(FMassDebuggerProcessingGraphNode(ProcessorData, Node));
+		GraphNodes.Reserve(ProcessingGraph.Num());
+		for (const UMassCompositeProcessor::FDependencyNode& Node : ProcessingGraph)
+		{
+			check(Node.Processor);
+			const TSharedPtr<FMassDebuggerProcessorData>& ProcessorData = DebuggerModel.GetProcessorDataChecked(*Node.Processor);
+			check(ProcessorData.IsValid());
+			GraphNodes.Add(FMassDebuggerProcessingGraphNode(ProcessorData, Node));
+		}
+	}
+	// it's possible for the graph to be empty if InGraphOwner has been populated for a single-thread execution.
+	// See if there are any processors owned by InGraphOwner.
+	else if (InGraphOwner.IsEmpty() == false)
+	{
+		TConstArrayView<TObjectPtr<UMassProcessor>> HostedProcessors = FMassDebugger::GetHostedProcessors(InGraphOwner);
+		for (const TObjectPtr<UMassProcessor>& Processor : HostedProcessors)
+		{
+			check(Processor);
+			const TSharedPtr<FMassDebuggerProcessorData>& ProcessorData = DebuggerModel.GetProcessorDataChecked(*Processor);
+			check(ProcessorData.IsValid());
+			GraphNodes.Add(FMassDebuggerProcessingGraphNode(ProcessorData));
+		}
+
+		// if we have processors, but the flat processing graph is empty, it means it's a single-threaded composite processor
+		bSingleTheadGraph = true;
 	}
 #endif // WITH_MASSENTITY_DEBUG
 }
diff --git a/Engine/Plugins/Runtime/MassEntity/Source/MassEntityDebugger/Private/MassDebuggerModel.h b/Engine/Plugins/Runtime/MassEntity/Source/MassEntityDebugger/Private/MassDebuggerModel.h
--- a/Engine/Plugins/Runtime/MassEntity/Source/MassEntityDebugger/Private/MassDebuggerModel.h
+++ b/Engine/Plugins/Runtime/MassEntity/Source/MassEntityDebugger/Private/MassDebuggerModel.h
@@ -113,7 +113,7 @@ public:
 
 struct FMassDebuggerProcessingGraphNode
 {
-	FMassDebuggerProcessingGraphNode(const TSharedPtr<FMassDebuggerProcessorData>& InProcessorData, const UMassCompositeProcessor::FDependencyNode& InProcessorNode);
+	explicit FMassDebuggerProcessingGraphNode(const TSharedPtr<FMassDebuggerProcessorData>& InProcessorData, const UMassCompositeProcessor::FDependencyNode& InProcessorNode = UMassCompositeProcessor::FDependencyNode());
 	
 	FText GetLabel() const;
 
@@ -129,6 +129,7 @@ struct FMassDebuggerProcessingGraph
 
 	FString Label;
 	TArray<FMassDebuggerProcessingGraphNode> GraphNodes;
+	bool bSingleTheadGraph = !bool(MASS_DO_PARALLEL);
 };
 
 
