Marked MassEntity plugin as Beta

diff --git a/Engine/Plugins/Runtime/MassEntity/MassEntity.uplugin b/Engine/Plugins/Runtime/MassEntity/MassEntity.uplugin
--- a/Engine/Plugins/Runtime/MassEntity/MassEntity.uplugin
+++ b/Engine/Plugins/Runtime/MassEntity/MassEntity.uplugin
@@ -11,8 +11,7 @@
 	"MarketplaceURL": "",
 	"SupportURL": "",
 	"CanContainContent": true,
-	"IsBetaVersion": false,
-	"IsExperimentalVersion": true,
+	"IsBetaVersion": true,
 	"Installed": false,
 	"EnabledByDefault": false,
 	"Modules": [
