Fix actor pooling logic when the entity was destroyed. The code was calling the wrong method to unregister to release the entity fragments. Then when it tried to do it it was too late after the entity was already destroyed.
Supporting MassAgentComponent with empty MassEntityConfig, this is usefull when the actor is a puppet.
Added a destroy entity debug command
Fix from an UDN case 00482111

diff --git a/Engine/Plugins/Runtime/MassEntity/Source/MassEntity/Private/MassDebugger.cpp b/Engine/Plugins/Runtime/MassEntity/Source/MassEntity/Private/MassDebugger.cpp
--- a/Engine/Plugins/Runtime/MassEntity/Source/MassEntity/Private/MassDebugger.cpp
+++ b/Engine/Plugins/Runtime/MassEntity/Source/MassEntity/Private/MassDebugger.cpp
@@ -13,6 +13,8 @@
 #include "Misc/OutputDevice.h"
 #include "Engine/World.h"
 #include "Engine/Engine.h"
+#include "MassEntityUtils.h"
+#include "MassCommandBuffer.h"
 
 
 namespace UE::Mass::Debug
@@ -189,6 +191,42 @@ namespace UE::Mass::Debug
 				PrintKnownTypes(FMassChunkFragmentBitSet::DebugGetAllStructTypes());
 			}));
 
+	static FAutoConsoleCommandWithWorldAndArgs DestroyEntity(
+		TEXT("ai.debug.mass.DestroyEntity"),
+		TEXT("ID of a lightweight entity that we want to destroy.")
+		TEXT("Usage: \"ai.debug.mass.DestoryEntity <Entity>\""),
+		FConsoleCommandWithWorldAndArgsDelegate::CreateLambda([](const TArray<FString>& Args, UWorld* World)
+	{
+		if (Args.Num() != 1)
+		{
+			UE_LOG(LogConsoleResponse, Display, TEXT("Error: Expecting 1 parameter"));
+			return;
+		}
+
+		int32 ID = INDEX_NONE;
+		if (!LexTryParseString<int32>(ID, *Args[0]))
+		{
+			UE_LOG(LogConsoleResponse, Display, TEXT("Error: parameter must be an integer"));
+			return;
+		}
+
+		if (!World)
+		{
+			UE_LOG(LogConsoleResponse, Display, TEXT("Error: invalid world"));
+			return;
+		}
+
+		FMassEntityManager& EntityManager = UE::Mass::Utils::GetEntityManagerChecked(*World);
+		FMassEntityHandle EntityToDestroy = EntityManager.DebugGetEntityIndexHandle(ID);
+		if (!EntityToDestroy.IsSet())
+		{
+			UE_LOG(LogConsoleResponse, Display, TEXT("Error: cannot find entity for this index"));
+			return;
+		}
+
+		EntityManager.Defer().DestroyEntity(EntityToDestroy);
+	}));
+
 } // namespace UE::Mass::Debug
 
 //----------------------------------------------------------------------//
diff --git a/Engine/Plugins/Runtime/MassGameplay/Source/MassActors/Private/MassActorSpawnerSubsystem.cpp b/Engine/Plugins/Runtime/MassGameplay/Source/MassActors/Private/MassActorSpawnerSubsystem.cpp
--- a/Engine/Plugins/Runtime/MassGameplay/Source/MassActors/Private/MassActorSpawnerSubsystem.cpp
+++ b/Engine/Plugins/Runtime/MassGameplay/Source/MassActors/Private/MassActorSpawnerSubsystem.cpp
@@ -59,10 +59,8 @@ void UMassActorSpawnerSubsystem::DestroyActor(AActor* Actor, bool bImmediate /*=
 	// We need to unregister immediately MassAgentComponent as it will become out of sync with mass
 	if (UMassAgentComponent* AgentComp = Actor->FindComponentByClass<UMassAgentComponent>())
 	{
-		if (AgentComp->IsRegistered())
-		{
-			AgentComp->UnregisterComponent();
-		}
+		// All we want here it to unregister with the subsysem, not unregister the component as we want to keep it for futher usage if we put the actor in the pool.
+		AgentComp->UnregisterWithAgentSubsystem();
 	}
 
 	if(bImmediate)
diff --git a/Engine/Plugins/Runtime/MassGameplay/Source/MassSpawner/Private/MassEntityConfigAsset.cpp b/Engine/Plugins/Runtime/MassGameplay/Source/MassSpawner/Private/MassEntityConfigAsset.cpp
--- a/Engine/Plugins/Runtime/MassGameplay/Source/MassSpawner/Private/MassEntityConfigAsset.cpp
+++ b/Engine/Plugins/Runtime/MassGameplay/Source/MassSpawner/Private/MassEntityConfigAsset.cpp
@@ -96,10 +96,10 @@ const FMassEntityTemplate& FMassEntityConfig::GetOrCreateEntityTemplate(const UW
 	BuildContext.BuildFromTraits(CombinedTraits, World);
 	Template.SetTemplateName(ConfigOwner.GetName());
 
-	if (ensureMsgf(!Template.IsEmpty(), TEXT("Need at least one fragment to create an Archetype")))
-	{
-		TemplateRegistry.InitializeEntityTemplate(Template);
-	}
+	// It is ok to have an empty template, 
+    // but be aware there will be an error if you try to create an entity with it
+    // as there will be no archetype associated with this template...
+	TemplateRegistry.InitializeEntityTemplate(Template);
 
 	return Template;
 }
diff --git a/Engine/Plugins/Runtime/MassGameplay/Source/MassSpawner/Private/MassEntityTemplateRegistry.cpp b/Engine/Plugins/Runtime/MassGameplay/Source/MassSpawner/Private/MassEntityTemplateRegistry.cpp
--- a/Engine/Plugins/Runtime/MassGameplay/Source/MassSpawner/Private/MassEntityTemplateRegistry.cpp
+++ b/Engine/Plugins/Runtime/MassGameplay/Source/MassSpawner/Private/MassEntityTemplateRegistry.cpp
@@ -147,8 +147,12 @@ bool UMassEntityTemplateRegistry::BuildTemplateImpl(const FStructToTemplateBuild
 
 void UMassEntityTemplateRegistry::InitializeEntityTemplate(FMassEntityTemplate& OutTemplate) const
 {
-	// expected to be ensured by the caller
-	check(!OutTemplate.IsEmpty());
+	if (OutTemplate.IsEmpty() == true)
+	{
+		// nothing to initialize, return.
+		// this will leave the ArchetypeHandle invalid
+		return;
+	}
 
 	UWorld* World = GetWorld();
 	check(World);
