Fixed MassDebuggerModule not being able to load in cooked editors

diff --git a/Engine/Plugins/Runtime/MassEntity/Source/MassEntityDebugger/Private/MassDebuggerModule.cpp b/Engine/Plugins/Runtime/MassEntity/Source/MassEntityDebugger/Private/MassDebuggerModule.cpp
--- a/Engine/Plugins/Runtime/MassEntity/Source/MassEntityDebugger/Private/MassDebuggerModule.cpp
+++ b/Engine/Plugins/Runtime/MassEntity/Source/MassEntityDebugger/Private/MassDebuggerModule.cpp
@@ -1,12 +1,11 @@
 // Copyright Epic Games, Inc. All Rights Reserved.
 
 #include "MassProcessingTypes.h"
-
-#if WITH_MASSENTITY_DEBUG
-
 #include "IMassDebuggerModule.h"
 #include "Features/IModularFeatures.h"
 
+#if WITH_MASSENTITY_DEBUG
+
 #include "WorkspaceMenuStructure.h"
 #include "WorkspaceMenuStructureModule.h"
 
@@ -73,6 +72,15 @@ namespace UE::Mass::Debug::Private
 	);
 
 } // UE::Mass::Debug::Private
+
+#else // WITH_MASSENTITY_DEBUG
+
+class FMassDebuggerModule : public IMassDebuggerModule, public IModularFeature
+{
+};
+
+IMPLEMENT_MODULE(FMassDebuggerModule, MassEntityDebugger)
+
 #endif // WITH_MASSENTITY_DEBUG
 
 #undef LOCTEXT_NAMESPACE
