Fixed FMassCommandBuffer::Flush not maintaining the order of operations of the same type (i.e. the sorting is not stable).

diff --git a/Engine/Plugins/Runtime/MassEntity/Source/MassEntity/Private/MassCommandBuffer.cpp b/Engine/Plugins/Runtime/MassEntity/Source/MassEntity/Private/MassCommandBuffer.cpp
--- a/Engine/Plugins/Runtime/MassEntity/Source/MassEntity/Private/MassCommandBuffer.cpp
+++ b/Engine/Plugins/Runtime/MassEntity/Source/MassEntity/Private/MassCommandBuffer.cpp
@@ -130,7 +130,7 @@ void FMassCommandBuffer::Flush(FMassEntityManager& EntityManager)
 			const FMassBatchedCommand* Command = CommandInstances[i];
 			CommandsOrder.Add(FBatchedCommandsSortedIndex(i, (Command && Command->HasWork())? CommandTypeOrder[(int)Command->GetOperationType()] : MAX_int32));
 		}
-		CommandsOrder.Sort();
+		CommandsOrder.StableSort();
 				
 		for (int32 k = 0; k < CommandsOrder.Num() && CommandsOrder[k].IsValid(); ++k)
 		{
