Added option to check fragment/tag types used in deferred commands to be checked at runtime instead of compile time.

diff --git a/Engine/Plugins/Runtime/MassEntity/Source/MassEntity/Public/MassCommandBuffer.h b/Engine/Plugins/Runtime/MassEntity/Source/MassEntity/Public/MassCommandBuffer.h
--- a/Engine/Plugins/Runtime/MassEntity/Source/MassEntity/Public/MassCommandBuffer.h
+++ b/Engine/Plugins/Runtime/MassEntity/Source/MassEntity/Public/MassCommandBuffer.h
@@ -54,14 +54,28 @@ public:
 	void AddFragment(FMassEntityHandle Entity)
 	{
 		static_assert(TIsDerivedFrom<T, FMassFragment>::IsDerived, "Given struct type is not a valid fragment type.");
-		PushCommand<FMassCommandAddFragments<T>>(Entity);
+		PushCommand<FMassCommandAddFragmentsInternal<EMassCommandCheckTime::CompileTimeCheck, T>>(Entity);
+	}
+
+	template<typename T>
+	void AddFragment_RuntimeCheck(FMassEntityHandle Entity)
+	{
+		checkf(T::StaticStruct()->IsChildOf(FMassFragment::StaticStruct()), TEXT("Given struct type is not a valid fragment type."));
+		PushCommand<FMassCommandAddFragmentsInternal<EMassCommandCheckTime::RuntimeCheck, T>>(Entity);
 	}
 
 	template<typename T>
 	void RemoveFragment(FMassEntityHandle Entity)
 	{
 		static_assert(TIsDerivedFrom<T, FMassFragment>::IsDerived, "Given struct type is not a valid fragment type.");
-		PushCommand<FMassCommandRemoveFragments<T>>(Entity);
+		PushCommand<FMassCommandRemoveFragmentsInternal<EMassCommandCheckTime::CompileTimeCheck, T>>(Entity);
+	}
+
+	template<typename T>
+	void RemoveFragment_RuntimeCheck(FMassEntityHandle Entity)
+	{
+		checkf(T::StaticStruct()->IsChildOf(FMassFragment::StaticStruct()), TEXT("Given struct type is not a valid fragment type."));
+		PushCommand<FMassCommandRemoveFragmentsInternal<EMassCommandCheckTime::RuntimeCheck, T>>(Entity);
 	}
 
 	/** the convenience function equivalent to calling PushCommand<FMassCommandAddTag<T>>(Entity) */
@@ -69,7 +83,14 @@ public:
 	void AddTag(FMassEntityHandle Entity)
 	{
 		static_assert(TIsDerivedFrom<T, FMassTag>::IsDerived, "Given struct type is not a valid tag type.");
-		PushCommand<FMassCommandAddTag<T>>(Entity);
+		PushCommand<FMassCommandAddTagInternal<EMassCommandCheckTime::CompileTimeCheck, T>>(Entity);
+	}
+
+	template<typename T>
+	void AddTag_RuntimeCheck(FMassEntityHandle Entity)
+	{
+		checkf(T::StaticStruct()->IsChildOf(FMassTag::StaticStruct()), TEXT("Given struct type is not a valid tag type."));
+		PushCommand<FMassCommandAddTagInternal<EMassCommandCheckTime::RuntimeCheck, T>>(Entity);
 	}
 
 	/** the convenience function equivalent to calling PushCommand<FMassCommandRemoveTag<T>>(Entity) */
@@ -77,7 +98,14 @@ public:
 	void RemoveTag(FMassEntityHandle Entity)
 	{
 		static_assert(TIsDerivedFrom<T, FMassTag>::IsDerived, "Given struct type is not a valid tag type.");
-		PushCommand<FMassCommandRemoveTag<T>>(Entity);
+		PushCommand<FMassCommandRemoveTagInternal<EMassCommandCheckTime::CompileTimeCheck, T>>(Entity);
+	}
+
+	template<typename T>
+	void RemoveTag_RuntimeCheck(FMassEntityHandle Entity)
+	{
+		checkf(T::StaticStruct()->IsChildOf(FMassTag::StaticStruct()), TEXT("Given struct type is not a valid tag type."));
+		PushCommand<FMassCommandRemoveTagInternal<EMassCommandCheckTime::RuntimeCheck, T>>(Entity);
 	}
 
 	/** the convenience function equivalent to calling PushCommand<FMassCommandSwapTags<TOld, TNew>>(Entity)  */
@@ -86,7 +114,15 @@ public:
 	{
 		static_assert(TIsDerivedFrom<TOld, FMassTag>::IsDerived, "Given struct type is not a valid tag type.");
 		static_assert(TIsDerivedFrom<TNew, FMassTag>::IsDerived, "Given struct type is not a valid tag type.");
-		PushCommand<FMassCommandSwapTags<TOld, TNew>>(Entity);
+		PushCommand<FMassCommandSwapTagsInternal<EMassCommandCheckTime::CompileTimeCheck, TOld, TNew>>(Entity);
+	}
+
+	template<typename TOld, typename TNew>
+	void SwapTags_RuntimeCheck(FMassEntityHandle Entity)
+	{
+		checkf(TOld::StaticStruct()->IsChildOf(FMassTag::StaticStruct()), TEXT("Given struct type is not a valid tag type."));
+		checkf(TNew::StaticStruct()->IsChildOf(FMassTag::StaticStruct()), TEXT("Given struct type is not a valid tag type."));
+		PushCommand<FMassCommandSwapTagsInternal<EMassCommandCheckTime::RuntimeCheck, TOld, TNew>>(Entity);
 	}
 
 	void DestroyEntity(FMassEntityHandle Entity)
diff --git a/Engine/Plugins/Runtime/MassEntity/Source/MassEntity/Public/MassCommands.h b/Engine/Plugins/Runtime/MassEntity/Source/MassEntity/Public/MassCommands.h
--- a/Engine/Plugins/Runtime/MassEntity/Source/MassEntity/Public/MassCommands.h
+++ b/Engine/Plugins/Runtime/MassEntity/Source/MassEntity/Public/MassCommands.h
@@ -27,8 +27,60 @@ enum class EMassCommandOperationType : uint8
 	MAX
 };
 
+enum class EMassCommandCheckTime : bool
+{
+	RuntimeCheck = true,
+	CompileTimeCheck = false
+};
+
+#if CSV_PROFILER || WITH_MASSENTITY_DEBUG
+#	define DEBUG_NAME(Name) , FName(TEXT(Name))
+#else
+#	define DEBUG_NAME(Name)
+#endif // CSV_PROFILER || WITH_MASSENTITY_DEBUG
+
+namespace UE::Mass::Utils
+{
+	template<typename BitSetType, EMassCommandCheckTime CheckTime, typename... TTypes>
+	BitSetType ConstructBitSet()
+	{
+		if constexpr (CheckTime == EMassCommandCheckTime::RuntimeCheck)
+		{
+			return BitSetType({ TTypes::StaticStruct()... });
+		}
+		else
+		{
+			BitSetType Result;
+			UE::Mass::TMultiTypeList<TTypes...>::PopulateBitSet(Result);
+			return Result;
+		}
+	}
+
+	template<EMassCommandCheckTime CheckTime, typename... TTypes>
+	FMassFragmentBitSet ConstructFragmentBitSet()
+	{
+		return ConstructBitSet<FMassFragmentBitSet, CheckTime, TTypes...>();
+	}
+
+	template<EMassCommandCheckTime CheckTime, typename... TTypes>
+	FMassTagBitSet ConstructTagBitSet()
+	{
+		return ConstructBitSet<FMassTagBitSet, CheckTime, TTypes...>();
+	}
+} // namespace UE::Mass::Utils
+
 struct MASSENTITY_API FMassBatchedCommand
 {
+	FMassBatchedCommand() = default;
+	explicit FMassBatchedCommand(EMassCommandOperationType OperationType)
+		: OperationType(OperationType)
+	{}
+#if CSV_PROFILER || WITH_MASSENTITY_DEBUG
+	FMassBatchedCommand(EMassCommandOperationType OperationType, FName DebugName)
+		: OperationType(OperationType)
+		, DebugName(DebugName)
+	{}
+#endif // CSV_PROFILER || WITH_MASSENTITY_DEBUG
 	virtual ~FMassBatchedCommand() {}
 
 	virtual void Execute(FMassEntityManager& System) const = 0;
@@ -72,6 +124,17 @@ struct FMassBatchedEntityCommand : public FMassBatchedCommand
 {
 	using Super = FMassBatchedCommand;
 
+	FMassBatchedEntityCommand() = default;
+	explicit FMassBatchedEntityCommand(EMassCommandOperationType OperationType)
+		: Super(OperationType)
+	{}
+
+#if CSV_PROFILER || WITH_MASSENTITY_DEBUG
+	FMassBatchedEntityCommand(EMassCommandOperationType OperationType, FName DebugName)
+		: Super(OperationType, DebugName)
+	{}
+#endif // CSV_PROFILER || WITH_MASSENTITY_DEBUG
+
 	void Add(FMassEntityHandle Entity)
 	{
 		UE_MT_SCOPED_READ_ACCESS(EntitiesAccessDetector);
@@ -133,14 +196,15 @@ protected:
 //////////////////////////////////////////////////////////////////////
 // Simple fragment composition change
 
-template<typename... TTypes>
-struct FMassCommandAddFragments : public FMassBatchedEntityCommand
+template<EMassCommandCheckTime CheckTime, typename... TTypes>
+struct FMassCommandAddFragmentsInternal : public FMassBatchedEntityCommand
 {
-	FMassCommandAddFragments()
-	{
-		UE::Mass::TMultiTypeList<TTypes...>::PopulateBitSet(FragmentsAffected);
-		OperationType = EMassCommandOperationType::Add;
-	}
+	using Super = FMassBatchedEntityCommand; 
+	FMassCommandAddFragmentsInternal()
+		: Super(EMassCommandOperationType::Add)
+		, FragmentsAffected(UE::Mass::Utils::ConstructFragmentBitSet<CheckTime, TTypes...>())
+	{}
+
 protected:
 	virtual void Execute(FMassEntityManager& System) const override
 	{
@@ -153,13 +217,17 @@ protected:
 };
 
 template<typename... TTypes>
-struct FMassCommandRemoveFragments : public FMassBatchedEntityCommand
+using FMassCommandAddFragments = FMassCommandAddFragmentsInternal<EMassCommandCheckTime::CompileTimeCheck, TTypes...>;
+
+template<EMassCommandCheckTime CheckTime, typename... TTypes>
+struct FMassCommandRemoveFragmentsInternal : public FMassBatchedEntityCommand
 {
-	FMassCommandRemoveFragments()
-	{
-		UE::Mass::TMultiTypeList<TTypes...>::PopulateBitSet(FragmentsAffected);
-		OperationType = EMassCommandOperationType::Remove;
-	}
+	using Super = FMassBatchedEntityCommand;
+	FMassCommandRemoveFragmentsInternal()
+		: Super(EMassCommandOperationType::Remove)
+		, FragmentsAffected(UE::Mass::Utils::ConstructFragmentBitSet<CheckTime, TTypes...>())
+	{}
+
 protected:
 	virtual void Execute(FMassEntityManager& System) const override
 	{
@@ -171,17 +239,32 @@ protected:
 	FMassFragmentBitSet FragmentsAffected;
 };
 
+template<typename... TTypes>
+using FMassCommandRemoveFragments = FMassCommandRemoveFragmentsInternal<EMassCommandCheckTime::CompileTimeCheck, TTypes...>;
+
 //////////////////////////////////////////////////////////////////////
 // Simple tag composition change
 
 struct FMassCommandChangeTags : public FMassBatchedEntityCommand
 {
 	using Super = FMassBatchedEntityCommand;
-
 	FMassCommandChangeTags()
-	{
-		OperationType = EMassCommandOperationType::ChangeComposition;
-	}
+		: Super(EMassCommandOperationType::ChangeComposition)
+	{}
+
+	FMassCommandChangeTags(EMassCommandOperationType OperationType, FMassTagBitSet TagsToAdd, FMassTagBitSet TagsToRemove)
+		: Super(OperationType)
+		, TagsToAdd(TagsToAdd)
+		, TagsToRemove(TagsToRemove)
+	{}
+
+#if CSV_PROFILER || WITH_MASSENTITY_DEBUG
+	FMassCommandChangeTags(EMassCommandOperationType OperationType, FMassTagBitSet TagsToAdd, FMassTagBitSet TagsToRemove, FName DebugName)
+		: Super(OperationType, DebugName)
+		, TagsToAdd(TagsToAdd)
+		, TagsToRemove(TagsToRemove)
+	{}
+#endif // CSV_PROFILER || WITH_MASSENTITY_DEBUG
 
 protected:
 	virtual void Execute(FMassEntityManager& System) const override
@@ -202,46 +285,54 @@ protected:
 	FMassTagBitSet TagsToRemove;
 };
 
-template<typename T>
-struct FMassCommandAddTag : public FMassCommandChangeTags
+template<EMassCommandCheckTime CheckTime, typename T>
+struct FMassCommandAddTagInternal : public FMassCommandChangeTags
 {
-	FMassCommandAddTag()
-	{
-		OperationType = EMassCommandOperationType::Add;
-		TagsToAdd.Add<T>();
-#if CSV_PROFILER || WITH_MASSENTITY_DEBUG
-		DebugName = TEXT("CommandAddTag");
-#endif // CSV_PROFILER || WITH_MASSENTITY_DEBUG
-	}
+	using Super = FMassCommandChangeTags;
+	FMassCommandAddTagInternal()
+		: Super(
+			EMassCommandOperationType::Add, 
+			UE::Mass::Utils::ConstructTagBitSet<CheckTime, T>(),
+			{} 
+			DEBUG_NAME("CommandAddTag"))
+	{}
 };
 
 template<typename T>
-struct FMassCommandRemoveTag : public FMassCommandChangeTags
+using FMassCommandAddTag = FMassCommandAddTagInternal<EMassCommandCheckTime::CompileTimeCheck, T>;
+
+template<EMassCommandCheckTime CheckTime, typename T>
+struct FMassCommandRemoveTagInternal : public FMassCommandChangeTags
 {
-	FMassCommandRemoveTag()
-	{
-		OperationType = EMassCommandOperationType::Remove;
-		TagsToRemove.Add<T>();
-#if CSV_PROFILER || WITH_MASSENTITY_DEBUG
-		DebugName = TEXT("RemoveAddTag");
-#endif // CSV_PROFILER || WITH_MASSENTITY_DEBUG
-	}
+	using Super = FMassCommandChangeTags;
+	FMassCommandRemoveTagInternal()
+		: Super(
+			EMassCommandOperationType::Remove, 
+			{}, 
+			UE::Mass::Utils::ConstructTagBitSet<CheckTime, T>()
+			DEBUG_NAME("CommandRemoveTag"))
+	{}
 };
 
-template<typename TOld, typename TNew>
-struct FMassCommandSwapTags : public FMassCommandChangeTags
+template<typename T>
+using FMassCommandAddTag = FMassCommandAddTagInternal<EMassCommandCheckTime::CompileTimeCheck, T>;
+
+template<EMassCommandCheckTime CheckTime, typename TOld, typename TNew>
+struct FMassCommandSwapTagsInternal : public FMassCommandChangeTags
 {
-	FMassCommandSwapTags()
-	{
-		OperationType = EMassCommandOperationType::ChangeComposition;
-		TagsToRemove.Add<TOld>();
-		TagsToAdd.Add<TNew>();
-#if CSV_PROFILER || WITH_MASSENTITY_DEBUG
-		DebugName = TEXT("SwapAddTags");
-#endif // CSV_PROFILER || WITH_MASSENTITY_DEBUG
-	}
+	using Super = FMassCommandChangeTags;
+	FMassCommandSwapTagsInternal()
+		: Super(
+			EMassCommandOperationType::ChangeComposition,
+			UE::Mass::Utils::ConstructTagBitSet<CheckTime, TNew>(),
+			UE::Mass::Utils::ConstructTagBitSet<CheckTime, TOld>()
+			DEBUG_NAME("CommandSwapTags"))
+	{}
 };
 
+template<typename TOld, typename TNew>
+using FMassCommandSwapTags = FMassCommandSwapTagsInternal<EMassCommandCheckTime::CompileTimeCheck, TOld, TNew>;
+
 //////////////////////////////////////////////////////////////////////
 // Struct Instances adding and setting
 
@@ -251,13 +342,9 @@ struct FMassCommandAddFragmentInstances : public FMassBatchedEntityCommand
 	using Super = FMassBatchedEntityCommand;
 
 	FMassCommandAddFragmentInstances()
-	{
-		OperationType = EMassCommandOperationType::Set;
-		UE::Mass::TMultiTypeList<TOthers...>::PopulateBitSet(FragmentsAffected);
-#if CSV_PROFILER || WITH_MASSENTITY_DEBUG
-		DebugName = TEXT("AddFragmentInstanceList");
-#endif // CSV_PROFILER || WITH_MASSENTITY_DEBUG
-	}
+		: Super(EMassCommandOperationType::Set DEBUG_NAME("AddFragmentInstanceList"))
+		, FragmentsAffected(UE::Mass::Utils::ConstructFragmentBitSet<EMassCommandCheckTime::CompileTimeCheck, TOthers...>())
+	{}
 
 	void Add(FMassEntityHandle Entity, TOthers... InFragments)
 	{
@@ -335,13 +422,9 @@ struct FMassCommandBuildEntityWithSharedFragments : public FMassBatchedCommand
 	using Super = FMassBatchedCommand;
 
 	FMassCommandBuildEntityWithSharedFragments()
-	{
-		OperationType = EMassCommandOperationType::Create;
-		UE::Mass::TMultiTypeList<TOthers...>::PopulateBitSet(FragmentsAffected);
-#if CSV_PROFILER || WITH_MASSENTITY_DEBUG
-		DebugName = TEXT("FMassCommandBuildEntityWithSharedFragments");
-#endif // CSV_PROFILER || WITH_MASSENTITY_DEBUG
-	}
+		: Super(EMassCommandOperationType::Create DEBUG_NAME("FMassCommandBuildEntityWithSharedFragments"))
+		, FragmentsAffected(UE::Mass::Utils::ConstructFragmentBitSet<EMassCommandCheckTime::CompileTimeCheck, TOthers...>())
+	{}
 
 	void Add(FMassEntityHandle Entity, FMassArchetypeSharedFragmentValues&& InSharedFragments, TOthers... InFragments)
 	{
@@ -447,12 +530,8 @@ struct FMassDeferredCommand : public FMassBatchedCommand
 	using FExecFunction = TFunction<void(FMassEntityManager& System)>;
 
 	FMassDeferredCommand()
-	{
-		OperationType = OpType;
-#if CSV_PROFILER || WITH_MASSENTITY_DEBUG
-		DebugName = TEXT("BatchedDeferredCommand");
-#endif // CSV_PROFILER || WITH_MASSENTITY_DEBUG
-	};
+		: Super(OpType DEBUG_NAME("BatchedDeferredCommand"))
+	{}
 
 	void Add(FExecFunction&& ExecFunction)
 	{
@@ -504,3 +583,5 @@ using FMassDeferredRemoveCommand = FMassDeferredCommand<EMassCommandOperationTyp
 using FMassDeferredChangeCompositionCommand = FMassDeferredCommand<EMassCommandOperationType::ChangeComposition>;
 using FMassDeferredSetCommand = FMassDeferredCommand<EMassCommandOperationType::Set>;
 using FMassDeferredDestroyCommand = FMassDeferredCommand<EMassCommandOperationType::Destroy>;
+
+#undef DEBUG_NAME
\ No newline at end of file
diff --git a/Engine/Plugins/Runtime/MassEntity/Source/MassEntity/Public/MassEntityTypes.h b/Engine/Plugins/Runtime/MassEntity/Source/MassEntity/Public/MassEntityTypes.h
--- a/Engine/Plugins/Runtime/MassEntity/Source/MassEntity/Public/MassEntityTypes.h
+++ b/Engine/Plugins/Runtime/MassEntity/Source/MassEntity/Public/MassEntityTypes.h
@@ -432,7 +432,7 @@ namespace UE::Mass
 		};
 
 		template<typename TBitSetType>
-		static void PopulateBitSet(TBitSetType& OutBitSet)
+		constexpr static void PopulateBitSet(TBitSetType& OutBitSet)
 		{
 			Super::PopulateBitSet(OutBitSet);
 			OutBitSet += TBitSetType::template GetTypeBitSet<FType>();
@@ -450,7 +450,7 @@ namespace UE::Mass
 		};
 
 		template<typename TBitSetType>
-		static void PopulateBitSet(TBitSetType& OutBitSet)
+		constexpr static void PopulateBitSet(TBitSetType& OutBitSet)
 		{
 			OutBitSet += TBitSetType::template GetTypeBitSet<FType>();
 		}
