Decoupled most MassEntity unit tests from UWorld (i.e. no longer requiring its existence)

diff --git a/Engine/Plugins/Runtime/MassEntity/Source/MassEntityTestSuite/Private/MassExecutionTest.cpp b/Engine/Plugins/Runtime/MassEntity/Source/MassEntityTestSuite/Private/MassExecutionTest.cpp
--- a/Engine/Plugins/Runtime/MassEntity/Source/MassEntityTestSuite/Private/MassExecutionTest.cpp
+++ b/Engine/Plugins/Runtime/MassEntity/Source/MassEntityTestSuite/Private/MassExecutionTest.cpp
@@ -19,9 +19,7 @@ struct FExecution_Setup : FExecutionTestBase
 {
 	virtual bool InstantTest() override
 	{
-		AITEST_NOT_NULL("World needs to exist for the test to be performed", World);
 		AITEST_NOT_NULL("EntitySubsystem needs to be created for the test to be performed", EntityManager.Get());
-
 		return true;
 	}
 };
diff --git a/Engine/Plugins/Runtime/MassEntity/Source/MassEntityTestSuite/Private/MassSystemRequirementTest.cpp b/Engine/Plugins/Runtime/MassEntity/Source/MassEntityTestSuite/Private/MassSystemRequirementTest.cpp
--- a/Engine/Plugins/Runtime/MassEntity/Source/MassEntityTestSuite/Private/MassSystemRequirementTest.cpp
+++ b/Engine/Plugins/Runtime/MassEntity/Source/MassEntityTestSuite/Private/MassSystemRequirementTest.cpp
@@ -6,25 +6,19 @@
 #include "MassProcessorDependencySolver.h"
 #include "MassEntityTestTypes.h"
 
-// D:\p4\Starship_review\Engine\Source\Runtime\NavigationSystem\Public\NavigationSystem.h
-//#include "MassNavigationSubsystem.h"
-
 #define LOCTEXT_NAMESPACE "MassTest"
 
-//template<>
-//struct TMassExternalSystemGetter<FMassEntityManager>
-//{
-//	
-//};
-
 namespace FMassSystemRequirementTest
 {
 
 struct FSystemRequirementTestBase : FEntityTestBase
 {
 	using Super = FEntityTestBase;
+	UWorld* World = nullptr;
+
 	virtual bool SetUp() override
 	{		
+		World = FAITestHelpers::GetWorld();
 		return Super::SetUp();
 	}
 };	
diff --git a/Engine/Plugins/Runtime/MassEntity/Source/MassEntityTestSuite/Private/MassTestTypes.cpp b/Engine/Plugins/Runtime/MassEntity/Source/MassEntityTestSuite/Private/MassTestTypes.cpp
--- a/Engine/Plugins/Runtime/MassEntity/Source/MassEntityTestSuite/Private/MassTestTypes.cpp
+++ b/Engine/Plugins/Runtime/MassEntity/Source/MassEntityTestSuite/Private/MassTestTypes.cpp
@@ -10,7 +10,6 @@
 //----------------------------------------------------------------------//
 bool FExecutionTestBase::SetUp()
 {
-	World = FAITestHelpers::GetWorld();
 	EntityManager = MakeShareable(new FMassEntityManager);
 	EntityManager->SetDebugName(TEXT("MassEntityTestSuite"));
 	EntityManager->Initialize();
diff --git a/Engine/Plugins/Runtime/MassEntity/Source/MassEntityTestSuite/Public/MassEntityTestTypes.h b/Engine/Plugins/Runtime/MassEntity/Source/MassEntityTestSuite/Public/MassEntityTestTypes.h
--- a/Engine/Plugins/Runtime/MassEntity/Source/MassEntityTestSuite/Public/MassEntityTestTypes.h
+++ b/Engine/Plugins/Runtime/MassEntity/Source/MassEntityTestSuite/Public/MassEntityTestTypes.h
@@ -13,7 +13,6 @@
 #include "MassEntityTestTypes.generated.h"
 
 
-class UWorld;
 struct FMassEntityManager;
 
 USTRUCT()
@@ -192,7 +191,6 @@ public:
 struct MASSENTITYTESTSUITE_API FExecutionTestBase : FAITestBase
 {
 	TSharedPtr<FMassEntityManager> EntityManager;
-	UWorld* World = nullptr;
 
 	virtual bool SetUp() override;
 };
