Updated MASS so it supports processors in the editor.

A new processor execution type was added specifically for the editor. This allows the editor to tick MASS processors in the background and during PIE/SIE.

Parts of the work in this change list is based on work provided by [REDACTED_PERSON].

diff --git a/Engine/Plugins/Runtime/MassEntity/Source/MassEntity/Private/MassEntityUtils.cpp b/Engine/Plugins/Runtime/MassEntity/Source/MassEntity/Private/MassEntityUtils.cpp
--- a/Engine/Plugins/Runtime/MassEntity/Source/MassEntity/Private/MassEntityUtils.cpp
+++ b/Engine/Plugins/Runtime/MassEntity/Source/MassEntity/Private/MassEntityUtils.cpp
@@ -10,29 +10,30 @@
 
 namespace UE::Mass::Utils
 {
-
 EProcessorExecutionFlags GetProcessorExecutionFlagsForWold(const UWorld& World)
 {
-	EProcessorExecutionFlags ExecutionFlags = EProcessorExecutionFlags::None;
-	const ENetMode NetMode = World.GetNetMode();
-	switch (NetMode)
+	if (World.IsStorageWorld() && World.IsEditorWorld())
+	{
+		return EProcessorExecutionFlags::Editor;
+	}
+	else
 	{
+		const ENetMode NetMode = World.GetNetMode();
+		switch (NetMode)
+		{
 		case NM_ListenServer:
-			ExecutionFlags = EProcessorExecutionFlags::Client | EProcessorExecutionFlags::Server;
-			break;
+			return EProcessorExecutionFlags::Client | EProcessorExecutionFlags::Server;
 		case NM_DedicatedServer:
-			ExecutionFlags = EProcessorExecutionFlags::Server;
-			break;
+			return EProcessorExecutionFlags::Server;
 		case NM_Client:
-			ExecutionFlags = EProcessorExecutionFlags::Client;
-			break;
+			return EProcessorExecutionFlags::Client;
+		case NM_Standalone:
+			return EProcessorExecutionFlags::Standalone;
 		default:
-			check(NetMode == NM_Standalone);
-			ExecutionFlags = EProcessorExecutionFlags::Standalone;
-			break;
+			checkf(false, TEXT("Unsupported ENetMode type (%i) found while determining MASS processor execution flags."), NetMode);
+			return EProcessorExecutionFlags::None;
+		}
 	}
-
-	return ExecutionFlags;
 }
 
 void CreateEntityCollections(const FMassEntityManager& EntityManager, const TConstArrayView<FMassEntityHandle> Entities
diff --git a/Engine/Plugins/Runtime/MassEntity/Source/MassEntity/Private/MassProcessingPhaseManager.cpp b/Engine/Plugins/Runtime/MassEntity/Source/MassEntity/Private/MassProcessingPhaseManager.cpp
--- a/Engine/Plugins/Runtime/MassEntity/Source/MassEntity/Private/MassProcessingPhaseManager.cpp
+++ b/Engine/Plugins/Runtime/MassEntity/Source/MassEntity/Private/MassProcessingPhaseManager.cpp
@@ -46,11 +46,12 @@ FMassProcessingPhase::FMassProcessingPhase()
 {
 	bCanEverTick = true;
 	bStartWithTickEnabled = false;
+	SupportedTickTypes = (1 << LEVELTICK_All) | (1 << LEVELTICK_TimeOnly);
 }
 
 void FMassProcessingPhase::ExecuteTick(float DeltaTime, ELevelTick TickType, ENamedThreads::Type CurrentThread, const FGraphEventRef& MyCompletionGraphEvent)
 {
-	if (TickType == LEVELTICK_ViewportsOnly || TickType == LEVELTICK_PauseTick)
+	if (ShouldTick(TickType) == false)
 	{
 		return;
 	}
@@ -168,7 +169,7 @@ void FMassProcessingPhaseManager::Initialize(UObject& InOwner, TConstArrayView<F
 {
 #if WITH_EDITOR
 	UWorld* World = InOwner.GetWorld();
-	const bool bIsEditorWorld = (World != 0) && (World->IsEditorWorld() && !World->IsGameWorld());
+	const bool bCreateProcessorGraphPreview = (World != nullptr) && (World->IsEditorWorld() && !World->IsGameWorld() && !World->IsStorageWorld());
 #endif // WITH_EDITOR
 	Owner = &InOwner;
 	ProcessingPhasesConfig = InProcessingPhasesConfig;
@@ -194,7 +195,7 @@ void FMassProcessingPhaseManager::Initialize(UObject& InOwner, TConstArrayView<F
 #endif // WITH_MASSENTITY_DEBUG
 
 #if WITH_EDITOR
-		if (bIsEditorWorld)
+		if (bCreateProcessorGraphPreview)
 		{
 			// populating the phase processor with initial data for editor world so that the default processing graph
 			// can be investigated in the editor without running PIE.
diff --git a/Engine/Plugins/Runtime/MassEntity/Source/MassEntity/Public/MassProcessingPhaseManager.h b/Engine/Plugins/Runtime/MassEntity/Source/MassEntity/Public/MassProcessingPhaseManager.h
--- a/Engine/Plugins/Runtime/MassEntity/Source/MassEntity/Public/MassProcessingPhaseManager.h
+++ b/Engine/Plugins/Runtime/MassEntity/Source/MassEntity/Public/MassProcessingPhaseManager.h
@@ -68,8 +68,12 @@ protected:
 	void ConfigureForParallelMode() { bRunInParallelMode = true; }
 	void ConfigureForSingleThreadMode() { bRunInParallelMode = false; }
 
+	bool ShouldTick(const ELevelTick TickType) const { return SupportedTickTypes & (1 << TickType); }
+
 public:
 	void Initialize(FMassProcessingPhaseManager& InPhaseManager, const EMassProcessingPhase InPhase, const ETickingGroup InTickGroup, UMassCompositeProcessor& InPhaseProcessor);
+	void AddSupportedTickType(const ELevelTick TickType) { SupportedTickTypes |= (1 << TickType); }
+	void RemoveSupportedTickType(const ELevelTick TickType) { SupportedTickTypes &= ~(1 << TickType); }
 
 protected:
 	friend FMassProcessingPhaseManager;
@@ -83,8 +87,9 @@ protected:
 
 private:
 	FMassProcessingPhaseManager* PhaseManager = nullptr;
-	bool bRunInParallelMode = false;
 	std::atomic<bool> bIsDuringMassProcessing = false;
+	bool bRunInParallelMode = false;
+	uint8 SupportedTickTypes = 0;
 };
 
 
diff --git a/Engine/Plugins/Runtime/MassEntity/Source/MassEntity/Public/MassProcessingTypes.h b/Engine/Plugins/Runtime/MassEntity/Source/MassEntity/Public/MassProcessingTypes.h
--- a/Engine/Plugins/Runtime/MassEntity/Source/MassEntity/Public/MassProcessingTypes.h
+++ b/Engine/Plugins/Runtime/MassEntity/Source/MassEntity/Public/MassProcessingTypes.h
@@ -24,7 +24,9 @@ enum class EProcessorExecutionFlags : uint8
 	Standalone = 1 << 0,
 	Server = 1 << 1,
 	Client = 1 << 2,
-	All = Standalone | Server | Client UMETA(Hidden)
+	Editor = 1 << 3,
+	AllNetModes = Standalone | Server | Client UMETA(Hidden),
+	All = Standalone | Server | Client | Editor UMETA(Hidden)
 };
 ENUM_CLASS_FLAGS(EProcessorExecutionFlags);
 
