Renamed FMassEntityModuleModule to FMassEntityModule

diff --git a/Engine/Plugins/Runtime/MassEntity/Source/MassEntity/Private/MassEntityModule.cpp b/Engine/Plugins/Runtime/MassEntity/Source/MassEntity/Private/MassEntityModule.cpp
--- a/Engine/Plugins/Runtime/MassEntity/Source/MassEntity/Private/MassEntityModule.cpp
+++ b/Engine/Plugins/Runtime/MassEntity/Source/MassEntity/Private/MassEntityModule.cpp
@@ -11,7 +11,7 @@
 
 #define LOCTEXT_NAMESPACE "Mass"
 
-class FMassEntityModuleModule : public IMassEntityModule
+class FMassEntityModule : public IMassEntityModule
 {
 	/** IModuleInterface implementation */
 	virtual void StartupModule() override;
@@ -22,10 +22,10 @@ class FMassEntityModuleModule : public IMassEntityModule
 #endif // WITH_UNREAL_DEVELOPER_TOOLS
 };
 
-IMPLEMENT_MODULE(FMassEntityModuleModule, MassEntity)
+IMPLEMENT_MODULE(FMassEntityModule, MassEntity)
 
 
-void FMassEntityModuleModule::StartupModule()
+void FMassEntityModule::StartupModule()
 {
 	// This code will execute after your module is loaded into memory (but after global variables are initialized, of course.)
 
@@ -37,7 +37,7 @@ void FMassEntityModuleModule::StartupModule()
 		InitOptions.bShowFilters = true;
 		MessageLogModule.RegisterLogListing("MassEntity", LOCTEXT("MassEntity", "MassEntity"), InitOptions);
 
-		OnWorldCleanupHandle = FWorldDelegates::OnWorldCleanup.AddStatic(&FMassEntityModuleModule::OnWorldCleanup);
+		OnWorldCleanupHandle = FWorldDelegates::OnWorldCleanup.AddStatic(&FMassEntityModule::OnWorldCleanup);
 	}
 #endif // WITH_UNREAL_DEVELOPER_TOOLS
 
@@ -49,7 +49,7 @@ void FMassEntityModuleModule::StartupModule()
 }
 
 
-void FMassEntityModuleModule::ShutdownModule()
+void FMassEntityModule::ShutdownModule()
 {
 	// This function may be called during shutdown to clean up your module.  For modules that support dynamic reloading,
 	// we call this function before unloading the module.
@@ -59,7 +59,7 @@ void FMassEntityModuleModule::ShutdownModule()
 }
 
 #if WITH_UNREAL_DEVELOPER_TOOLS
-void FMassEntityModuleModule::OnWorldCleanup(UWorld* /*World*/, bool /*bSessionEnded*/, bool /*bCleanupResources*/)
+void FMassEntityModule::OnWorldCleanup(UWorld* /*World*/, bool /*bSessionEnded*/, bool /*bCleanupResources*/)
 {
 	// clearing out messages from the world being cleaned up
 	FMessageLog("MassEntity").NewPage(FText::FromString(TEXT("MassEntity")));
