Modified FMassCommandBuffer.DestroyEntities's signature to use a array view, which makes it both backwards compatible as well as more flexible.

diff --git a/Engine/Plugins/Runtime/MassEntity/Source/MassEntity/Public/MassCommandBuffer.h b/Engine/Plugins/Runtime/MassEntity/Source/MassEntity/Public/MassCommandBuffer.h
--- a/Engine/Plugins/Runtime/MassEntity/Source/MassEntity/Public/MassCommandBuffer.h
+++ b/Engine/Plugins/Runtime/MassEntity/Source/MassEntity/Public/MassCommandBuffer.h
@@ -129,7 +129,7 @@ public:
 		PushCommand<FMassCommandDestroyEntities>(Entity);
 	}
 
-	void DestroyEntities(const TArray<FMassEntityHandle>& InEntitiesToDestroy)
+	void DestroyEntities(TConstArrayView<FMassEntityHandle> InEntitiesToDestroy)
 	{
 		PushCommand<FMassCommandDestroyEntities>(InEntitiesToDestroy);
 	}
