Allowed MASS to have tag-only queries.

In rare cases MASS processors may only require tags to do their work. Examples of this is in situations where a processors removes a temporary tag or merges multiple tags into a new tag/fragment.

diff --git a/Engine/Plugins/Runtime/MassEntity/Source/MassEntity/Private/MassEntityManager.cpp b/Engine/Plugins/Runtime/MassEntity/Source/MassEntity/Private/MassEntityManager.cpp
--- a/Engine/Plugins/Runtime/MassEntity/Source/MassEntity/Private/MassEntityManager.cpp
+++ b/Engine/Plugins/Runtime/MassEntity/Source/MassEntity/Private/MassEntityManager.cpp
@@ -1162,14 +1162,36 @@ void FMassEntityManager::GetValidArchetypes(const FMassEntityQuery& Query, TArra
 
 	// First get set of all archetypes that contain *any* fragment
 	TSet<TSharedPtr<FMassArchetypeData>> AnyArchetypes;
-	for (const FMassFragmentRequirementDescription& Requirement : Query.GetFragmentRequirements())
+	TConstArrayView<FMassFragmentRequirementDescription> FragmentRequirements = Query.GetFragmentRequirements();
+	if (!FragmentRequirements.IsEmpty())
 	{
-		check(Requirement.StructType);
-		if (Requirement.Presence != EMassFragmentPresence::None)
+		for (const FMassFragmentRequirementDescription& Requirement : FragmentRequirements)
 		{
-			if (const TArray<TSharedPtr<FMassArchetypeData>>* pData = FragmentTypeToArchetypeMap.Find(Requirement.StructType))
+			check(Requirement.StructType);
+			if (Requirement.Presence != EMassFragmentPresence::None)
 			{
-				AnyArchetypes.Append(*pData);
+				if (const TArray<TSharedPtr<FMassArchetypeData>>* pData = FragmentTypeToArchetypeMap.Find(Requirement.StructType))
+				{
+					AnyArchetypes.Append(*pData);
+				}
+			}
+		}
+	}
+	else
+	{
+		// If there are no fragment requirements, assume this is a tag-only query.
+		const FMassTagBitSet& AllTags = Query.GetRequiredAllTags();
+		const FMassTagBitSet& AnyTags = Query.GetRequiredAnyTags();
+
+		for (const TPair<uint32, TArray<TSharedPtr<FMassArchetypeData>>>& KeyValue : FragmentHashToArchetypeMap)
+		{
+			for (const TSharedPtr<FMassArchetypeData>& ArchetypeData : KeyValue.Value)
+			{
+				const FMassTagBitSet ArchetypeTags = ArchetypeData->GetTagBitSet();
+				if (ArchetypeTags.HasAll(AllTags) || ArchetypeTags.HasAny(AnyTags))
+				{
+					AnyArchetypes.Add(ArchetypeData);
+				}
 			}
 		}
 	}
diff --git a/Engine/Plugins/Runtime/MassEntity/Source/MassEntity/Private/MassRequirements.cpp b/Engine/Plugins/Runtime/MassEntity/Source/MassEntity/Private/MassRequirements.cpp
--- a/Engine/Plugins/Runtime/MassEntity/Source/MassEntity/Private/MassRequirements.cpp
+++ b/Engine/Plugins/Runtime/MassEntity/Source/MassEntity/Private/MassRequirements.cpp
@@ -83,7 +83,9 @@ void FMassFragmentRequirements::SortRequirements()
 
 bool FMassFragmentRequirements::CheckValidity() const
 {
-	return RequiredAllFragments.IsEmpty() == false || RequiredAnyFragments.IsEmpty() == false || RequiredOptionalFragments.IsEmpty() == false;
+	return
+		RequiredAllFragments.IsEmpty() == false || RequiredAnyFragments.IsEmpty() == false || RequiredOptionalFragments.IsEmpty() == false ||
+		RequiredAllTags.IsEmpty() == false || RequiredAnyTags.IsEmpty() == false;
 }
 
 bool FMassFragmentRequirements::IsEmpty() const
