Extended creation of MASS archetypes based on existing archetypes.

It's now possible to create a new archetype based on an existing archetype with both fragments and tags. This version is slower than other options to duplicate an archetype but allows tags instead of only fragments and is forgiving when fragments/tags are added that are already in the archetype.

#nojira

diff --git a/Engine/Plugins/Runtime/MassEntity/Source/MassEntity/Private/MassEntityManager.cpp b/Engine/Plugins/Runtime/MassEntity/Source/MassEntity/Private/MassEntityManager.cpp
--- a/Engine/Plugins/Runtime/MassEntity/Source/MassEntity/Private/MassEntityManager.cpp
+++ b/Engine/Plugins/Runtime/MassEntity/Source/MassEntity/Private/MassEntityManager.cpp
@@ -181,33 +181,17 @@ void FMassEntityManager::Deinitialize()
 
 FMassArchetypeHandle FMassEntityManager::CreateArchetype(TConstArrayView<const UScriptStruct*> FragmentsAndTagsList, const FName ArchetypeDebugName)
 {
-	FMassChunkFragmentBitSet ChunkFragments;
-	FMassTagBitSet Tags;
-	TArray<const UScriptStruct*, TInlineAllocator<16>> FragmentList;
-	FragmentList.Reserve(FragmentsAndTagsList.Num());
-
-	for (const UScriptStruct* Type : FragmentsAndTagsList)
-	{
-		if (Type->IsChildOf(FMassFragment::StaticStruct()))
-		{
-			FragmentList.Add(Type);
-		}
-		else if (Type->IsChildOf(FMassTag::StaticStruct()))
-		{
-			Tags.Add(*Type);
-		}
-		else if (Type->IsChildOf(FMassChunkFragment::StaticStruct()))
-		{
-			ChunkFragments.Add(*Type);
-		}
-		else
-		{
-			UE_LOG(LogMass, Warning, TEXT("%s: %s is not a valid fragment nor tag type. Ignoring.")
-				, ANSI_TO_TCHAR(__FUNCTION__), *GetNameSafe(Type));
-		}
-	}
+	FMassArchetypeCompositionDescriptor Composition;
+	InternalAppendFragmentsAndTagsToArchetypeCompositionDescriptor(Composition, FragmentsAndTagsList);
+	return CreateArchetype(Composition, ArchetypeDebugName);
+}
 
-	const FMassArchetypeCompositionDescriptor Composition(FMassFragmentBitSet(FragmentList), Tags, ChunkFragments, FMassSharedFragmentBitSet());
+FMassArchetypeHandle FMassEntityManager::CreateArchetype(FMassArchetypeHandle SourceArchetype,
+	TConstArrayView<const UScriptStruct*> FragmentsAndTagsList, const FName ArchetypeDebugName)
+{
+	const FMassArchetypeData& ArchetypeData = FMassArchetypeHelper::ArchetypeDataFromHandleChecked(SourceArchetype);
+	FMassArchetypeCompositionDescriptor Composition = ArchetypeData.GetCompositionDescriptor();
+	InternalAppendFragmentsAndTagsToArchetypeCompositionDescriptor(Composition, FragmentsAndTagsList);
 	return CreateArchetype(Composition, ArchetypeDebugName);
 }
 
@@ -323,6 +307,31 @@ FMassArchetypeHandle FMassEntityManager::InternalCreateSimilarArchetype(const FM
 	return FMassArchetypeHelper::ArchetypeHandleFromData(ArchetypeDataPtr);
 }
 
+void FMassEntityManager::InternalAppendFragmentsAndTagsToArchetypeCompositionDescriptor(
+	FMassArchetypeCompositionDescriptor& InOutComposition, TConstArrayView<const UScriptStruct*> FragmentsAndTagsList) const
+{
+	for (const UScriptStruct* Type : FragmentsAndTagsList)
+	{
+		if (Type->IsChildOf(FMassFragment::StaticStruct()))
+		{
+			InOutComposition.Fragments.Add(*Type);
+		}
+		else if (Type->IsChildOf(FMassTag::StaticStruct()))
+		{
+			InOutComposition.Tags.Add(*Type);
+		}
+		else if (Type->IsChildOf(FMassChunkFragment::StaticStruct()))
+		{
+			InOutComposition.ChunkFragments.Add(*Type);
+		}
+		else
+		{
+			UE_LOG(LogMass, Warning, TEXT("%s: %s is not a valid fragment nor tag type. Ignoring.")
+				, ANSI_TO_TCHAR(__FUNCTION__), *GetNameSafe(Type));
+		}
+	}
+}
+
 FMassArchetypeHandle FMassEntityManager::GetArchetypeForEntity(FMassEntityHandle Entity) const
 {
 	if (IsEntityValid(Entity))
diff --git a/Engine/Plugins/Runtime/MassEntity/Source/MassEntity/Public/MassEntityManager.h b/Engine/Plugins/Runtime/MassEntity/Source/MassEntity/Public/MassEntityManager.h
--- a/Engine/Plugins/Runtime/MassEntity/Source/MassEntity/Public/MassEntityManager.h
+++ b/Engine/Plugins/Runtime/MassEntity/Source/MassEntity/Public/MassEntityManager.h
@@ -106,6 +106,14 @@ public:
 	 */
 	FMassArchetypeHandle CreateArchetype(TConstArrayView<const UScriptStruct*> FragmentsAndTagsList, const FName ArchetypeDebugName = FName());
 
+	/**
+	 * A special, relaxed but slower version of CreateArchetype functions that allows FragmentAngTagsList to contain
+	 * both fragments and tags. This version takes an original archetype and copies it layout, then appends any fragments and tags from the
+	 * provided list if they're not already in the original archetype.
+	 */
+	FMassArchetypeHandle CreateArchetype(FMassArchetypeHandle SourceArchetype, 
+		TConstArrayView<const UScriptStruct*> FragmentsAndTagsList, const FName ArchetypeDebugName = FName());
+
 	/**
 	 * CreateArchetype from a composition descriptor and initial values
 	 *
@@ -426,6 +434,9 @@ protected:
 
 	FMassArchetypeHandle InternalCreateSimilarArchetype(const FMassArchetypeData& SourceArchetypeRef, FMassArchetypeCompositionDescriptor&& NewComposition);
 
+	void InternalAppendFragmentsAndTagsToArchetypeCompositionDescriptor(FMassArchetypeCompositionDescriptor& InOutComposition,
+		TConstArrayView<const UScriptStruct*> FragmentsAndTagsList) const;
+
 private:
 	void InternalBuildEntity(FMassEntityHandle Entity, const FMassArchetypeHandle& ArchetypeHandle, const FMassArchetypeSharedFragmentValues& SharedFragmentValues);
 	void InternalReleaseEntity(FMassEntityHandle Entity);
@@ -442,7 +453,6 @@ private:
 	void InternalAddFragmentListToEntity(FMassEntityHandle Entity, const FMassFragmentBitSet& InFragments);
 	void* InternalGetFragmentDataChecked(FMassEntityHandle Entity, const UScriptStruct* FragmentType) const;
 	void* InternalGetFragmentDataPtr(FMassEntityHandle Entity, const UScriptStruct* FragmentType) const;
-
 private:
 	TChunkedArray<FEntityData> Entities;
 	TArray<int32> EntityFreeIndexList;
