Mass's support for any kind of USubsystem.
Added new struct type, FMassSubsystemAccess, storing bitsets indicating which subsystems are to be accessed and in which mode (RO/WR) along with cached instances of the systems (populated either via CacheSubsystemRequirements call or iteratively as the subsystems get requested)
The new  FMassSubsystemAccess is instantiated by FMassExecutionContext.
Also:
* removed need of supplying a UWorld instance when fetching subsystems - the value is now fetched from the MassEntityManager used and cached.
* reworked the unit test for Mass USubsystem access, following the extension of Mass subsystems access support.
* changes due to Mass Subsystem access API changes.

diff --git a/Engine/Plugins/AI/MassAI/Source/MassAIBehavior/Private/MassLookAtProcessors.cpp b/Engine/Plugins/AI/MassAI/Source/MassAIBehavior/Private/MassLookAtProcessors.cpp
--- a/Engine/Plugins/AI/MassAI/Source/MassAIBehavior/Private/MassLookAtProcessors.cpp
+++ b/Engine/Plugins/AI/MassAI/Source/MassAIBehavior/Private/MassLookAtProcessors.cpp
@@ -4,6 +4,7 @@
 #include "MassEntityView.h"
 #include "MassAIBehaviorTypes.h"
 #include "MassCommonFragments.h"
+#include "MassExecutionContext.h"
 #include "MassLookAtFragments.h"
 #include "MassMovementFragments.h"
 #include "MassNavigationFragments.h"
@@ -120,10 +121,10 @@ void UMassLookAtProcessor::Execute(FMassEntityManager& EntityManager, FMassExecu
 
 	const double CurrentTime = GetWorld()->GetTimeSeconds();
 
-	EntityQuery_Conditional.ForEachEntityChunk(EntityManager, Context, [this, &EntityManager, CurrentTime, World = EntityManager.GetWorld()](FMassExecutionContext& Context)
+	EntityQuery_Conditional.ForEachEntityChunk(EntityManager, Context, [this, &EntityManager, CurrentTime](FMassExecutionContext& Context)
 		{
-			const UMassNavigationSubsystem& MassNavSystem = Context.GetSubsystemChecked<UMassNavigationSubsystem>(World);
-			const UZoneGraphSubsystem& ZoneGraphSubsystem = Context.GetSubsystemChecked<UZoneGraphSubsystem>(World);
+			const UMassNavigationSubsystem& MassNavSystem = Context.GetSubsystemChecked<UMassNavigationSubsystem>();
+			const UZoneGraphSubsystem& ZoneGraphSubsystem = Context.GetSubsystemChecked<UZoneGraphSubsystem>();
 
 			const int32 NumEntities = Context.GetNumEntities();
 			const TArrayView<FMassLookAtFragment> LookAtList = Context.GetMutableFragmentView<FMassLookAtFragment>();
diff --git a/Engine/Plugins/AI/MassAI/Source/MassAIBehavior/Private/MassStateTreeExecutionContext.cpp b/Engine/Plugins/AI/MassAI/Source/MassAIBehavior/Private/MassStateTreeExecutionContext.cpp
--- a/Engine/Plugins/AI/MassAI/Source/MassAIBehavior/Private/MassStateTreeExecutionContext.cpp
+++ b/Engine/Plugins/AI/MassAI/Source/MassAIBehavior/Private/MassStateTreeExecutionContext.cpp
@@ -3,6 +3,7 @@
 #include "MassStateTreeExecutionContext.h"
 #include "MassStateTreeTypes.h"
 #include "MassEntityManager.h"
+#include "MassExecutionContext.h"
 #include "MassSignalSubsystem.h"
 #include "Engine/World.h"
 
diff --git a/Engine/Plugins/AI/MassAI/Source/MassAIBehavior/Private/MassStateTreeProcessors.cpp b/Engine/Plugins/AI/MassAI/Source/MassAIBehavior/Private/MassStateTreeProcessors.cpp
--- a/Engine/Plugins/AI/MassAI/Source/MassAIBehavior/Private/MassStateTreeProcessors.cpp
+++ b/Engine/Plugins/AI/MassAI/Source/MassAIBehavior/Private/MassStateTreeProcessors.cpp
@@ -4,6 +4,7 @@
 #include "StateTree.h"
 #include "MassStateTreeExecutionContext.h"
 #include "MassEntityView.h"
+#include "MassExecutionContext.h"
 #include "MassNavigationTypes.h"
 #include "MassSimulationLOD.h"
 #include "MassComponentHitTypes.h"
@@ -229,9 +230,9 @@ void UMassStateTreeFragmentDestructor::Execute(FMassEntityManager& EntityManager
 	UE::MassBehavior::FCachedExternalData CachedExternalData;
 
 	EntityQuery.ForEachEntityChunk(EntityManager, Context,
-		[&EntityManager, &CachedExternalData, SignalSubsystem = SignalSubsystem, World = EntityManager.GetWorld()](FMassExecutionContext& Context)
+		[&EntityManager, &CachedExternalData, SignalSubsystem = SignalSubsystem](FMassExecutionContext& Context)
 		{
-			UMassStateTreeSubsystem& MassStateTreeSubsystem = Context.GetMutableSubsystemChecked<UMassStateTreeSubsystem>(World);
+			UMassStateTreeSubsystem& MassStateTreeSubsystem = Context.GetMutableSubsystemChecked<UMassStateTreeSubsystem>();
 			const TArrayView<FMassStateTreeInstanceFragment> StateTreeInstanceList = Context.GetMutableFragmentView<FMassStateTreeInstanceFragment>();
 			const FMassStateTreeSharedFragment& SharedStateTree = Context.GetConstSharedFragment<FMassStateTreeSharedFragment>();
 
@@ -281,7 +282,7 @@ void UMassStateTreeActivationProcessor::ConfigureQueries()
 
 void UMassStateTreeActivationProcessor::Execute(FMassEntityManager& EntityManager, FMassExecutionContext& Context)
 {
-	UMassSignalSubsystem& SignalSubsystem = Context.GetMutableSubsystemChecked<UMassSignalSubsystem>(EntityManager.GetWorld());
+	UMassSignalSubsystem& SignalSubsystem = Context.GetMutableSubsystemChecked<UMassSignalSubsystem>();
 
 	const UMassBehaviorSettings* BehaviorSettings = GetDefault<UMassBehaviorSettings>();
 	check(BehaviorSettings);
@@ -295,9 +296,9 @@ void UMassStateTreeActivationProcessor::Execute(FMassEntityManager& EntityManage
 	int32 ActivationCounts[EMassLOD::Max] {0,0,0,0};
 	
 	EntityQuery.ForEachEntityChunk(EntityManager, Context,
-		[&EntitiesToSignal, &ActivationCounts, &CachedExternalData, SignalSubsystem = &SignalSubsystem, &EntityManager, MaxActivationsPerLOD = BehaviorSettings->MaxActivationsPerLOD, TimeInSeconds, World = EntityManager.GetWorld()](FMassExecutionContext& Context)
+		[&EntitiesToSignal, &ActivationCounts, &CachedExternalData, SignalSubsystem = &SignalSubsystem, &EntityManager, MaxActivationsPerLOD = BehaviorSettings->MaxActivationsPerLOD, TimeInSeconds](FMassExecutionContext& Context)
 		{
-			UMassStateTreeSubsystem& MassStateTreeSubsystem = Context.GetMutableSubsystemChecked<UMassStateTreeSubsystem>(World);
+			UMassStateTreeSubsystem& MassStateTreeSubsystem = Context.GetMutableSubsystemChecked<UMassStateTreeSubsystem>();
 			const int32 NumEntities = Context.GetNumEntities();
 			// Check if we already reached the maximum for this frame
 			const EMassLOD::Type ChunkLOD = FMassSimulationVariableTickChunkFragment::GetChunkLOD(Context);
@@ -405,8 +406,7 @@ void UMassStateTreeProcessor::ConfigureQueries()
 
 void UMassStateTreeProcessor::SignalEntities(FMassEntityManager& EntityManager, FMassExecutionContext& Context, FMassSignalNameLookup& EntitySignals)
 {
-	UWorld* World = EntityManager.GetWorld();
-	UMassSignalSubsystem& SignalSubsystem = Context.GetMutableSubsystemChecked<UMassSignalSubsystem>(World);
+	UMassSignalSubsystem& SignalSubsystem = Context.GetMutableSubsystemChecked<UMassSignalSubsystem>();
 	
 	QUICK_SCOPE_CYCLE_COUNTER(StateTreeProcessor_Run);
 	CSV_SCOPED_TIMING_STAT_EXCLUSIVE(StateTreeProcessorExecute);
@@ -417,12 +417,12 @@ void UMassStateTreeProcessor::SignalEntities(FMassEntityManager& EntityManager,
 	TArray<FMassEntityHandle> EntitiesToSignal;
 
 	EntityQuery.ForEachEntityChunk(EntityManager, Context,
-		[this, &TimeInSeconds, &EntitiesToSignal, &CachedExternalData, &EntityManager, SignalSubsystem = &SignalSubsystem, World, &EntitySignals](FMassExecutionContext& Context)
+		[this, &TimeInSeconds, &EntitiesToSignal, &CachedExternalData, &EntityManager, SignalSubsystem = &SignalSubsystem, &EntitySignals](FMassExecutionContext& Context)
 		{
 			// Keep stats regarding the amount of tree instances ticked per frame
 			CSV_CUSTOM_STAT(StateTreeProcessor, NumTickedStateTree, Context.GetNumEntities(), ECsvCustomStatOp::Accumulate);
 
-			UMassStateTreeSubsystem& MassStateTreeSubsystem = Context.GetMutableSubsystemChecked<UMassStateTreeSubsystem>(World);
+			UMassStateTreeSubsystem& MassStateTreeSubsystem = Context.GetMutableSubsystemChecked<UMassStateTreeSubsystem>();
 
 			UE::MassBehavior::ForEachEntityInChunk(CachedExternalData, Context, EntityManager, *SignalSubsystem, MassStateTreeSubsystem,
 				[TimeInSeconds, &EntitiesToSignal, &EntitySignals, &MassStateTreeSubsystem]
diff --git a/Engine/Plugins/AI/MassAI/Source/MassAIBehavior/Private/MassZoneGraphAnnotationFragments.cpp b/Engine/Plugins/AI/MassAI/Source/MassAIBehavior/Private/MassZoneGraphAnnotationFragments.cpp
--- a/Engine/Plugins/AI/MassAI/Source/MassAIBehavior/Private/MassZoneGraphAnnotationFragments.cpp
+++ b/Engine/Plugins/AI/MassAI/Source/MassAIBehavior/Private/MassZoneGraphAnnotationFragments.cpp
@@ -2,6 +2,7 @@
 
 #include "MassZoneGraphAnnotationFragments.h"
 #include "MassEntityManager.h"
+#include "MassExecutionContext.h"
 #include "MassSimulationLOD.h"
 
 namespace UE::Mass::ZoneGraphAnnotations
diff --git a/Engine/Plugins/AI/MassAI/Source/MassAIBehavior/Private/MassZoneGraphAnnotationProcessors.cpp b/Engine/Plugins/AI/MassAI/Source/MassAIBehavior/Private/MassZoneGraphAnnotationProcessors.cpp
--- a/Engine/Plugins/AI/MassAI/Source/MassAIBehavior/Private/MassZoneGraphAnnotationProcessors.cpp
+++ b/Engine/Plugins/AI/MassAI/Source/MassAIBehavior/Private/MassZoneGraphAnnotationProcessors.cpp
@@ -10,6 +10,7 @@
 #include "ZoneGraphAnnotationSubsystem.h"
 #include "MassSignalSubsystem.h"
 #include "MassSimulationLOD.h"
+#include "MassExecutionContext.h"
 #include "Engine/World.h"
 
 //----------------------------------------------------------------------//
@@ -31,9 +32,9 @@ void UMassZoneGraphAnnotationTagsInitializer::ConfigureQueries()
 
 void UMassZoneGraphAnnotationTagsInitializer::Execute(FMassEntityManager& EntityManager, FMassExecutionContext& Context)
 {
-	EntityQuery.ForEachEntityChunk(EntityManager, Context, [this, World = EntityManager.GetWorld()](FMassExecutionContext& Context)
+	EntityQuery.ForEachEntityChunk(EntityManager, Context, [this](FMassExecutionContext& Context)
 	{
-		const UZoneGraphAnnotationSubsystem& ZoneGraphAnnotationSubsystem = Context.GetSubsystemChecked<UZoneGraphAnnotationSubsystem>(World);
+		const UZoneGraphAnnotationSubsystem& ZoneGraphAnnotationSubsystem = Context.GetSubsystemChecked<UZoneGraphAnnotationSubsystem>();
 
 		const int32 NumEntities = Context.GetNumEntities();
 		const TArrayView<FMassZoneGraphAnnotationFragment> AnnotationTagsList = Context.GetMutableFragmentView<FMassZoneGraphAnnotationFragment>();
@@ -93,8 +94,7 @@ void UMassZoneGraphAnnotationTagUpdateProcessor::Execute(FMassEntityManager& Ent
 	// Calling super will update the signals, and call SignalEntities() below.
 	Super::Execute(EntityManager, Context);
 
-	UWorld* World = EntityManager.GetWorld();
-	EntityQuery.ForEachEntityChunk(EntityManager, Context, [this, World](FMassExecutionContext& Context)
+	EntityQuery.ForEachEntityChunk(EntityManager, Context, [this](FMassExecutionContext& Context)
 	{
 		// Periodically update tags.
 		if (!FMassZoneGraphAnnotationVariableTickChunkFragment::UpdateChunk(Context))
@@ -102,7 +102,7 @@ void UMassZoneGraphAnnotationTagUpdateProcessor::Execute(FMassEntityManager& Ent
 			return;
 		}
 		
-		UZoneGraphAnnotationSubsystem& ZoneGraphAnnotationSubsystem = Context.GetMutableSubsystemChecked<UZoneGraphAnnotationSubsystem>(World);
+		UZoneGraphAnnotationSubsystem& ZoneGraphAnnotationSubsystem = Context.GetMutableSubsystemChecked<UZoneGraphAnnotationSubsystem>();
 
 		const int32 NumEntities = Context.GetNumEntities();
 		const TArrayView<FMassZoneGraphAnnotationFragment> AnnotationTagsList = Context.GetMutableFragmentView<FMassZoneGraphAnnotationFragment>();
@@ -119,7 +119,7 @@ void UMassZoneGraphAnnotationTagUpdateProcessor::Execute(FMassEntityManager& Ent
 
 	if (TransientEntitiesToSignal.Num())
 	{
-		UMassSignalSubsystem& SignalSubsystem = Context.GetMutableSubsystemChecked<UMassSignalSubsystem>(World);
+		UMassSignalSubsystem& SignalSubsystem = Context.GetMutableSubsystemChecked<UMassSignalSubsystem>();
 		SignalSubsystem.SignalEntities(UE::Mass::Signals::AnnotationTagsChanged, TransientEntitiesToSignal);
 	}
 }
@@ -145,9 +145,9 @@ void UMassZoneGraphAnnotationTagUpdateProcessor::UpdateAnnotationTags(UZoneGraph
 
 void UMassZoneGraphAnnotationTagUpdateProcessor::SignalEntities(FMassEntityManager& EntityManager, FMassExecutionContext& Context, FMassSignalNameLookup& EntitySignals)
 {
-	EntityQuery.ForEachEntityChunk(EntityManager, Context, [this, World = EntityManager.GetWorld()](FMassExecutionContext& Context)
+	EntityQuery.ForEachEntityChunk(EntityManager, Context, [this](FMassExecutionContext& Context)
 	{
-		UZoneGraphAnnotationSubsystem& ZoneGraphAnnotationSubsystem = Context.GetMutableSubsystemChecked<UZoneGraphAnnotationSubsystem>(World);
+		UZoneGraphAnnotationSubsystem& ZoneGraphAnnotationSubsystem = Context.GetMutableSubsystemChecked<UZoneGraphAnnotationSubsystem>();
 
 		const int32 NumEntities = Context.GetNumEntities();
 		const TConstArrayView<FMassZoneGraphLaneLocationFragment> LaneLocationList = Context.GetFragmentView<FMassZoneGraphLaneLocationFragment>();
diff --git a/Engine/Plugins/AI/MassAI/Source/MassAIDebug/Private/GameplayDebuggerCategory_Mass.cpp b/Engine/Plugins/AI/MassAI/Source/MassAIDebug/Private/GameplayDebuggerCategory_Mass.cpp
--- a/Engine/Plugins/AI/MassAI/Source/MassAIDebug/Private/GameplayDebuggerCategory_Mass.cpp
+++ b/Engine/Plugins/AI/MassAI/Source/MassAIDebug/Private/GameplayDebuggerCategory_Mass.cpp
@@ -5,6 +5,7 @@
 #if WITH_GAMEPLAY_DEBUGGER && WITH_MASSGAMEPLAY_DEBUG
 #include "MassGameplayDebugTypes.h"
 #include "MassEntityView.h"
+#include "MassExecutionContext.h"
 #include "GameplayDebuggerConfig.h"
 #include "GameplayDebuggerCategoryReplicator.h"
 #include "MassDebuggerSubsystem.h"
@@ -140,7 +141,7 @@ void FGameplayDebuggerCategory_Mass::PickEntity(const APlayerController& OwnerPC
 	{
 		TArray<FMassEntityHandle> Entities;
 		TArray<FVector> Locations;
-		FMassExecutionContext ExecutionContext;
+		FMassExecutionContext ExecutionContext(EntityManager);
 		FMassEntityQuery Query;
 		Query.AddRequirement<FTransformFragment>(EMassFragmentAccess::ReadOnly);
 		Query.ForEachEntityChunk(EntityManager, ExecutionContext, [&Entities, &Locations](FMassExecutionContext& Context)
@@ -389,7 +390,7 @@ void FGameplayDebuggerCategory_Mass::CollectData(APlayerController* OwnerPC, AAc
 		
 		if (MassStateTreeSubsystem && SignalSubsystem && SmartObjectSubsystem)
 		{
-			FMassExecutionContext Context(EntityManager.AsShared(), 0.0f);
+			FMassExecutionContext Context(EntityManager, 0.0f);
 		
 			EntityQuery.ForEachEntityChunk(EntityManager, Context, [this, MassStateTreeSubsystem, SignalSubsystem, SmartObjectSubsystem, OwnerPC, ViewLocation, ViewDirection, CurrentTime](FMassExecutionContext& Context)
 			{
@@ -635,7 +636,7 @@ void FGameplayDebuggerCategory_Mass::CollectData(APlayerController* OwnerPC, AAc
 			FMassEntityQuery EntityColliderQuery;
 			EntityColliderQuery.AddRequirement<FMassAvoidanceColliderFragment>(EMassFragmentAccess::ReadOnly);
 			EntityColliderQuery.AddRequirement<FTransformFragment>(EMassFragmentAccess::ReadOnly);
-			FMassExecutionContext Context(EntityManager.AsShared(), 0.f);
+			FMassExecutionContext Context(EntityManager, 0.f);
 			EntityColliderQuery.ForEachEntityChunk(EntityManager, Context, [this, ViewLocation, ViewDirection](const FMassExecutionContext& Context)
 			{
 				const int32 NumEntities = Context.GetNumEntities();
diff --git a/Engine/Plugins/AI/MassAI/Source/MassAIDebug/Private/MassDebugStateTreeProcessor.cpp b/Engine/Plugins/AI/MassAI/Source/MassAIDebug/Private/MassDebugStateTreeProcessor.cpp
--- a/Engine/Plugins/AI/MassAI/Source/MassAIDebug/Private/MassDebugStateTreeProcessor.cpp
+++ b/Engine/Plugins/AI/MassAI/Source/MassAIDebug/Private/MassDebugStateTreeProcessor.cpp
@@ -6,6 +6,7 @@
 #include "MassStateTreeExecutionContext.h"
 #include "MassStateTreeFragments.h"
 #include "MassCommonFragments.h"
+#include "MassExecutionContext.h"
 #include "VisualLogger/VisualLogger.h"
 #include "Engine/World.h"
 
diff --git a/Engine/Plugins/AI/MassAI/Source/MassAIReplication/Private/MassReplicationPathHandlers.cpp b/Engine/Plugins/AI/MassAI/Source/MassAIReplication/Private/MassReplicationPathHandlers.cpp
--- a/Engine/Plugins/AI/MassAI/Source/MassAIReplication/Private/MassReplicationPathHandlers.cpp
+++ b/Engine/Plugins/AI/MassAI/Source/MassAIReplication/Private/MassReplicationPathHandlers.cpp
@@ -3,6 +3,7 @@
 #include "MassReplicationPathHandlers.h"
 #include "MassEntityQuery.h"
 #include "MassCommonFragments.h"
+#include "MassExecutionContext.h"
 #include "MassSimulationSubsystem.h"
 #include "MassZoneGraphNavigationUtils.h"
 #include "ZoneGraphSubsystem.h"
diff --git a/Engine/Plugins/AI/MassAI/Source/MassNavigation/Private/MassAvoidanceProcessors.cpp b/Engine/Plugins/AI/MassAI/Source/MassNavigation/Private/MassAvoidanceProcessors.cpp
--- a/Engine/Plugins/AI/MassAI/Source/MassNavigation/Private/MassAvoidanceProcessors.cpp
+++ b/Engine/Plugins/AI/MassAI/Source/MassNavigation/Private/MassAvoidanceProcessors.cpp
@@ -4,6 +4,7 @@
 #include "Avoidance/MassAvoidanceFragments.h"
 #include "DrawDebugHelpers.h"
 #include "MassEntityView.h"
+#include "MassExecutionContext.h"
 #include "VisualLogger/VisualLogger.h"
 #include "Math/Vector2D.h"
 #include "Logging/LogMacros.h"
diff --git a/Engine/Plugins/AI/MassAI/Source/MassNavigation/Private/MassNavigationProcessors.cpp b/Engine/Plugins/AI/MassAI/Source/MassNavigation/Private/MassNavigationProcessors.cpp
--- a/Engine/Plugins/AI/MassAI/Source/MassNavigation/Private/MassNavigationProcessors.cpp
+++ b/Engine/Plugins/AI/MassAI/Source/MassNavigation/Private/MassNavigationProcessors.cpp
@@ -10,6 +10,7 @@
 #include "MassMovementTypes.h"
 #include "MassMovementFragments.h"
 #include "MassEntityView.h"
+#include "MassExecutionContext.h"
 #include "Engine/World.h"
 
 
@@ -195,9 +196,9 @@ void UMassNavigationObstacleGridProcessor::ConfigureQueries()
 void UMassNavigationObstacleGridProcessor::Execute(FMassEntityManager& EntityManager, FMassExecutionContext& Context)
 {
 	// can't be ParallelFor due to MovementSubsystem->GetGridMutable().Move not being thread-safe
-	AddToGridEntityQuery.ForEachEntityChunk(EntityManager, Context, [this, &EntityManager, World = EntityManager.GetWorld()](FMassExecutionContext& Context)
+	AddToGridEntityQuery.ForEachEntityChunk(EntityManager, Context, [this, &EntityManager](FMassExecutionContext& Context)
 	{
-		FNavigationObstacleHashGrid2D& HashGrid = Context.GetMutableSubsystemChecked<UMassNavigationSubsystem>(World).GetObstacleGridMutable();
+		FNavigationObstacleHashGrid2D& HashGrid = Context.GetMutableSubsystemChecked<UMassNavigationSubsystem>().GetObstacleGridMutable();
 		const int32 NumEntities = Context.GetNumEntities();
 
 		TConstArrayView<FTransformFragment> LocationList = Context.GetFragmentView<FTransformFragment>();
@@ -222,9 +223,9 @@ void UMassNavigationObstacleGridProcessor::Execute(FMassEntityManager& EntityMan
 		}
 	});
 
-	UpdateGridEntityQuery.ForEachEntityChunk(EntityManager, Context, [this, &EntityManager, World = EntityManager.GetWorld()](FMassExecutionContext& Context)
+	UpdateGridEntityQuery.ForEachEntityChunk(EntityManager, Context, [this, &EntityManager](FMassExecutionContext& Context)
 	{
-		FNavigationObstacleHashGrid2D& HashGrid = Context.GetMutableSubsystemChecked<UMassNavigationSubsystem>(World).GetObstacleGridMutable();
+		FNavigationObstacleHashGrid2D& HashGrid = Context.GetMutableSubsystemChecked<UMassNavigationSubsystem>().GetObstacleGridMutable();
 		const int32 NumEntities = Context.GetNumEntities();
 
 		TConstArrayView<FTransformFragment> LocationList = Context.GetFragmentView<FTransformFragment>();
@@ -256,9 +257,9 @@ void UMassNavigationObstacleGridProcessor::Execute(FMassEntityManager& EntityMan
 		}
 	});
 
-	RemoveFromGridEntityQuery.ForEachEntityChunk(EntityManager, Context, [this, &EntityManager, World = EntityManager.GetWorld()](FMassExecutionContext& Context)
+	RemoveFromGridEntityQuery.ForEachEntityChunk(EntityManager, Context, [this, &EntityManager](FMassExecutionContext& Context)
 	{
-		FNavigationObstacleHashGrid2D& HashGrid = Context.GetMutableSubsystemChecked<UMassNavigationSubsystem>(World).GetObstacleGridMutable();
+		FNavigationObstacleHashGrid2D& HashGrid = Context.GetMutableSubsystemChecked<UMassNavigationSubsystem>().GetObstacleGridMutable();
 		const int32 NumEntities = Context.GetNumEntities();
 
 		TArrayView<FMassNavigationObstacleGridCellLocationFragment> AvoidanceObstacleCellLocationList = Context.GetMutableFragmentView<FMassNavigationObstacleGridCellLocationFragment>();
@@ -293,9 +294,9 @@ void UMassNavigationObstacleRemoverProcessor::ConfigureQueries()
 
 void UMassNavigationObstacleRemoverProcessor::Execute(FMassEntityManager& EntityManager, FMassExecutionContext& Context)
 {
-	EntityQuery.ForEachEntityChunk(EntityManager, Context, [this, World = EntityManager.GetWorld()](FMassExecutionContext& Context)
+	EntityQuery.ForEachEntityChunk(EntityManager, Context, [this](FMassExecutionContext& Context)
 	{
-		FNavigationObstacleHashGrid2D& HashGrid = Context.GetMutableSubsystemChecked<UMassNavigationSubsystem>(World).GetObstacleGridMutable();
+		FNavigationObstacleHashGrid2D& HashGrid = Context.GetMutableSubsystemChecked<UMassNavigationSubsystem>().GetObstacleGridMutable();
 		const int32 NumEntities = Context.GetNumEntities();
 		const TArrayView<FMassNavigationObstacleGridCellLocationFragment> AvoidanceObstacleCellLocationList = Context.GetMutableFragmentView<FMassNavigationObstacleGridCellLocationFragment>();
 
diff --git a/Engine/Plugins/AI/MassAI/Source/MassNavigation/Private/MassSmoothOrientationProcessors.cpp b/Engine/Plugins/AI/MassAI/Source/MassNavigation/Private/MassSmoothOrientationProcessors.cpp
--- a/Engine/Plugins/AI/MassAI/Source/MassNavigation/Private/MassSmoothOrientationProcessors.cpp
+++ b/Engine/Plugins/AI/MassAI/Source/MassNavigation/Private/MassSmoothOrientationProcessors.cpp
@@ -4,6 +4,7 @@
 #include "SmoothOrientation/MassSmoothOrientationFragments.h"
 #include "MassCommandBuffer.h"
 #include "MassCommonFragments.h"
+#include "MassExecutionContext.h"
 #include "MassNavigationFragments.h"
 #include "MassMovementFragments.h"
 #include "Math/UnrealMathUtility.h"
diff --git a/Engine/Plugins/AI/MassAI/Source/MassNavigation/Private/MassSteeringProcessors.cpp b/Engine/Plugins/AI/MassAI/Source/MassNavigation/Private/MassSteeringProcessors.cpp
--- a/Engine/Plugins/AI/MassAI/Source/MassNavigation/Private/MassSteeringProcessors.cpp
+++ b/Engine/Plugins/AI/MassAI/Source/MassNavigation/Private/MassSteeringProcessors.cpp
@@ -4,6 +4,7 @@
 #include "MassCommonUtils.h"
 #include "MassCommandBuffer.h"
 #include "MassCommonFragments.h"
+#include "MassExecutionContext.h"
 #include "MassMovementFragments.h"
 #include "MassNavigationFragments.h"
 #include "Steering/MassSteeringFragments.h"
diff --git a/Engine/Plugins/AI/MassAI/Source/MassZoneGraphNavigation/Private/MassZoneGraphNavigationProcessors.cpp b/Engine/Plugins/AI/MassAI/Source/MassZoneGraphNavigation/Private/MassZoneGraphNavigationProcessors.cpp
--- a/Engine/Plugins/AI/MassAI/Source/MassZoneGraphNavigation/Private/MassZoneGraphNavigationProcessors.cpp
+++ b/Engine/Plugins/AI/MassAI/Source/MassZoneGraphNavigation/Private/MassZoneGraphNavigationProcessors.cpp
@@ -7,6 +7,7 @@
 #include "MassNavigationUtils.h"
 #include "Avoidance/MassAvoidanceFragments.h"
 #include "MassCommonFragments.h"
+#include "MassExecutionContext.h"
 #include "MassSignalSubsystem.h"
 #include "ZoneGraphSubsystem.h"
 #include "ZoneGraphQuery.h"
@@ -55,9 +56,9 @@ void UMassZoneGraphLocationInitializer::ConfigureQueries()
 
 void UMassZoneGraphLocationInitializer::Execute(FMassEntityManager& EntityManager, FMassExecutionContext& Context)
 {
-	EntityQuery.ForEachEntityChunk(EntityManager, Context, [this, World = EntityManager.GetWorld()](FMassExecutionContext& Context)
+	EntityQuery.ForEachEntityChunk(EntityManager, Context, [this](FMassExecutionContext& Context)
 	{
-		const UZoneGraphSubsystem& ZoneGraphSubsystem = Context.GetSubsystemChecked<UZoneGraphSubsystem>(World);
+		const UZoneGraphSubsystem& ZoneGraphSubsystem = Context.GetSubsystemChecked<UZoneGraphSubsystem>();
 
 		const int32 NumEntities = Context.GetNumEntities();
 		const TArrayView<FMassZoneGraphLaneLocationFragment> LaneLocationList = Context.GetMutableFragmentView<FMassZoneGraphLaneLocationFragment>();
@@ -149,9 +150,9 @@ void UMassZoneGraphPathFollowProcessor::Execute(FMassEntityManager& EntityManage
 	TArray<FMassEntityHandle> EntitiesToSignalPathDone;
 	TArray<FMassEntityHandle> EntitiesToSignalLaneChanged;
 
-	EntityQuery_Conditional.ForEachEntityChunk(EntityManager, Context, [this, &EntitiesToSignalPathDone, &EntitiesToSignalLaneChanged, World = EntityManager.GetWorld()](FMassExecutionContext& Context)
+	EntityQuery_Conditional.ForEachEntityChunk(EntityManager, Context, [this, &EntitiesToSignalPathDone, &EntitiesToSignalLaneChanged](FMassExecutionContext& Context)
 	{
-		const UZoneGraphSubsystem& ZoneGraphSubsystem = Context.GetSubsystemChecked<UZoneGraphSubsystem>(World);
+		const UZoneGraphSubsystem& ZoneGraphSubsystem = Context.GetSubsystemChecked<UZoneGraphSubsystem>();
 
 		const int32 NumEntities = Context.GetNumEntities();
 		const TArrayView<FMassZoneGraphShortPathFragment> ShortPathList = Context.GetMutableFragmentView<FMassZoneGraphShortPathFragment>();
diff --git a/Engine/Plugins/AI/MassCrowd/Source/MassCrowd/Private/MassCrowdBubble.cpp b/Engine/Plugins/AI/MassCrowd/Source/MassCrowd/Private/MassCrowdBubble.cpp
--- a/Engine/Plugins/AI/MassCrowd/Source/MassCrowd/Private/MassCrowdBubble.cpp
+++ b/Engine/Plugins/AI/MassCrowd/Source/MassCrowd/Private/MassCrowdBubble.cpp
@@ -5,6 +5,7 @@
 #include "DrawDebugHelpers.h"
 #include "MassSpawnerTypes.h"
 #include "MassEntityView.h"
+#include "MassExecutionContext.h"
 #include "MassCommonFragments.h"
 #include "MassMovementFragments.h"
 
diff --git a/Engine/Plugins/AI/MassCrowd/Source/MassCrowd/Private/MassCrowdNavigationProcessor.cpp b/Engine/Plugins/AI/MassCrowd/Source/MassCrowd/Private/MassCrowdNavigationProcessor.cpp
--- a/Engine/Plugins/AI/MassCrowd/Source/MassCrowd/Private/MassCrowdNavigationProcessor.cpp
+++ b/Engine/Plugins/AI/MassCrowd/Source/MassCrowd/Private/MassCrowdNavigationProcessor.cpp
@@ -2,6 +2,7 @@
 
 #include "MassCrowdNavigationProcessor.h"
 #include "MassCommonFragments.h"
+#include "MassExecutionContext.h"
 #include "MassAIBehaviorTypes.h"
 #include "MassCrowdSubsystem.h"
 #include "MassCrowdFragments.h"
@@ -40,9 +41,9 @@ void UMassCrowdLaneTrackingSignalProcessor::Initialize(UObject& Owner)
 
 void UMassCrowdLaneTrackingSignalProcessor::SignalEntities(FMassEntityManager& EntityManager, FMassExecutionContext& Context, FMassSignalNameLookup& EntitySignals)
 {
-	EntityQuery.ForEachEntityChunk(EntityManager, Context, [this, World = EntityManager.GetWorld()](FMassExecutionContext& Context)
+	EntityQuery.ForEachEntityChunk(EntityManager, Context, [this](FMassExecutionContext& Context)
 	{
-		UMassCrowdSubsystem& MassCrowdSubsystem = Context.GetMutableSubsystemChecked<UMassCrowdSubsystem>(World);
+		UMassCrowdSubsystem& MassCrowdSubsystem = Context.GetMutableSubsystemChecked<UMassCrowdSubsystem>();
 		const int32 NumEntities = Context.GetNumEntities();
 		const TConstArrayView<FMassZoneGraphLaneLocationFragment> LaneLocationList = Context.GetFragmentView<FMassZoneGraphLaneLocationFragment>();
 		const TArrayView<FMassCrowdLaneTrackingFragment> LaneTrackingList = Context.GetMutableFragmentView<FMassCrowdLaneTrackingFragment>();
@@ -80,9 +81,9 @@ void UMassCrowdLaneTrackingDestructor::ConfigureQueries()
 
 void UMassCrowdLaneTrackingDestructor::Execute(FMassEntityManager& EntityManager, FMassExecutionContext& Context)
 {
-	EntityQuery.ForEachEntityChunk(EntityManager, Context, [this, World = EntityManager.GetWorld()](FMassExecutionContext& Context)
+	EntityQuery.ForEachEntityChunk(EntityManager, Context, [this](FMassExecutionContext& Context)
 	{
-		UMassCrowdSubsystem& MassCrowdSubsystem = Context.GetMutableSubsystemChecked<UMassCrowdSubsystem>(World);
+		UMassCrowdSubsystem& MassCrowdSubsystem = Context.GetMutableSubsystemChecked<UMassCrowdSubsystem>();
 		const int32 NumEntities = Context.GetNumEntities();
 		const TConstArrayView<FMassCrowdLaneTrackingFragment> LaneTrackingList = Context.GetFragmentView<FMassCrowdLaneTrackingFragment>();
 
diff --git a/Engine/Plugins/AI/MassCrowd/Source/MassCrowd/Private/MassCrowdReplicator.cpp b/Engine/Plugins/AI/MassCrowd/Source/MassCrowd/Private/MassCrowdReplicator.cpp
--- a/Engine/Plugins/AI/MassCrowd/Source/MassCrowd/Private/MassCrowdReplicator.cpp
+++ b/Engine/Plugins/AI/MassCrowd/Source/MassCrowd/Private/MassCrowdReplicator.cpp
@@ -5,6 +5,7 @@
 #include "MassClientBubbleHandler.h"
 #include "MassCrowdBubble.h"
 #include "MassCommonFragments.h"
+#include "MassExecutionContext.h"
 #include "MassLODSubsystem.h"
 #include "MassCrowdFragments.h"
 #include "MassReplicationPathHandlers.h"
diff --git a/Engine/Plugins/AI/MassCrowd/Source/MassCrowd/Private/MassCrowdServerRepresentationLODProcessor.cpp b/Engine/Plugins/AI/MassCrowd/Source/MassCrowd/Private/MassCrowdServerRepresentationLODProcessor.cpp
--- a/Engine/Plugins/AI/MassCrowd/Source/MassCrowd/Private/MassCrowdServerRepresentationLODProcessor.cpp
+++ b/Engine/Plugins/AI/MassCrowd/Source/MassCrowd/Private/MassCrowdServerRepresentationLODProcessor.cpp
@@ -2,6 +2,7 @@
 
 #include "MassCrowdServerRepresentationLODProcessor.h"
 #include "MassCommonFragments.h"
+#include "MassExecutionContext.h"
 #include "MassRepresentationFragments.h"
 
 namespace UE::MassCrowd
@@ -50,7 +51,7 @@ void UMassCrowdServerRepresentationLODProcessor::Execute(FMassEntityManager& Ent
 {
 	TRACE_CPUPROFILER_EVENT_SCOPE(TEXT("CrowdServerRepresentationLOD"))
 
-	const UMassLODSubsystem& LODSubsystem = Context.GetSubsystemChecked<UMassLODSubsystem>(EntityManager.GetWorld());
+	const UMassLODSubsystem& LODSubsystem = Context.GetSubsystemChecked<UMassLODSubsystem>();
 	const TArray<FViewerInfo>& Viewers = LODSubsystem.GetViewers();
 	LODCalculator.PrepareExecution(Viewers);
 	
diff --git a/Engine/Plugins/AI/MassCrowd/Source/MassCrowd/Private/MassCrowdVisualizationLODProcessor.cpp b/Engine/Plugins/AI/MassCrowd/Source/MassCrowd/Private/MassCrowdVisualizationLODProcessor.cpp
--- a/Engine/Plugins/AI/MassCrowd/Source/MassCrowd/Private/MassCrowdVisualizationLODProcessor.cpp
+++ b/Engine/Plugins/AI/MassCrowd/Source/MassCrowd/Private/MassCrowdVisualizationLODProcessor.cpp
@@ -2,6 +2,7 @@
 
 #include "MassCrowdVisualizationLODProcessor.h"
 #include "MassCommonFragments.h"
+#include "MassExecutionContext.h"
 #include "MassCrowdFragments.h"
 
 namespace UE::MassCrowd
diff --git a/Engine/Plugins/AI/MassCrowd/Source/MassCrowd/Private/MassCrowdVisualizationProcessor.cpp b/Engine/Plugins/AI/MassCrowd/Source/MassCrowd/Private/MassCrowdVisualizationProcessor.cpp
--- a/Engine/Plugins/AI/MassCrowd/Source/MassCrowd/Private/MassCrowdVisualizationProcessor.cpp
+++ b/Engine/Plugins/AI/MassCrowd/Source/MassCrowd/Private/MassCrowdVisualizationProcessor.cpp
@@ -2,6 +2,7 @@
 
 #include "MassCrowdVisualizationProcessor.h"
 #include "MassCrowdFragments.h"
+#include "MassExecutionContext.h"
 #include "Components/CapsuleComponent.h"
 #include "GameFramework/CharacterMovementComponent.h"
 #include "MassActorSubsystem.h"
diff --git a/Engine/Plugins/Runtime/MassEntity/Source/MassEntity/MassEntity.Build.cs b/Engine/Plugins/Runtime/MassEntity/Source/MassEntity/MassEntity.Build.cs
--- a/Engine/Plugins/Runtime/MassEntity/Source/MassEntity/MassEntity.Build.cs
+++ b/Engine/Plugins/Runtime/MassEntity/Source/MassEntity/MassEntity.Build.cs
@@ -26,9 +26,10 @@ namespace UnrealBuildTool.Rules
 				}
 			);
 
-			if (Target.bBuildEditor == true)
+			if (Target.bBuildEditor || Target.bCompileAgainstEditor)
 			{
 				PrivateDependencyModuleNames.Add("UnrealEd");
+				PrivateDependencyModuleNames.Add("EditorSubsystem");
 			}
 
 			if (Target.bBuildDeveloperTools == true)
diff --git a/Engine/Plugins/Runtime/MassEntity/Source/MassEntity/Private/MassArchetypeData.cpp b/Engine/Plugins/Runtime/MassEntity/Source/MassEntity/Private/MassArchetypeData.cpp
--- a/Engine/Plugins/Runtime/MassEntity/Source/MassEntity/Private/MassArchetypeData.cpp
+++ b/Engine/Plugins/Runtime/MassEntity/Source/MassEntity/Private/MassArchetypeData.cpp
@@ -2,6 +2,7 @@
 
 #include "MassArchetypeData.h"
 #include "MassEntityTypes.h"
+#include "MassExecutionContext.h"
 #include "Misc/StringBuilder.h"
 
 //////////////////////////////////////////////////////////////////////
diff --git a/Engine/Plugins/Runtime/MassEntity/Source/MassEntity/Private/MassEntityManager.cpp b/Engine/Plugins/Runtime/MassEntity/Source/MassEntity/Private/MassEntityManager.cpp
--- a/Engine/Plugins/Runtime/MassEntity/Source/MassEntity/Private/MassEntityManager.cpp
+++ b/Engine/Plugins/Runtime/MassEntity/Source/MassEntity/Private/MassEntityManager.cpp
@@ -1351,7 +1351,7 @@ void FMassEntityManager::GetValidArchetypes(const FMassEntityQuery& Query, TArra
 
 FMassExecutionContext FMassEntityManager::CreateExecutionContext(const float DeltaSeconds)
 {
-	FMassExecutionContext ExecutionContext(AsShared(), DeltaSeconds);
+	FMassExecutionContext ExecutionContext(*this, DeltaSeconds);
 	ExecutionContext.SetDeferredCommandBuffer(DeferredCommandBuffer);
 	return MoveTemp(ExecutionContext);
 }
diff --git a/Engine/Plugins/Runtime/MassEntity/Source/MassEntity/Private/MassEntityQuery.cpp b/Engine/Plugins/Runtime/MassEntity/Source/MassEntity/Private/MassEntityQuery.cpp
--- a/Engine/Plugins/Runtime/MassEntity/Source/MassEntity/Private/MassEntityQuery.cpp
+++ b/Engine/Plugins/Runtime/MassEntity/Source/MassEntity/Private/MassEntityQuery.cpp
@@ -5,6 +5,7 @@
 #include "MassEntityManager.h"
 #include "MassArchetypeData.h"
 #include "MassCommandBuffer.h"
+#include "MassExecutionContext.h"
 #include "VisualLogger/VisualLogger.h"
 #include "Async/ParallelFor.h"
 #include "Containers/UnrealString.h"
@@ -161,14 +162,12 @@ void FMassEntityQuery::ForEachEntityChunk(FMassEntityManager& EntityManager, FMa
 		FScopedSubsystemRequirementsRestore(FMassExecutionContext& ExecutionContext)
 			: CachedExecutionContext(ExecutionContext)
 		{
-			ConstSubsystemsBitSet = ExecutionContext.ConstSubsystemsBitSet;
-			MutableSubsystemsBitSet = ExecutionContext.MutableSubsystemsBitSet;
+			CachedExecutionContext.GetSubsystemRequirementBits(ConstSubsystemsBitSet, MutableSubsystemsBitSet);
 		}
 
 		~FScopedSubsystemRequirementsRestore()
 		{
-			CachedExecutionContext.ConstSubsystemsBitSet = ConstSubsystemsBitSet;
-			CachedExecutionContext.MutableSubsystemsBitSet = MutableSubsystemsBitSet;
+			CachedExecutionContext.SetSubsystemRequirementBits(ConstSubsystemsBitSet, MutableSubsystemsBitSet);
 		}
 
 		FMassExecutionContext& CachedExecutionContext;
diff --git a/Engine/Plugins/Runtime/MassEntity/Source/MassEntity/Private/MassExecutionContext.cpp b/Engine/Plugins/Runtime/MassEntity/Source/MassEntity/Private/MassExecutionContext.cpp
--- a/Engine/Plugins/Runtime/MassEntity/Source/MassEntity/Private/MassExecutionContext.cpp
+++ b/Engine/Plugins/Runtime/MassEntity/Source/MassEntity/Private/MassExecutionContext.cpp
@@ -4,11 +4,21 @@
 #include "MassArchetypeData.h"
 
 
+//////////////////////////////////////////////////////////////////////////
+// FMassExecutionContext
+
+FMassExecutionContext::FMassExecutionContext(FMassEntityManager& InEntityManager, const float InDeltaTimeSeconds, const bool bInFlushDeferredCommands)
+	: SubsystemAccess(InEntityManager.GetWorld())
+	, DeltaTimeSeconds(InDeltaTimeSeconds)
+	, EntityManager(InEntityManager.AsShared())
+	, bFlushDeferredCommands(bInFlushDeferredCommands)
+{
+}
+
 void FMassExecutionContext::FlushDeferred()
 {
 	if (bFlushDeferredCommands && DeferredCommandBuffer)
 	{
-		check(EntityManager);
 		EntityManager->FlushCommands(DeferredCommandBuffer);
 	}
 }
@@ -23,61 +33,10 @@ void FMassExecutionContext::ClearExecutionData()
 	ChunkSerialModificationNumber = -1;
 }
 
-bool FMassExecutionContext::CacheSubsystem(const UWorld* World, const uint32 SystemIndex)
-{
-	if (UNLIKELY(Subsystems.IsValidIndex(SystemIndex) == false))
-	{
-		Subsystems.AddZeroed(Subsystems.Num() - SystemIndex + 1);
-	}
-
-	if (Subsystems[SystemIndex])
-	{
-		return true;
-	}
-
-	const UClass* SubsystemClass = FMassExternalSubsystemBitSet::GetTypeAtIndex(SystemIndex);
-	checkSlow(SubsystemClass);
-	
-	const TSubclassOf<UWorldSubsystem> WorldSubsystemClass(const_cast<UClass*>(SubsystemClass));
-	checkSlow(*WorldSubsystemClass);
-		
-	UWorldSubsystem* SystemInstance = FMassExternalSubsystemTraits::GetInstance<UWorldSubsystem>(World, WorldSubsystemClass);
-	Subsystems[SystemIndex] = SystemInstance;
-	return SystemInstance != nullptr;
-}
-
 bool FMassExecutionContext::CacheSubsystemRequirements(const UWorld* World, const FMassSubsystemRequirements& SubsystemRequirements)
 {
-	if (SubsystemRequirements.IsEmpty())
-	{
-		return true;
-	}
-
-	for (FMassExternalSubsystemBitSet::FIndexIterator It = SubsystemRequirements.GetRequiredConstSubsystems().GetIndexIterator(); It; ++It)
-	{
-		if (CacheSubsystem(World, *It) == false)
-		{
-			return false;
-		}
-	}
-
-	for (FMassExternalSubsystemBitSet::FIndexIterator It = SubsystemRequirements.GetRequiredMutableSubsystems().GetIndexIterator(); It; ++It)
-	{
-		if (CacheSubsystem(World, *It) == false)
-		{
-			return false;
-		}
-	}
-
-	ConstSubsystemsBitSet = SubsystemRequirements.GetRequiredConstSubsystems();
-	MutableSubsystemsBitSet = SubsystemRequirements.GetRequiredMutableSubsystems();
-	return true;
-}
-
-void FMassExecutionContext::SetSubsystemRequirements(const FMassSubsystemRequirements& SubsystemRequirements)
-{
-	ConstSubsystemsBitSet = SubsystemRequirements.GetRequiredConstSubsystems();
-	MutableSubsystemsBitSet = SubsystemRequirements.GetRequiredMutableSubsystems();
+	check(EntityManager->GetWorld() == World);
+	return SubsystemAccess.CacheSubsystemRequirements(SubsystemRequirements);
 }
 
 void FMassExecutionContext::SetEntityCollection(const FMassArchetypeEntityCollection& InEntityCollection)
@@ -130,3 +89,8 @@ void FMassExecutionContext::SetFragmentRequirements(const FMassFragmentRequireme
 		}
 	}
 }
+
+UWorld* FMassExecutionContext::GetWorld() 
+{ 
+	return EntityManager->GetWorld(); 
+}
diff --git a/Engine/Plugins/Runtime/MassEntity/Source/MassEntity/Private/MassExecutor.cpp b/Engine/Plugins/Runtime/MassEntity/Source/MassEntity/Private/MassExecutor.cpp
--- a/Engine/Plugins/Runtime/MassEntity/Source/MassEntity/Private/MassExecutor.cpp
+++ b/Engine/Plugins/Runtime/MassEntity/Source/MassEntity/Private/MassExecutor.cpp
@@ -4,6 +4,7 @@
 #include "MassProcessingTypes.h"
 #include "MassProcessor.h"
 #include "MassCommandBuffer.h"
+#include "MassExecutionContext.h"
 #include "ProfilingDebugging/CpuProfilerTrace.h"
 
 namespace UE::Mass::Executor
@@ -84,7 +85,7 @@ void RunProcessorsView(TArrayView<UMassProcessor*> Processors, FMassProcessingCo
 
 	TRACE_CPUPROFILER_EVENT_SCOPE_STR("MassExecutor RunProcessorsView")
 
-	FMassExecutionContext ExecutionContext(ProcessingContext.EntityManager, ProcessingContext.DeltaSeconds);
+	FMassExecutionContext ExecutionContext(*ProcessingContext.EntityManager.Get(), ProcessingContext.DeltaSeconds);
 	if (EntityCollection)
 	{
 		ExecutionContext.SetEntityCollection(*EntityCollection);
@@ -186,7 +187,7 @@ FGraphEventRef TriggerParallelTasks(UMassProcessor& Processor, FMassProcessingCo
 	TRACE_CPUPROFILER_EVENT_SCOPE_STR("MassExecutor RunParallel")
 
 	// not going through FMassEntityManager::CreateExecutionContext on purpose - we do need a separate command buffer
-	FMassExecutionContext ExecutionContext(ProcessingContext.EntityManager, ProcessingContext.DeltaSeconds);
+	FMassExecutionContext ExecutionContext(*ProcessingContext.EntityManager.Get(), ProcessingContext.DeltaSeconds);
 	TSharedPtr<FMassCommandBuffer> CommandBuffer = ProcessingContext.CommandBuffer
 		? ProcessingContext.CommandBuffer : MakeShareable(new FMassCommandBuffer());
 	ExecutionContext.SetDeferredCommandBuffer(CommandBuffer);
diff --git a/Engine/Plugins/Runtime/MassEntity/Source/MassEntity/Private/MassProcessor.cpp b/Engine/Plugins/Runtime/MassEntity/Source/MassEntity/Private/MassProcessor.cpp
--- a/Engine/Plugins/Runtime/MassEntity/Source/MassEntity/Private/MassProcessor.cpp
+++ b/Engine/Plugins/Runtime/MassEntity/Source/MassEntity/Private/MassProcessor.cpp
@@ -6,6 +6,7 @@
 #include "VisualLogger/VisualLogger.h"
 #include "Engine/World.h"
 #include "MassCommandBuffer.h"
+#include "MassExecutionContext.h"
 
 #include UE_INLINE_GENERATED_CPP_BY_NAME(MassProcessor)
 
diff --git a/Engine/Plugins/Runtime/MassEntity/Source/MassEntity/Private/MassSubsystemAccess.cpp b/Engine/Plugins/Runtime/MassEntity/Source/MassEntity/Private/MassSubsystemAccess.cpp
new file mode 100644
--- /dev/null
+++ b/Engine/Plugins/Runtime/MassEntity/Source/MassEntity/Private/MassSubsystemAccess.cpp
@@ -0,0 +1,115 @@
+// Copyright Epic Games, Inc. All Rights Reserved.
+
+#include "MassSubsystemAccess.h"
+
+
+namespace UE::Mass::Private
+{
+	template<typename T>
+	TSubclassOf<T> ConvertToSubsystemClass(TSubclassOf<USubsystem> SubsystemClass)
+	{
+		return *(reinterpret_cast<TSubclassOf<T>*>(&SubsystemClass));
+	}
+}
+
+//////////////////////////////////////////////////////////////////////////
+// FMassSubsystemAccess
+
+FMassSubsystemAccess::FMassSubsystemAccess(UWorld* InWorld)
+	: World(InWorld)
+{
+	Subsystems.AddZeroed(FMassExternalSubsystemBitSet::GetMaxNum());
+}
+
+USubsystem* FMassSubsystemAccess::FetchSubsystemInstance(UWorld* World, TSubclassOf<USubsystem> SubsystemClass)
+{
+	check(SubsystemClass);
+	if (SubsystemClass->IsChildOf<UWorldSubsystem>())
+	{
+		check(World);
+		return World->GetSubsystemBase(UE::Mass::Private::ConvertToSubsystemClass<UWorldSubsystem>(SubsystemClass));
+	}
+	if (SubsystemClass->IsChildOf<UEngineSubsystem>())
+	{
+		return GEngine->GetEngineSubsystemBase(UE::Mass::Private::ConvertToSubsystemClass<UEngineSubsystem>(SubsystemClass));
+	}
+	if (SubsystemClass->IsChildOf<UGameInstanceSubsystem>())
+	{
+		check(World && World->GetGameInstance());
+		return World->GetGameInstance()->GetSubsystemBase(UE::Mass::Private::ConvertToSubsystemClass<UGameInstanceSubsystem>(SubsystemClass));
+	}
+	if (SubsystemClass->IsChildOf<ULocalPlayerSubsystem>())
+	{
+		check(World);
+		return World->GetFirstLocalPlayerFromController()->GetSubsystemBase(UE::Mass::Private::ConvertToSubsystemClass<ULocalPlayerSubsystem>(SubsystemClass));
+	}
+#if WITH_EDITOR
+	if (SubsystemClass->IsChildOf<UEditorSubsystem>())
+	{
+		return GEditor->GetEditorSubsystemBase(UE::Mass::Private::ConvertToSubsystemClass<UEditorSubsystem>(SubsystemClass));
+	}
+#endif // WITH_EDITOR
+	return nullptr;
+}
+
+bool FMassSubsystemAccess::CacheSubsystemRequirements(const FMassSubsystemRequirements& SubsystemRequirements)
+{
+	if (SubsystemRequirements.IsEmpty())
+	{
+		return true;
+	}
+
+	bool bResult = true;
+
+	for (FMassExternalSubsystemBitSet::FIndexIterator It = SubsystemRequirements.GetRequiredConstSubsystems().GetIndexIterator(); It && bResult; ++It)
+	{
+		bResult = bResult && CacheSubsystem(*It);
+	}
+
+	for (FMassExternalSubsystemBitSet::FIndexIterator It = SubsystemRequirements.GetRequiredMutableSubsystems().GetIndexIterator(); It && bResult; ++It)
+	{
+		bResult = bResult && CacheSubsystem(*It);
+	}
+
+	if (bResult)
+	{
+		ConstSubsystemsBitSet = SubsystemRequirements.GetRequiredConstSubsystems();
+		MutableSubsystemsBitSet = SubsystemRequirements.GetRequiredMutableSubsystems();
+	}
+
+	return bResult;
+}
+
+bool FMassSubsystemAccess::CacheSubsystem(const uint32 SystemIndex)
+{
+	if (UNLIKELY(Subsystems.IsValidIndex(SystemIndex) == false))
+	{
+		Subsystems.AddZeroed(Subsystems.Num() - SystemIndex + 1);
+	}
+
+	if (Subsystems[SystemIndex])
+	{
+		return true;
+	}
+
+	const UClass* SubsystemClass = FMassExternalSubsystemBitSet::GetTypeAtIndex(SystemIndex);
+	checkSlow(SubsystemClass);
+
+	TSubclassOf<USubsystem> SubsystemSubclass(const_cast<UClass*>(SubsystemClass));
+	checkSlow(*SubsystemSubclass);
+
+	if (SubsystemSubclass)
+	{
+		USubsystem* SystemInstance = FMassSubsystemAccess::FetchSubsystemInstance(World.Get(), SubsystemSubclass);
+		Subsystems[SystemIndex] = SystemInstance;
+		return SystemInstance != nullptr;
+	}
+
+	return false;
+}
+
+void FMassSubsystemAccess::SetSubsystemRequirements(const FMassSubsystemRequirements& SubsystemRequirements)
+{
+	ConstSubsystemsBitSet = SubsystemRequirements.GetRequiredConstSubsystems();
+	MutableSubsystemsBitSet = SubsystemRequirements.GetRequiredMutableSubsystems();
+}
diff --git a/Engine/Plugins/Runtime/MassEntity/Source/MassEntity/Public/MassEntityManager.h b/Engine/Plugins/Runtime/MassEntity/Source/MassEntity/Public/MassEntityManager.h
--- a/Engine/Plugins/Runtime/MassEntity/Source/MassEntity/Public/MassEntityManager.h
+++ b/Engine/Plugins/Runtime/MassEntity/Source/MassEntity/Public/MassEntityManager.h
@@ -9,9 +9,7 @@
 #include "MassEntityQuery.h"
 #include "StructUtilsTypes.h"
 #include "MassObserverManager.h"
-#include "MassExecutionContext.h"
 #include "Containers/MpscQueue.h"
-#include "MassExecutionContext.h"
 #include "MassRequirementAccessDetector.h"
 
 
diff --git a/Engine/Plugins/Runtime/MassEntity/Source/MassEntity/Public/MassEntityTypes.h b/Engine/Plugins/Runtime/MassEntity/Source/MassEntity/Public/MassEntityTypes.h
--- a/Engine/Plugins/Runtime/MassEntity/Source/MassEntity/Public/MassEntityTypes.h
+++ b/Engine/Plugins/Runtime/MassEntity/Source/MassEntity/Public/MassEntityTypes.h
@@ -124,7 +124,7 @@ DECLARE_STRUCTTYPEBITSET_EXPORTED(MASSENTITY_API, FMassFragmentBitSet, FMassFrag
 DECLARE_STRUCTTYPEBITSET_EXPORTED(MASSENTITY_API, FMassTagBitSet, FMassTag);
 DECLARE_STRUCTTYPEBITSET_EXPORTED(MASSENTITY_API, FMassChunkFragmentBitSet, FMassChunkFragment);
 DECLARE_STRUCTTYPEBITSET_EXPORTED(MASSENTITY_API, FMassSharedFragmentBitSet, FMassSharedFragment);
-DECLARE_CLASSTYPEBITSET_EXPORTED(MASSENTITY_API, FMassExternalSubsystemBitSet, UWorldSubsystem);
+DECLARE_CLASSTYPEBITSET_EXPORTED(MASSENTITY_API, FMassExternalSubsystemBitSet, USubsystem);
 
 /** The type summarily describing a composition of an entity or an archetype. It contains information on both the
  *  fragments as well as tags */
diff --git a/Engine/Plugins/Runtime/MassEntity/Source/MassEntity/Public/MassExecutionContext.h b/Engine/Plugins/Runtime/MassEntity/Source/MassEntity/Public/MassExecutionContext.h
--- a/Engine/Plugins/Runtime/MassEntity/Source/MassEntity/Public/MassExecutionContext.h
+++ b/Engine/Plugins/Runtime/MassEntity/Source/MassEntity/Public/MassExecutionContext.h
@@ -5,6 +5,7 @@
 #include "MassEntityTypes.h"
 #include "MassExternalSubsystemTraits.h"
 #include "MassEntityQuery.h"
+#include "MassSubsystemAccess.h"
 
 
 struct MASSENTITY_API FMassExecutionContext
@@ -34,9 +35,7 @@ private:
 	using FSharedFragmentView = TFragmentView<FStructView>;
 	TArray<FSharedFragmentView, TInlineAllocator<4>> SharedFragmentViews;
 
-	FMassExternalSubsystemBitSet ConstSubsystemsBitSet;
-	FMassExternalSubsystemBitSet MutableSubsystemsBitSet;
-	TArray<UWorldSubsystem*> Subsystems;
+	FMassSubsystemAccess SubsystemAccess;
 	
 	// mz@todo make this shared ptr thread-safe and never auto-flush in MT environment. 
 	TSharedPtr<FMassCommandBuffer> DeferredCommandBuffer;
@@ -52,7 +51,7 @@ private:
 	int32 ChunkSerialModificationNumber = -1;
 	FMassTagBitSet CurrentArchetypesTagBitSet;
 
-	TSharedPtr<FMassEntityManager> EntityManager;
+	TSharedRef<FMassEntityManager> EntityManager;
 
 #if WITH_MASSENTITY_DEBUG
 	FString DebugExecutionDescription;
@@ -67,25 +66,25 @@ private:
 	TArrayView<FConstSharedFragmentView> GetMutableConstSharedRequirements() { return ConstSharedFragmentViews; }
 	TArrayView<FSharedFragmentView> GetMutableSharedRequirements() { return SharedFragmentViews; }
 
+	void GetSubsystemRequirementBits(FMassExternalSubsystemBitSet& OutConstSubsystemsBitSet, FMassExternalSubsystemBitSet& OutMutableSubsystemsBitSet)
+	{
+		SubsystemAccess.GetSubsystemRequirementBits(OutConstSubsystemsBitSet, OutMutableSubsystemsBitSet);
+	}
+
+	void SetSubsystemRequirementBits(const FMassExternalSubsystemBitSet& InConstSubsystemsBitSet, const FMassExternalSubsystemBitSet& InMutableSubsystemsBitSet)
+	{
+		SubsystemAccess.SetSubsystemRequirementBits(InConstSubsystemsBitSet, InMutableSubsystemsBitSet);
+	}
+
 	EMassExecutionContextType ExecutionType = EMassExecutionContextType::Local;
 
 	friend FMassArchetypeData;
 	friend FMassEntityQuery;
 
 public:
-	explicit FMassExecutionContext(const TSharedPtr<FMassEntityManager>& InEntityManager, const float InDeltaTimeSeconds, const bool bInFlushDeferredCommands = true)
-		: DeltaTimeSeconds(InDeltaTimeSeconds)
-		, EntityManager(InEntityManager)
-		, bFlushDeferredCommands(bInFlushDeferredCommands)
-	{
-		Subsystems.AddZeroed(FMassExternalSubsystemBitSet::GetMaxNum());
-	}
-
-	FMassExecutionContext()
-		: FMassExecutionContext(nullptr, /*InDeltaTimeSeconds=*/0.f)
-	{}
+	explicit FMassExecutionContext(FMassEntityManager& InEntityManager, const float InDeltaTimeSeconds = 0.f, const bool bInFlushDeferredCommands = true);
 
-	FMassEntityManager& GetEntityManagerChecked() { check(EntityManager); return *EntityManager.Get(); }
+	FMassEntityManager& GetEntityManagerChecked() { return EntityManager.Get(); }
 
 #if WITH_MASSENTITY_DEBUG
 	const FString& DebugGetExecutionDesc() const { return DebugExecutionDescription; }
@@ -107,6 +106,8 @@ public:
 		return DeltaTimeSeconds;
 	}
 
+	UWorld* GetWorld();
+
 	TSharedPtr<FMassCommandBuffer> GetSharedDeferredCommandBuffer() const { return DeferredCommandBuffer; }
 	FMassCommandBuffer& Defer() const { checkSlow(DeferredCommandBuffer.IsValid()); return *DeferredCommandBuffer.Get(); }
 
@@ -245,84 +246,52 @@ public:
 		return View->FragmentView;
 	}
 
-	template<typename T>
-	T* GetMutableSubsystem(const UWorld* World)
+	template<typename T, typename = typename TEnableIf<TIsDerivedFrom<T, USubsystem>::IsDerived>::Type>
+	T* GetMutableSubsystem()
 	{
-		// @todo consider getting this directly from entity subsystem - it should cache all the used system
-		const uint32 SystemIndex = FMassExternalSubsystemBitSet::GetTypeIndex<T>();
-		if (ensure(MutableSubsystemsBitSet.IsBitSet(SystemIndex)))
-		{
-			return GetSubsystemInternal<T>(World, SystemIndex);
-		}
-
-		return nullptr;
+		return SubsystemAccess.GetMutableSubsystem<T>();
 	}
 
-	template<typename T>
-	T& GetMutableSubsystemChecked(const UWorld* World)
+	template<typename T, typename = typename TEnableIf<TIsDerivedFrom<T, USubsystem>::IsDerived>::Type>
+	T& GetMutableSubsystemChecked()
 	{
-		T* InstancePtr = GetMutableSubsystem<T>(World);
-		check(InstancePtr);
-		return *InstancePtr;
+		return SubsystemAccess.GetMutableSubsystemChecked<T>();
 	}
 
-	template<typename T>
-	const T* GetSubsystem(const UWorld* World)
+	template<typename T, typename = typename TEnableIf<TIsDerivedFrom<T, USubsystem>::IsDerived>::Type>
+	const T* GetSubsystem()
 	{
-		const uint32 SystemIndex = FMassExternalSubsystemBitSet::GetTypeIndex<T>();
-		if (ensure(ConstSubsystemsBitSet.IsBitSet(SystemIndex) || MutableSubsystemsBitSet.IsBitSet(SystemIndex)))
-		{
-			return GetSubsystemInternal<T>(World, SystemIndex);
-		}
-		return nullptr;
+		return SubsystemAccess.GetSubsystem<T>();
 	}
 
-	template<typename T>
-	const T& GetSubsystemChecked(const UWorld* World)
+	template<typename T, typename = typename TEnableIf<TIsDerivedFrom<T, USubsystem>::IsDerived>::Type>
+	const T& GetSubsystemChecked()
 	{
-		const T* InstancePtr = GetSubsystem<T>(World);
-		check(InstancePtr);
-		return *InstancePtr;
+		return SubsystemAccess.GetSubsystemChecked<T>();
 	}
 
-	template<typename T>
-	T* GetMutableSubsystem(const UWorld* World, const TSubclassOf<UWorldSubsystem> SubsystemClass)
+	template<typename T, typename = typename TEnableIf<TIsDerivedFrom<T, USubsystem>::IsDerived>::Type>
+	T* GetMutableSubsystem(const TSubclassOf<USubsystem> SubsystemClass)
 	{
-		// @todo consider getting this directly from entity subsystem - it should cache all the used system
-		const uint32 SystemIndex = FMassExternalSubsystemBitSet::GetTypeIndex(**SubsystemClass);
-		if (ensure(MutableSubsystemsBitSet.IsBitSet(SystemIndex)))
-		{
-			return GetSubsystemInternal<T>(World, SystemIndex, SubsystemClass);
-		}
-
-		return nullptr;
+		return SubsystemAccess.GetMutableSubsystem<T>(SubsystemClass);
 	}
 
-	template<typename T>
-	T& GetMutableSubsystemChecked(const UWorld* World, const TSubclassOf<UWorldSubsystem> SubsystemClass)
+	template<typename T, typename = typename TEnableIf<TIsDerivedFrom<T, USubsystem>::IsDerived>::Type>
+	T& GetMutableSubsystemChecked(const TSubclassOf<USubsystem> SubsystemClass)
 	{
-		T* InstancePtr = GetMutableSubsystem<T>(World, SubsystemClass);
-		check(InstancePtr);
-		return *InstancePtr;
+		return SubsystemAccess.GetMutableSubsystemChecked<T>(SubsystemClass);
 	}
 
-	template<typename T>
-	const T* GetSubsystem(const UWorld* World, const TSubclassOf<UWorldSubsystem> SubsystemClass)
+	template<typename T, typename = typename TEnableIf<TIsDerivedFrom<T, USubsystem>::IsDerived>::Type>
+	const T* GetSubsystem(const TSubclassOf<USubsystem> SubsystemClass)
 	{
-		const uint32 SystemIndex = FMassExternalSubsystemBitSet::GetTypeIndex(**SubsystemClass);
-		if (ensure(ConstSubsystemsBitSet.IsBitSet(SystemIndex) || MutableSubsystemsBitSet.IsBitSet(SystemIndex)))
-		{
-			return GetSubsystemInternal<T>(World, SystemIndex, SubsystemClass);
-		}
-		return nullptr;
+		return SubsystemAccess.GetSubsystem<T>(SubsystemClass);
 	}
 
-	template<typename T>
-	const T& GetSubsystemChecked(const UWorld* World, const TSubclassOf<UWorldSubsystem> SubsystemClass)
+	template<typename T, typename = typename TEnableIf<TIsDerivedFrom<T, USubsystem>::IsDerived>::Type>
+	const T& GetSubsystemChecked(const TSubclassOf<USubsystem> SubsystemClass)
 	{
-		const T* InstancePtr = GetSubsystem<T>(World, SubsystemClass);
-		check(InstancePtr);
-		return *InstancePtr;
+		return SubsystemAccess.GetSubsystemChecked<T>(SubsystemClass);
 	}
 
 	/** Sparse chunk related operation */
@@ -353,7 +322,10 @@ public:
 	bool CacheSubsystemRequirements(const UWorld* World, const FMassSubsystemRequirements& SubsystemRequirements);
 
 protected:
-	void SetSubsystemRequirements(const FMassSubsystemRequirements& SubsystemRequirements);
+	void SetSubsystemRequirements(const FMassSubsystemRequirements& SubsystemRequirements)
+	{
+		SubsystemAccess.SetSubsystemRequirements(SubsystemRequirements);
+	}
 
 	void SetFragmentRequirements(const FMassFragmentRequirements& FragmentRequirements);
 
@@ -377,39 +349,63 @@ protected:
 		}
 	}
 
+public:
+	//////////////////////////////////////////////////////////////////////////
+	// DEPRECATED functions
+
 	template<typename T>
-	T* GetSubsystemInternal(const UWorld* World, const uint32 SystemIndex)
+	UE_DEPRECATED(5.2, "The Get*Subsystem(World) functions have been deprecated. Use the world-less flavor.")
+	T* GetMutableSubsystem(const UWorld* World)
 	{
-		if (UNLIKELY(Subsystems.IsValidIndex(SystemIndex) == false))
-		{
-			Subsystems.AddZeroed(Subsystems.Num() - SystemIndex + 1);
-		}
+		return SubsystemAccess.GetMutableSubsystem<T>();
+	}
 
-		T* SystemInstance = (T*)Subsystems[SystemIndex];
-		if (SystemInstance == nullptr)
-		{
-			SystemInstance = FMassExternalSubsystemTraits::GetInstance<typename TRemoveConst<T>::Type>(World);
-			Subsystems[SystemIndex] = SystemInstance;
-		}
-		return SystemInstance;
+	template<typename T>
+	UE_DEPRECATED(5.2, "The Get*Subsystem(World) functions have been deprecated. Use the world-less flavor.")
+	T& GetMutableSubsystemChecked(const UWorld* World)
+	{
+		return SubsystemAccess.GetMutableSubsystemChecked<T>();
 	}
 
 	template<typename T>
-	T* GetSubsystemInternal(const UWorld* World, const uint32 SystemIndex, const TSubclassOf<UWorldSubsystem> SubsystemClass)
+	UE_DEPRECATED(5.2, "The Get*Subsystem(World) functions have been deprecated. Use the world-less flavor.")
+	const T* GetSubsystem(const UWorld* World)
 	{
-		if (UNLIKELY(Subsystems.IsValidIndex(SystemIndex) == false))
-		{
-			Subsystems.AddZeroed(Subsystems.Num() - SystemIndex + 1);
-		}
+		return SubsystemAccess.GetSubsystem<T>();
+	}
 
-		T* SystemInstance = (T*)Subsystems[SystemIndex];
-		if (SystemInstance == nullptr)
-		{
-			SystemInstance = FMassExternalSubsystemTraits::GetInstance<typename TRemoveConst<T>::Type>(World, SubsystemClass);
-			Subsystems[SystemIndex] = SystemInstance;
-		}
-		return SystemInstance;
+	template<typename T>
+	UE_DEPRECATED(5.2, "The Get*Subsystem(World) functions have been deprecated. Use the world-less flavor.")
+	const T& GetSubsystemChecked(const UWorld* World)
+	{
+		return SubsystemAccess.GetSubsystemChecked<T>();
 	}
 
-	bool CacheSubsystem(const UWorld* World, const uint32 SystemIndex);
+	template<typename T>
+	UE_DEPRECATED(5.2, "The Get*Subsystem(World) functions have been deprecated. Use the world-less flavor.")
+	T* GetMutableSubsystem(const UWorld* World, const TSubclassOf<USubsystem> SubsystemClass)
+	{
+		return SubsystemAccess.GetMutableSubsystem<T>(SubsystemClass);
+	}
+
+	template<typename T>
+	UE_DEPRECATED(5.2, "The Get*Subsystem(World) functions have been deprecated. Use the world-less flavor.")
+	T& GetMutableSubsystemChecked(const UWorld* World, const TSubclassOf<USubsystem> SubsystemClass)
+	{
+		return SubsystemAccess.GetMutableSubsystemChecked<T>(SubsystemClass);
+	}
+
+	template<typename T>
+	UE_DEPRECATED(5.2, "The Get*Subsystem(World) functions have been deprecated. Use the world-less flavor.")
+	const T* GetSubsystem(const UWorld* World, const TSubclassOf<USubsystem> SubsystemClass)
+	{
+		return SubsystemAccess.GetSubsystem<T>(SubsystemClass);
+	}
+
+	template<typename T>
+	UE_DEPRECATED(5.2, "The Get*Subsystem(World) functions have been deprecated. Use the world-less flavor.")
+	const T& GetSubsystemChecked(const UWorld* World, const TSubclassOf<USubsystem> SubsystemClass)
+	{
+		return SubsystemAccess.GetSubsystemChecked<T>(SubsystemClass);
+	}
 };
diff --git a/Engine/Plugins/Runtime/MassEntity/Source/MassEntity/Public/MassExternalSubsystemTraits.h b/Engine/Plugins/Runtime/MassEntity/Source/MassEntity/Public/MassExternalSubsystemTraits.h
--- a/Engine/Plugins/Runtime/MassEntity/Source/MassEntity/Public/MassExternalSubsystemTraits.h
+++ b/Engine/Plugins/Runtime/MassEntity/Source/MassEntity/Public/MassExternalSubsystemTraits.h
@@ -36,20 +36,22 @@ struct TMassExternalSubsystemTraits final
 
 namespace FMassExternalSubsystemTraits
 {
-	/** 
-	Every TMassExternalSubsystemTraits specialization needs to implement the following. Not supplying default implementations
-	to be able to catch missing implementations and header inclusion at compilation time.
-	
-	This is a getter function that given a UWorld* fetches an instance.
+	/**
+	 * Every TMassExternalSubsystemTraits specialization needs to implement the following. Not supplying default implementations
+	 * to be able to catch missing implementations and header inclusion at compilation time.
+	 * 
+	 * This is a getter function that given a UWorld* fetches an instance.
 	*/
 	template<typename T, typename = typename TEnableIf<TIsDerivedFrom<T, UWorldSubsystem>::IsDerived>::Type>
+	UE_DEPRECATED(5.2, "FMassExternalSubsystemTraits::GetInstance has been marked as deprecated. Use FMassSubsystemAccess::FetchSubsystemInstance instead.")
 	FORCEINLINE T* GetInstance(const UWorld* World)
 	{ 
 		// note that the default implementation works only for UWorldSubsystems
 		return UWorld::GetSubsystem<T>(World);
 	}
 
-	template<typename T>
+	template<typename T, typename = typename TEnableIf<TIsDerivedFrom<T, UWorldSubsystem>::IsDerived>::Type>
+	UE_DEPRECATED(5.2, "FMassExternalSubsystemTraits::GetInstance has been marked as deprecated. Use FMassSubsystemAccess::FetchSubsystemInstance instead.")
 	FORCEINLINE T* GetInstance(const UWorld* World, const TSubclassOf<UWorldSubsystem> SubsystemClass)
 	{
 		return World ? Cast<T>(World->GetSubsystemBase(SubsystemClass)) : (T*)nullptr;
diff --git a/Engine/Plugins/Runtime/MassEntity/Source/MassEntity/Public/MassRequirements.h b/Engine/Plugins/Runtime/MassEntity/Source/MassEntity/Public/MassRequirements.h
--- a/Engine/Plugins/Runtime/MassEntity/Source/MassEntity/Public/MassRequirements.h
+++ b/Engine/Plugins/Runtime/MassEntity/Source/MassEntity/Public/MassRequirements.h
@@ -111,7 +111,7 @@ public:
 		return *this;
 	}
 
-	FMassSubsystemRequirements& AddSubsystemRequirement(const TSubclassOf<UWorldSubsystem> SubsystemClass, const EMassFragmentAccess AccessMode, const bool bGameThreadOnly = true)
+	FMassSubsystemRequirements& AddSubsystemRequirement(const TSubclassOf<USubsystem> SubsystemClass, const EMassFragmentAccess AccessMode, const bool bGameThreadOnly = true)
 	{
 		check(AccessMode != EMassFragmentAccess::None && AccessMode != EMassFragmentAccess::MAX);
 
diff --git a/Engine/Plugins/Runtime/MassEntity/Source/MassEntity/Public/MassSubsystemAccess.h b/Engine/Plugins/Runtime/MassEntity/Source/MassEntity/Public/MassSubsystemAccess.h
new file mode 100644
--- /dev/null
+++ b/Engine/Plugins/Runtime/MassEntity/Source/MassEntity/Public/MassSubsystemAccess.h
@@ -0,0 +1,227 @@
+// Copyright Epic Games, Inc. All Rights Reserved.
+
+#pragma once
+
+#include "MassEntityTypes.h"
+#include "MassRequirements.h"
+#include "Subsystems/EngineSubsystem.h"
+#include "Subsystems/GameInstanceSubsystem.h"
+#include "Subsystems/LocalPlayerSubsystem.h"
+#include "Subsystems/WorldSubsystem.h"
+#include "Engine/LocalPlayer.h"
+#include "Engine/GameInstance.h"
+
+#if WITH_EDITOR
+#include "Editor.h"
+#include "EditorSubsystem.h"
+#endif // WITH_EDITOR
+
+
+struct MASSENTITY_API FMassSubsystemAccess
+{
+	explicit FMassSubsystemAccess(UWorld* InWorld = nullptr);
+
+	//////////////////////////////////////////////////////////////////////////
+	// Statically-typed subsystems
+
+	template<typename T, typename = typename TEnableIf<TIsDerivedFrom<T, USubsystem>::IsDerived>::Type>
+	T* GetMutableSubsystem()
+	{
+		const uint32 SystemIndex = FMassExternalSubsystemBitSet::GetTypeIndex<T>();
+		if (ensure(MutableSubsystemsBitSet.IsBitSet(SystemIndex)))
+		{
+			return GetSubsystemInternal<T>(SystemIndex);
+		}
+
+		return nullptr;
+	}
+
+	template<typename T, typename = typename TEnableIf<TIsDerivedFrom<T, USubsystem>::IsDerived>::Type>
+	T& GetMutableSubsystemChecked()
+	{
+		T* InstancePtr = GetMutableSubsystem<T>();
+		check(InstancePtr);
+		return *InstancePtr;
+	}
+
+	template<typename T, typename = typename TEnableIf<TIsDerivedFrom<T, USubsystem>::IsDerived>::Type>
+	const T* GetSubsystem()
+	{
+		const uint32 SystemIndex = FMassExternalSubsystemBitSet::GetTypeIndex<T>();
+		if (ensure(ConstSubsystemsBitSet.IsBitSet(SystemIndex) || MutableSubsystemsBitSet.IsBitSet(SystemIndex)))
+		{
+			return GetSubsystemInternal<T>(SystemIndex);
+		}
+		return nullptr;
+	}
+
+	template<typename T, typename = typename TEnableIf<TIsDerivedFrom<T, USubsystem>::IsDerived>::Type>
+	const T& GetSubsystemChecked()
+	{
+		const T* InstancePtr = GetSubsystem<T>();
+		check(InstancePtr);
+		return *InstancePtr;
+	}
+
+	//////////////////////////////////////////////////////////////////////////
+	// UClass-provided subsystems
+
+	template<typename T, typename = typename TEnableIf<TIsDerivedFrom<T, USubsystem>::IsDerived>::Type>
+	T* GetMutableSubsystem(const TSubclassOf<USubsystem> SubsystemClass)
+	{
+		const uint32 SystemIndex = FMassExternalSubsystemBitSet::GetTypeIndex(**SubsystemClass);
+		if (ensure(MutableSubsystemsBitSet.IsBitSet(SystemIndex)))
+		{
+			return GetSubsystemInternal<T>(SystemIndex, SubsystemClass);
+		}
+
+		return nullptr;
+	}
+
+	template<typename T, typename = typename TEnableIf<TIsDerivedFrom<T, USubsystem>::IsDerived>::Type>
+	T& GetMutableSubsystemChecked(const TSubclassOf<USubsystem> SubsystemClass)
+	{
+		T* InstancePtr = GetMutableSubsystem<T>(SubsystemClass);
+		check(InstancePtr);
+		return *InstancePtr;
+	}
+
+	template<typename T, typename = typename TEnableIf<TIsDerivedFrom<T, USubsystem>::IsDerived>::Type>
+	const T* GetSubsystem(const TSubclassOf<USubsystem> SubsystemClass)
+	{
+		const uint32 SystemIndex = FMassExternalSubsystemBitSet::GetTypeIndex(**SubsystemClass);
+		if (ensure(ConstSubsystemsBitSet.IsBitSet(SystemIndex) || MutableSubsystemsBitSet.IsBitSet(SystemIndex)))
+		{
+			return GetSubsystemInternal<T>(SystemIndex, SubsystemClass);
+		}
+		return nullptr;
+	}
+
+	template<typename T, typename = typename TEnableIf<TIsDerivedFrom<T, USubsystem>::IsDerived>::Type>
+	const T& GetSubsystemChecked(const TSubclassOf<USubsystem> SubsystemClass)
+	{
+		const T* InstancePtr = GetSubsystem<T>(SubsystemClass);
+		check(InstancePtr);
+		return *InstancePtr;
+	}
+
+	//////////////////////////////////////////////////////////////////////////
+	// remaining API
+
+	bool CacheSubsystemRequirements(const FMassSubsystemRequirements& SubsystemRequirements);
+	void SetSubsystemRequirements(const FMassSubsystemRequirements& SubsystemRequirements);
+
+	void GetSubsystemRequirementBits(FMassExternalSubsystemBitSet& OutConstSubsystemsBitSet, FMassExternalSubsystemBitSet& OutMutableSubsystemsBitSet)
+	{
+		OutConstSubsystemsBitSet = ConstSubsystemsBitSet;
+		OutMutableSubsystemsBitSet = MutableSubsystemsBitSet;
+	}
+
+	void SetSubsystemRequirementBits(const FMassExternalSubsystemBitSet& InConstSubsystemsBitSet, const FMassExternalSubsystemBitSet& InMutableSubsystemsBitSet)
+	{
+		ConstSubsystemsBitSet = InConstSubsystemsBitSet;
+		MutableSubsystemsBitSet = InMutableSubsystemsBitSet;
+	}
+
+	template<typename T>
+	static constexpr bool DoesRequireWorld()
+	{
+		constexpr bool bIsWorldSubsystem = TIsDerivedFrom<T, UWorldSubsystem>::IsDerived;
+		constexpr bool bIsGameInstanceSubsystem = TIsDerivedFrom<T, UGameInstanceSubsystem>::IsDerived;
+		constexpr bool bIsLocalPlayerSubsystem = TIsDerivedFrom<T, ULocalPlayerSubsystem>::IsDerived;
+
+		return (bIsWorldSubsystem || bIsGameInstanceSubsystem || bIsLocalPlayerSubsystem);
+	}
+
+	template<typename T, typename = typename TEnableIf<TIsDerivedFrom<T, USubsystem>::IsDerived>::Type>
+	static T* FetchSubsystemInstance(UWorld* World)
+	{
+		check(World);
+		if constexpr (TIsDerivedFrom<T, UWorldSubsystem>::IsDerived)
+		{
+			return UWorld::GetSubsystem<T>(World);
+		}
+		if constexpr (TIsDerivedFrom<T, UGameInstanceSubsystem>::IsDerived)
+		{
+			return UGameInstance::GetSubsystem<T>(World->GetGameInstance());
+		}
+		if constexpr (TIsDerivedFrom<T, ULocalPlayerSubsystem>::IsDerived)
+		{
+			// note that this default implementation will work only for the first player in a local-coop game
+			// to customize this behavior specialize the FetchSubsystemInstance template function for the type you need. 
+			return ULocalPlayer::GetSubsystem<T>(World->GetFirstLocalPlayerFromController());
+		}
+
+		checkf(false, TEXT("FMassSubsystemAccess::FetchSubsystemInstance: Unhandled world-related USubsystem class %s"), *T::StaticClass()->GetName());
+	}
+	
+	template<typename T, typename = typename TEnableIf<TIsDerivedFrom<T, USubsystem>::IsDerived>::Type>
+	static T* FetchSubsystemInstance()
+	{
+		if constexpr (TIsDerivedFrom<T, UEngineSubsystem>::IsDerived)
+		{
+			return GEngine->GetEngineSubsystem<T>();
+		}
+#if WITH_EDITOR
+		if constexpr (TIsDerivedFrom<T, UEditorSubsystem>::IsDerived)
+		{
+			return GEditor->GetEditorSubsystem<T>();
+		}
+#endif // WITH_EDITOR
+		
+		checkf(false, TEXT("FMassSubsystemAccess::FetchSubsystemInstance: Unhandled world-less USubsystem class %s"), *T::StaticClass()->GetName()); 
+	}
+
+	static USubsystem* FetchSubsystemInstance(UWorld* World, TSubclassOf<USubsystem> SubsystemClass);
+
+protected:
+	template<typename T>
+	T* GetSubsystemInternal(const uint32 SystemIndex)
+	{
+		QUICK_SCOPE_CYCLE_COUNTER(Mass_GetSubsystemInternal)
+
+		if (UNLIKELY(Subsystems.IsValidIndex(SystemIndex) == false))
+		{
+			Subsystems.AddZeroed(Subsystems.Num() - SystemIndex + 1);
+		}
+
+		T* SystemInstance = (T*)Subsystems[SystemIndex];
+		if (SystemInstance == nullptr)
+		{
+			if constexpr (DoesRequireWorld<T>())
+			{
+				SystemInstance = FetchSubsystemInstance<typename TRemoveConst<T>::Type>(World.Get());
+			}
+			else
+			{
+				SystemInstance = FetchSubsystemInstance<typename TRemoveConst<T>::Type>();
+			}
+			Subsystems[SystemIndex] = SystemInstance;
+		}
+		return SystemInstance;
+	}
+
+	template<typename T>
+	T* GetSubsystemInternal(const uint32 SystemIndex, const TSubclassOf<USubsystem> SubsystemClass)
+	{
+		if (UNLIKELY(Subsystems.IsValidIndex(SystemIndex) == false))
+		{
+			Subsystems.AddZeroed(Subsystems.Num() - SystemIndex + 1);
+		}
+
+		USubsystem* SystemInstance = (T*)Subsystems[SystemIndex];
+		if (SystemInstance == nullptr)
+		{
+			SystemInstance = FetchSubsystemInstance(World.Get(), SubsystemClass);
+			Subsystems[SystemIndex] = SystemInstance;
+		}
+		return Cast<T>(SystemInstance);
+	}
+
+	bool CacheSubsystem(const uint32 SystemIndex);
+
+	FMassExternalSubsystemBitSet ConstSubsystemsBitSet;
+	FMassExternalSubsystemBitSet MutableSubsystemsBitSet;
+	TArray<USubsystem*> Subsystems;
+	TWeakObjectPtr<UWorld> World;
+};
diff --git a/Engine/Plugins/Runtime/MassEntity/Source/MassEntityTestSuite/Private/MassCompositeProcessorTest.cpp b/Engine/Plugins/Runtime/MassEntity/Source/MassEntityTestSuite/Private/MassCompositeProcessorTest.cpp
--- a/Engine/Plugins/Runtime/MassEntity/Source/MassEntityTestSuite/Private/MassCompositeProcessorTest.cpp
+++ b/Engine/Plugins/Runtime/MassEntity/Source/MassEntityTestSuite/Private/MassCompositeProcessorTest.cpp
@@ -5,6 +5,7 @@
 
 #include "Engine/World.h"
 #include "MassEntityManager.h"
+#include "MassExecutionContext.h"
 #include "MassProcessingTypes.h"
 #include "MassEntityTestTypes.h"
 #include "MassExecutor.h"
diff --git a/Engine/Plugins/Runtime/MassEntity/Source/MassEntityTestSuite/Private/MassEntityTestFarmPlot.cpp b/Engine/Plugins/Runtime/MassEntity/Source/MassEntityTestSuite/Private/MassEntityTestFarmPlot.cpp
--- a/Engine/Plugins/Runtime/MassEntity/Source/MassEntityTestSuite/Private/MassEntityTestFarmPlot.cpp
+++ b/Engine/Plugins/Runtime/MassEntity/Source/MassEntityTestSuite/Private/MassEntityTestFarmPlot.cpp
@@ -6,6 +6,7 @@
 #include "Engine/CollisionProfile.h"
 #include "MassEntityUtils.h"
 #include "MassExecutor.h"
+#include "MassExecutionContext.h"
 
 #include UE_INLINE_GENERATED_CPP_BY_NAME(MassEntityTestFarmPlot)
 
diff --git a/Engine/Plugins/Runtime/MassEntity/Source/MassEntityTestSuite/Private/MassMultiThreadTest.cpp b/Engine/Plugins/Runtime/MassEntity/Source/MassEntityTestSuite/Private/MassMultiThreadTest.cpp
--- a/Engine/Plugins/Runtime/MassEntity/Source/MassEntityTestSuite/Private/MassMultiThreadTest.cpp
+++ b/Engine/Plugins/Runtime/MassEntity/Source/MassEntityTestSuite/Private/MassMultiThreadTest.cpp
@@ -6,6 +6,7 @@
 #include "MassEntitySettings.h"
 #include "MassExecutor.h"
 #include "MassEntityView.h"
+#include "MassExecutionContext.h"
 
 #define LOCTEXT_NAMESPACE "MassTest"
 
diff --git a/Engine/Plugins/Runtime/MassEntity/Source/MassEntityTestSuite/Private/MassObserverTest.cpp b/Engine/Plugins/Runtime/MassEntity/Source/MassEntityTestSuite/Private/MassObserverTest.cpp
--- a/Engine/Plugins/Runtime/MassEntity/Source/MassEntityTestSuite/Private/MassObserverTest.cpp
+++ b/Engine/Plugins/Runtime/MassEntity/Source/MassEntityTestSuite/Private/MassObserverTest.cpp
@@ -6,6 +6,7 @@
 #include "MassEntityTestTypes.h"
 #include "MassEntityTypes.h"
 #include "MassCommandBuffer.h"
+#include "MassExecutionContext.h"
 
 #define LOCTEXT_NAMESPACE "MassTest"
 
diff --git a/Engine/Plugins/Runtime/MassEntity/Source/MassEntityTestSuite/Private/MassProcessorTest.cpp b/Engine/Plugins/Runtime/MassEntity/Source/MassEntityTestSuite/Private/MassProcessorTest.cpp
--- a/Engine/Plugins/Runtime/MassEntity/Source/MassEntityTestSuite/Private/MassProcessorTest.cpp
+++ b/Engine/Plugins/Runtime/MassEntity/Source/MassEntityTestSuite/Private/MassProcessorTest.cpp
@@ -5,6 +5,7 @@
 
 #include "Engine/World.h"
 #include "MassEntityManager.h"
+#include "MassExecutionContext.h"
 #include "MassProcessingTypes.h"
 #include "MassEntityTestTypes.h"
 #include "MassExecutor.h"
diff --git a/Engine/Plugins/Runtime/MassEntity/Source/MassEntityTestSuite/Private/MassQueryTest.cpp b/Engine/Plugins/Runtime/MassEntity/Source/MassEntityTestSuite/Private/MassQueryTest.cpp
--- a/Engine/Plugins/Runtime/MassEntity/Source/MassEntityTestSuite/Private/MassQueryTest.cpp
+++ b/Engine/Plugins/Runtime/MassEntity/Source/MassEntityTestSuite/Private/MassQueryTest.cpp
@@ -8,6 +8,7 @@
 #include "MassProcessingTypes.h"
 #include "MassEntityTestTypes.h"
 #include "MassExecutor.h"
+#include "MassExecutionContext.h"
 
 #define LOCTEXT_NAMESPACE "MassTest"
 
@@ -111,7 +112,7 @@ struct FQueryTest_ExecuteSingleArchetype : FEntityTestBase
 		
 		int TotalProcessed = 0;
 
-		FMassExecutionContext ExecContext;
+		FMassExecutionContext ExecContext(*EntityManager.Get());
 		FMassEntityQuery Query({ FTestFragment_Float::StaticStruct() });
 		Query.ForEachEntityChunk(*EntityManager, ExecContext, [&TotalProcessed](FMassExecutionContext& Context)
 			{
@@ -155,7 +156,7 @@ struct FQueryTest_ExecuteMultipleArchetypes : FEntityTestBase
 		EntityManager->BatchCreateEntities(FloatsIntsArchetype, FloatsIntsArchetypeCreated, EntitiesCreated);
 
 		int TotalProcessed = 0;
-		FMassExecutionContext ExecContext;
+		FMassExecutionContext ExecContext(*EntityManager.Get());
 		FMassEntityQuery Query({ FTestFragment_Float::StaticStruct() });
 		Query.ForEachEntityChunk(*EntityManager, ExecContext, [&TotalProcessed](FMassExecutionContext& Context)
 			{
@@ -211,7 +212,7 @@ struct FQueryTest_ExecuteSparse : FEntityTestBase
 
 		int TotalProcessed = 0;
 
-		FMassExecutionContext ExecContext;
+		FMassExecutionContext ExecContext(*EntityManager.Get());
 		FMassEntityQuery TestQuery;
 		TestQuery.AddRequirement<FTestFragment_Float>(EMassFragmentAccess::ReadWrite);
 		TestQuery.ForEachEntityChunk(FMassArchetypeEntityCollection(FloatsArchetype, EntitiesToProcess, FMassArchetypeEntityCollection::NoDuplicates)
@@ -439,7 +440,7 @@ struct FQueryTest_UsingOptionalFragment : FEntityTestBase
 		FMassEntityQuery Query;
 		Query.AddRequirement<FTestFragment_Int>(EMassFragmentAccess::ReadWrite, EMassFragmentPresence::Optional);
 		Query.AddRequirement<FTestFragment_Float>(EMassFragmentAccess::ReadWrite, EMassFragmentPresence::All);
-		FMassExecutionContext ExecContext;
+		FMassExecutionContext ExecContext(*EntityManager.Get());
 		Query.ForEachEntityChunk(*EntityManager, ExecContext, [&TotalProcessed, &EmptyIntsViewCount, IntValueSet](FMassExecutionContext& Context) {
 			++TotalProcessed;
 			TArrayView<FTestFragment_Int> Ints = Context.GetMutableFragmentView<FTestFragment_Int>();
@@ -496,7 +497,7 @@ struct FQueryTest_AnyFragment : FEntityTestBase
 			EntityManager->CreateEntity(ArchetypeHandle);
 		}
 
-		FMassExecutionContext TestContext;
+		FMassExecutionContext TestContext(*EntityManager.Get());
 		Query.ForEachEntityChunk(*EntityManager, TestContext, [this](FMassExecutionContext& Context)
 			{
 				TArrayView<FTestFragment_Bool> BoolView = Context.GetMutableFragmentView<FTestFragment_Bool>();
diff --git a/Engine/Plugins/Runtime/MassEntity/Source/MassEntityTestSuite/Private/MassSystemRequirementTest.cpp b/Engine/Plugins/Runtime/MassEntity/Source/MassEntityTestSuite/Private/MassSystemRequirementTest.cpp
--- a/Engine/Plugins/Runtime/MassEntity/Source/MassEntityTestSuite/Private/MassSystemRequirementTest.cpp
+++ b/Engine/Plugins/Runtime/MassEntity/Source/MassEntityTestSuite/Private/MassSystemRequirementTest.cpp
@@ -2,85 +2,176 @@
 
 #include "CoreMinimal.h"
 #include "AITestsCommon.h"
-
+#include "Engine/Engine.h"
 #include "MassProcessorDependencySolver.h"
 #include "MassEntityTestTypes.h"
+#include "MassSubsystemAccess.h"
+#include "MassExecutionContext.h"
 
 #define LOCTEXT_NAMESPACE "MassTest"
 
+UE_DISABLE_OPTIMIZATION_SHIP
+
+//////////////////////////////////////////////////////////////////////////
+// UMassTestCustomSubsystem 
+
+TWeakObjectPtr<UMassTestCustomSubsystem> UMassTestCustomSubsystem::Instance = nullptr;
+
+UMassTestCustomSubsystem* UMassTestCustomSubsystem::Create()
+{
+	ensure(Instance.IsValid() == false);
+	Instance = NewObject<UMassTestCustomSubsystem>();
+	return Instance.Get();
+}
+
+UMassTestCustomSubsystem* UMassTestCustomSubsystem::Get()
+{
+	return Instance.Get();
+}
+
+/** FMassSubsystemAccess::FetchSubsystemInstance specialization for the UMassTestCustomSubsystem */
+template<>
+UMassTestCustomSubsystem* FMassSubsystemAccess::FetchSubsystemInstance<UMassTestCustomSubsystem>()
+{
+	return UMassTestCustomSubsystem::Get();
+}
+
+//////////////////////////////////////////////////////////////////////////
+// Tests
 namespace FMassSystemRequirementTest
 {
 
 struct FSystemRequirementTestBase : FEntityTestBase
 {
 	using Super = FEntityTestBase;
-	UWorld* World = nullptr;
-
-	virtual bool SetUp() override
-	{		
-		World = FAITestHelpers::GetWorld();
-		return Super::SetUp();
-	}
-};	
+	FMassEntityQuery EntityQuery;
 
-struct FMutableRequirement : FSystemRequirementTestBase
-{
-	virtual bool InstantTest() override
+	FSystemRequirementTestBase()
 	{
-		EntityManager->CreateEntity(FloatsArchetype);
+		bMakeWorldEntityManagersOwner = true;
+	}
 
-		FMassEntityQuery EntityQuery;
-		EntityQuery.AddRequirement<FTestFragment_Float>(EMassFragmentAccess::ReadOnly);
-		EntityQuery.AddSubsystemRequirement<UMassTestWorldSubsystem>(EMassFragmentAccess::ReadWrite);
+	virtual bool SetUp() override
+	{		
+		AddAutoDestroyObject(*UMassTestCustomSubsystem::Create());
 
-		check(World);
-		UMassTestWorldSubsystem* TestSystemActual = World->GetSubsystem<UMassTestWorldSubsystem>();
-		UMassTestWorldSubsystem* TestSystem = nullptr;
-		const UMassTestWorldSubsystem* TestSystem2 = nullptr;
-		FMassExecutionContext ExecutionContext;
-		EntityQuery.ForEachEntityChunk(*EntityManager, ExecutionContext, [this, &TestSystem, &TestSystem2](FMassExecutionContext& Context)
+		if (Super::SetUp())
 		{
-			TestSystem = Context.GetMutableSubsystem<UMassTestWorldSubsystem>(World);
-			TestSystem2 = Context.GetSubsystem<UMassTestWorldSubsystem>(World);
-		});
+			EntityManager->CreateEntity(FloatsArchetype);
+			EntityQuery.AddRequirement<FTestFragment_Float>(EMassFragmentAccess::ReadOnly);
+			return true;
+		}
 
-		AITEST_NOT_NULL(TEXT(""), TestSystem);
-		AITEST_EQUAL(TEXT(""), TestSystem, TestSystemActual);
-		AITEST_EQUAL(TEXT(""), TestSystem, TestSystem2);		
-
-		return true;
+		return false;
 	}
-};
-IMPLEMENT_AI_INSTANT_TEST(FMutableRequirement, "System.Mass.Query.System.TrivialMutable");
+};	
 
-struct FConstRequirement : FSystemRequirementTestBase
+//////////////////////////////////////////////////////////////////////////
+// The main point of this specific test is to verify that we're getting the same answers from both
+// FMassEntityQuery and FMassSubsystemAccess.
+struct FSubsystemAccessAPI : FSystemRequirementTestBase
 {
 	virtual bool InstantTest() override
 	{
-		EntityManager->CreateEntity(FloatsArchetype);
-
-		FMassEntityQuery EntityQuery;
-		EntityQuery.AddRequirement<FTestFragment_Float>(EMassFragmentAccess::ReadOnly);
-		EntityQuery.AddSubsystemRequirement<UMassTestWorldSubsystem>(EMassFragmentAccess::ReadOnly);
-
+		UWorld* World = FAITestHelpers::GetWorld();
 		check(World);
-		UMassTestWorldSubsystem* TestMutableSystem = nullptr;
-		const UMassTestWorldSubsystem* TestConstSystem = nullptr;
-		FMassExecutionContext ExecutionContext;
-		EntityQuery.ForEachEntityChunk(*EntityManager, ExecutionContext, [this, &TestMutableSystem, &TestConstSystem](FMassExecutionContext& Context)
-		{
-			// commented out since there's no way to seamlessly handle failed ensures, and the line below should cause that.
-			TestConstSystem = Context.GetSubsystem<UMassTestWorldSubsystem>(World);
-		});
+		FMassSubsystemAccess SubsystemAccessBits(World);
+		FMassSubsystemAccess SubsystemAccessClasses(World);
 
-		AITEST_NOT_NULL(TEXT("It should be possible to only fetch const system instance"), TestConstSystem);
-		AITEST_NULL(TEXT("Mutable system access should not be possible"), TestMutableSystem);
+		EntityQuery.AddSubsystemRequirement<UMassTestWorldSubsystem>(EMassFragmentAccess::ReadWrite);
+		EntityQuery.AddSubsystemRequirement<UMassTestEngineSubsystem>(EMassFragmentAccess::ReadWrite);
+		EntityQuery.AddSubsystemRequirement<UMassTestCustomSubsystem>(EMassFragmentAccess::ReadWrite);
+		EntityQuery.AddSubsystemRequirement<UMassTestLocalPlayerSubsystem>(EMassFragmentAccess::ReadWrite);
+		EntityQuery.AddSubsystemRequirement<UMassTestGameInstanceSubsystem>(EMassFragmentAccess::ReadWrite);
+
+		SubsystemAccessBits.SetSubsystemRequirementBits(EntityQuery.GetRequiredConstSubsystems(), EntityQuery.GetRequiredMutableSubsystems());
+		SubsystemAccessClasses.SetSubsystemRequirements(EntityQuery);
+
+		UMassTestWorldSubsystem* TestWorldSubsystemActual = World->GetSubsystem<UMassTestWorldSubsystem>();
+		UMassTestEngineSubsystem* TestEngineSubsystemActual = GEngine->GetEngineSubsystem<UMassTestEngineSubsystem>();
+		UMassTestCustomSubsystem* TestCustomSubsystemActual = UMassTestCustomSubsystem::Get();
+		ULocalPlayer* LocalPlayer = World->GetFirstLocalPlayerFromController();
+		check(LocalPlayer);
+		UMassTestLocalPlayerSubsystem* TestLocalPlayerSubsystemActual = LocalPlayer->GetSubsystem<UMassTestLocalPlayerSubsystem>();
+		UGameInstance* GameInstance = World->GetGameInstance();
+		check(GameInstance);
+		UMassTestGameInstanceSubsystem* TestGameInstanceSubsystemActual = GameInstance->GetSubsystem<UMassTestGameInstanceSubsystem>();
+
+		UMassTestWorldSubsystem* AccessWorldSubsystem = nullptr;
+		UMassTestEngineSubsystem* AccessEngineSubsystem = nullptr;
+		UMassTestCustomSubsystem* AccessCustomSubsystem = nullptr;
+		UMassTestLocalPlayerSubsystem* AccessLocalPlayerSubsystem = nullptr;
+		UMassTestGameInstanceSubsystem* AccessGameInstanceSubsystem = nullptr;
+
+		const UMassTestWorldSubsystem* ConstAccessWorldSubsystem = nullptr;
+		const UMassTestEngineSubsystem* ConstAccessEngineSubsystem = nullptr;
+		const UMassTestCustomSubsystem* ConstAccessCustomSubsystem = nullptr;
+		const UMassTestLocalPlayerSubsystem* ConstAccessLocalPlayerSubsystem = nullptr;
+		const UMassTestGameInstanceSubsystem* ConstAccessGameInstanceSubsystem = nullptr;
+
+		FMassExecutionContext ExecutionContext(*EntityManager);
+		EntityQuery.ForEachEntityChunk(*EntityManager, ExecutionContext, [&](FMassExecutionContext& Context)
+			{
+				AccessWorldSubsystem = Context.GetMutableSubsystem<UMassTestWorldSubsystem>();
+				AccessEngineSubsystem = Context.GetMutableSubsystem<UMassTestEngineSubsystem>();
+				AccessCustomSubsystem = Context.GetMutableSubsystem<UMassTestCustomSubsystem>();
+				AccessLocalPlayerSubsystem = Context.GetMutableSubsystem<UMassTestLocalPlayerSubsystem>();
+				AccessGameInstanceSubsystem = Context.GetMutableSubsystem<UMassTestGameInstanceSubsystem>();
+
+				ConstAccessWorldSubsystem = Context.GetSubsystem<UMassTestWorldSubsystem>();
+				ConstAccessEngineSubsystem = Context.GetSubsystem<UMassTestEngineSubsystem>();
+				ConstAccessCustomSubsystem = Context.GetSubsystem<UMassTestCustomSubsystem>();
+				ConstAccessLocalPlayerSubsystem = Context.GetSubsystem<UMassTestLocalPlayerSubsystem>();
+				ConstAccessGameInstanceSubsystem = Context.GetSubsystem<UMassTestGameInstanceSubsystem>();
+			});
+
+		AITEST_NOT_NULL(TEXT("WorldSubsystem: Subsystem Actual is expected to be not NULL"), TestWorldSubsystemActual);
+		AITEST_EQUAL(TEXT("WorldSubsystem: Mutable subsystem fetched is expected to be the same as the Actual"), TestWorldSubsystemActual, AccessWorldSubsystem);
+		AITEST_EQUAL(TEXT("WorldSubsystem: Const subsystem fetched is expected to be the same as the Actual"), TestWorldSubsystemActual, ConstAccessWorldSubsystem);
+		AITEST_EQUAL(TEXT("WorldSubsystem: subsystem fetched via bits-requirements is expected to be the same as the Actual")
+			, TestWorldSubsystemActual, SubsystemAccessBits.GetSubsystem<UMassTestWorldSubsystem>());
+		AITEST_EQUAL(TEXT("WorldSubsystem: subsystem fetched via class-requirements is expected to be the same as the Actual")
+			, TestWorldSubsystemActual, SubsystemAccessClasses.GetSubsystem<UMassTestWorldSubsystem>());
+
+		AITEST_NOT_NULL(TEXT("EngineSubsystem: Subsystem Actual is expected to be not NULL"), TestEngineSubsystemActual);
+		AITEST_EQUAL(TEXT("EngineSubsystem: Mutable subsystem fetched is expected to be the same as the Actual"), TestEngineSubsystemActual, AccessEngineSubsystem);
+		AITEST_EQUAL(TEXT("EngineSubsystem: Const subsystem fetched is expected to be the same as the Actual"), TestEngineSubsystemActual, ConstAccessEngineSubsystem);
+		AITEST_EQUAL(TEXT("EngineSubsystem: subsystem fetched via bits-requirements is expected to be the same as the Actual")
+			, TestEngineSubsystemActual, SubsystemAccessBits.GetSubsystem<UMassTestEngineSubsystem>());
+		AITEST_EQUAL(TEXT("EngineSubsystem: subsystem fetched via class-requirements is expected to be the same as the Actual")
+			, TestEngineSubsystemActual, SubsystemAccessClasses.GetSubsystem<UMassTestEngineSubsystem>());
+
+		AITEST_NOT_NULL(TEXT("CustomSubsystem: Subsystem Actual is expected to be not NULL"), TestCustomSubsystemActual);
+		AITEST_EQUAL(TEXT("CustomSubsystem: Mutable subsystem fetched is expected to be the same as the Actual"), TestCustomSubsystemActual, AccessCustomSubsystem);
+		AITEST_EQUAL(TEXT("CustomSubsystem: Const subsystem fetched is expected to be the same as the Actual"), TestCustomSubsystemActual, ConstAccessCustomSubsystem);
+		AITEST_EQUAL(TEXT("CustomSubsystem: subsystem fetched via bits-requirements is expected to be the same as the Actual")
+			, TestCustomSubsystemActual, SubsystemAccessBits.GetSubsystem<UMassTestCustomSubsystem>());
+		AITEST_EQUAL(TEXT("CustomSubsystem: subsystem fetched via class-requirements is expected to be the same as the Actual")
+			, TestCustomSubsystemActual, SubsystemAccessClasses.GetSubsystem<UMassTestCustomSubsystem>());
+
+		AITEST_NOT_NULL(TEXT("LocalPlayerSubsystem: Subsystem Actual is expected to be not NULL"), TestLocalPlayerSubsystemActual);
+		AITEST_EQUAL(TEXT("LocalPlayerSubsystem: Mutable subsystem fetched is expected to be the same as the Actual"), TestLocalPlayerSubsystemActual, AccessLocalPlayerSubsystem);
+		AITEST_EQUAL(TEXT("LocalPlayerSubsystem: Const subsystem fetched is expected to be the same as the Actual"), TestLocalPlayerSubsystemActual, ConstAccessLocalPlayerSubsystem);
+		AITEST_EQUAL(TEXT("LocalPlayerSubsystem: subsystem fetched via bits-requirements is expected to be the same as the Actual")
+			, TestLocalPlayerSubsystemActual, SubsystemAccessBits.GetSubsystem<UMassTestLocalPlayerSubsystem>());
+		AITEST_EQUAL(TEXT("LocalPlayerSubsystem: subsystem fetched via class-requirements is expected to be the same as the Actual")
+			, TestLocalPlayerSubsystemActual, SubsystemAccessClasses.GetSubsystem<UMassTestLocalPlayerSubsystem>());
+
+		AITEST_NOT_NULL(TEXT("GameInstanceSubsystem: Subsystem Actual is expected to be not NULL"), TestGameInstanceSubsystemActual);
+		AITEST_EQUAL(TEXT("GameInstanceSubsystem: Mutable subsystem fetched is expected to be the same as the Actual"), TestGameInstanceSubsystemActual, AccessGameInstanceSubsystem);
+		AITEST_EQUAL(TEXT("GameInstanceSubsystem: Const subsystem fetched is expected to be the same as the Actual"), TestGameInstanceSubsystemActual, ConstAccessGameInstanceSubsystem);
+		AITEST_EQUAL(TEXT("GameInstanceSubsystem: subsystem fetched via bits-requirements is expected to be the same as the Actual")
+			, TestGameInstanceSubsystemActual, SubsystemAccessBits.GetSubsystem<UMassTestGameInstanceSubsystem>());
+		AITEST_EQUAL(TEXT("GameInstanceSubsystem: subsystem fetched via class-requirements is expected to be the same as the Actual")
+			, TestGameInstanceSubsystemActual, SubsystemAccessClasses.GetSubsystem<UMassTestGameInstanceSubsystem>());
 
 		return true;
 	}
 };
-IMPLEMENT_AI_INSTANT_TEST(FConstRequirement, "System.Mass.Query.System.TrivialConst");
+IMPLEMENT_AI_INSTANT_TEST(FSubsystemAccessAPI, "System.Mass.Query.System.AccessAPI");
 
 } // FMassSystemRequirementTest
 
 #undef LOCTEXT_NAMESPACE
+
+UE_ENABLE_OPTIMIZATION_SHIP
\ No newline at end of file
diff --git a/Engine/Plugins/Runtime/MassEntity/Source/MassEntityTestSuite/Private/MassTestTypes.cpp b/Engine/Plugins/Runtime/MassEntity/Source/MassEntityTestSuite/Private/MassTestTypes.cpp
--- a/Engine/Plugins/Runtime/MassEntity/Source/MassEntityTestSuite/Private/MassTestTypes.cpp
+++ b/Engine/Plugins/Runtime/MassEntity/Source/MassEntityTestSuite/Private/MassTestTypes.cpp
@@ -3,6 +3,7 @@
 #include "MassEntityTestTypes.h"
 #include "MassEntityManager.h"
 #include "MassExecutor.h"
+#include "MassExecutionContext.h"
 #include "Engine/World.h"
 
 //----------------------------------------------------------------------//
@@ -10,7 +11,7 @@
 //----------------------------------------------------------------------//
 bool FExecutionTestBase::SetUp()
 {
-	EntityManager = MakeShareable(new FMassEntityManager);
+	EntityManager = MakeShareable(new FMassEntityManager(bMakeWorldEntityManagersOwner ? FAITestHelpers::GetWorld() : nullptr));
 	EntityManager->SetDebugName(TEXT("MassEntityTestSuite"));
 	EntityManager->Initialize();
 
diff --git a/Engine/Plugins/Runtime/MassEntity/Source/MassEntityTestSuite/Public/MassEntityTestTypes.h b/Engine/Plugins/Runtime/MassEntity/Source/MassEntityTestSuite/Public/MassEntityTestTypes.h
--- a/Engine/Plugins/Runtime/MassEntity/Source/MassEntityTestSuite/Public/MassEntityTestTypes.h
+++ b/Engine/Plugins/Runtime/MassEntity/Source/MassEntityTestSuite/Public/MassEntityTestTypes.h
@@ -7,7 +7,10 @@
 #include "MassEntityManager.h"
 #include "AITestsCommon.h"
 #include "Math/RandomStream.h"
+#include "Subsystems/EngineSubsystem.h"
 #include "Subsystems/WorldSubsystem.h"
+#include "Subsystems/LocalPlayerSubsystem.h"
+#include "Subsystems/GameInstanceSubsystem.h"
 #include "Misc/MTAccessDetector.h"
 #include "MassExternalSubsystemTraits.h"
 #include "MassEntityTestTypes.generated.h"
@@ -191,6 +194,7 @@ public:
 struct MASSENTITYTESTSUITE_API FExecutionTestBase : FAITestBase
 {
 	TSharedPtr<FMassEntityManager> EntityManager;
+	bool bMakeWorldEntityManagersOwner = false;
 
 	virtual bool SetUp() override;
 };
@@ -242,3 +246,77 @@ struct TMassExternalSubsystemTraits<UMassTestWorldSubsystem>
 		ThreadSafeWrite = false,
 	};
 };
+
+UCLASS()
+class UMassTestEngineSubsystem : public UEngineSubsystem
+{
+	GENERATED_BODY()
+};
+
+template<>
+struct TMassExternalSubsystemTraits<UMassTestEngineSubsystem>
+{
+	enum
+	{
+		GameThreadOnly = false,
+		ThreadSafeRead = true,
+		ThreadSafeWrite = false,
+	};
+};
+
+UCLASS()
+class UMassTestCustomSubsystem : public USubsystem
+{
+	GENERATED_BODY()
+
+public:
+	static UMassTestCustomSubsystem* Create();
+	static UMassTestCustomSubsystem* Get();
+private:
+	static TWeakObjectPtr<UMassTestCustomSubsystem> Instance;
+};
+
+template<>
+struct TMassExternalSubsystemTraits<UMassTestCustomSubsystem>
+{
+	enum
+	{
+		GameThreadOnly = false,
+		ThreadSafeRead = true,
+		ThreadSafeWrite = false,
+	};
+};
+
+UCLASS()
+class UMassTestLocalPlayerSubsystem : public ULocalPlayerSubsystem
+{
+	GENERATED_BODY()
+};
+
+template<>
+struct TMassExternalSubsystemTraits<UMassTestLocalPlayerSubsystem>
+{
+	enum
+	{
+		GameThreadOnly = false,
+		ThreadSafeRead = true,
+		ThreadSafeWrite = false,
+	};
+}; 
+
+UCLASS()
+class UMassTestGameInstanceSubsystem : public UGameInstanceSubsystem
+{
+	GENERATED_BODY()
+};
+
+template<>
+struct TMassExternalSubsystemTraits<UMassTestGameInstanceSubsystem>
+{
+	enum
+	{
+		GameThreadOnly = false,
+		ThreadSafeRead = true,
+		ThreadSafeWrite = false,
+	};
+};
\ No newline at end of file
diff --git a/Engine/Plugins/Runtime/MassGameplay/Source/MassActors/Private/Translators/MassCapsuleComponentTranslators.cpp b/Engine/Plugins/Runtime/MassGameplay/Source/MassActors/Private/Translators/MassCapsuleComponentTranslators.cpp
--- a/Engine/Plugins/Runtime/MassGameplay/Source/MassActors/Private/Translators/MassCapsuleComponentTranslators.cpp
+++ b/Engine/Plugins/Runtime/MassGameplay/Source/MassActors/Private/Translators/MassCapsuleComponentTranslators.cpp
@@ -3,6 +3,7 @@
 #include "Translators/MassCapsuleComponentTranslators.h"
 #include "MassCommonTypes.h"
 #include "MassEntityManager.h"
+#include "MassExecutionContext.h"
 #include "Components/CapsuleComponent.h"
 
 
diff --git a/Engine/Plugins/Runtime/MassGameplay/Source/MassActors/Private/Translators/MassCharacterMovementTranslators.cpp b/Engine/Plugins/Runtime/MassGameplay/Source/MassActors/Private/Translators/MassCharacterMovementTranslators.cpp
--- a/Engine/Plugins/Runtime/MassGameplay/Source/MassActors/Private/Translators/MassCharacterMovementTranslators.cpp
+++ b/Engine/Plugins/Runtime/MassGameplay/Source/MassActors/Private/Translators/MassCharacterMovementTranslators.cpp
@@ -6,6 +6,7 @@
 #include "VisualLogger/VisualLogger.h"
 #include "MassEntityManager.h"
 #include "MassCommonTypes.h"
+#include "MassExecutionContext.h"
 #include "MassMovementFragments.h"
 #include "MassSpawnerTypes.h"
 
diff --git a/Engine/Plugins/Runtime/MassGameplay/Source/MassActors/Private/Translators/MassSceneComponentLocationTranslator.cpp b/Engine/Plugins/Runtime/MassGameplay/Source/MassActors/Private/Translators/MassSceneComponentLocationTranslator.cpp
--- a/Engine/Plugins/Runtime/MassGameplay/Source/MassActors/Private/Translators/MassSceneComponentLocationTranslator.cpp
+++ b/Engine/Plugins/Runtime/MassGameplay/Source/MassActors/Private/Translators/MassSceneComponentLocationTranslator.cpp
@@ -5,6 +5,7 @@
 #include "Components/SceneComponent.h"
 #include "MassEntityManager.h"
 #include "MassCommonTypes.h"
+#include "MassExecutionContext.h"
 
 //----------------------------------------------------------------------//
 //  UMassSceneComponentLocationToMassTranslator
diff --git a/Engine/Plugins/Runtime/MassGameplay/Source/MassGameplayDebug/Private/AssignDebugVisProcessor.cpp b/Engine/Plugins/Runtime/MassGameplay/Source/MassGameplayDebug/Private/AssignDebugVisProcessor.cpp
--- a/Engine/Plugins/Runtime/MassGameplay/Source/MassGameplayDebug/Private/AssignDebugVisProcessor.cpp
+++ b/Engine/Plugins/Runtime/MassGameplay/Source/MassGameplayDebug/Private/AssignDebugVisProcessor.cpp
@@ -4,7 +4,7 @@
 #include "MassGameplayDebugTypes.h"
 #include "MassDebugVisualizationComponent.h"
 #include "MassDebuggerSubsystem.h"
-#include "Engine/World.h"
+#include "MassExecutionContext.h"
 
 //----------------------------------------------------------------------//
 // UAssignDebugVisProcessor
@@ -31,19 +31,18 @@ void UAssignDebugVisProcessor::Execute(FMassEntityManager& EntityManager, FMassE
 #if WITH_EDITORONLY_DATA
 	QUICK_SCOPE_CYCLE_COUNTER(AssignDebugVisProcessor_Execute);
 	
-	UWorld* World = EntityManager.GetWorld();
 	// @todo this code bit is temporary, so is the Visualizer->DirtyVisuals at the end of the function. Will be wrapped in
 	// "executable task" once that's implemented. 
 	UMassDebugVisualizationComponent* Visualizer = nullptr;
-	UMassDebuggerSubsystem& Debugger = Context.GetMutableSubsystemChecked<UMassDebuggerSubsystem>(World);
+	UMassDebuggerSubsystem& Debugger = Context.GetMutableSubsystemChecked<UMassDebuggerSubsystem>();
 	check(Debugger.GetVisualizationComponent());
 	Visualizer = Debugger.GetVisualizationComponent();
 	// note that this function will create the "visual components" only it they're missing or out of sync. 
 	Debugger.GetVisualizationComponent()->ConditionallyConstructVisualComponent();
 
-	EntityQuery.ForEachEntityChunk(EntityManager, Context, [this, World](FMassExecutionContext& Context)
+	EntityQuery.ForEachEntityChunk(EntityManager, Context, [this](FMassExecutionContext& Context)
 	{
-		UMassDebuggerSubsystem& Debugger = Context.GetMutableSubsystemChecked<UMassDebuggerSubsystem>(World);
+		UMassDebuggerSubsystem& Debugger = Context.GetMutableSubsystemChecked<UMassDebuggerSubsystem>();
 		UMassDebugVisualizationComponent* Visualizer = Debugger.GetVisualizationComponent();
 		check(Visualizer);
 
diff --git a/Engine/Plugins/Runtime/MassGameplay/Source/MassGameplayDebug/Private/DebugVisLocationProcessor.cpp b/Engine/Plugins/Runtime/MassGameplay/Source/MassGameplayDebug/Private/DebugVisLocationProcessor.cpp
--- a/Engine/Plugins/Runtime/MassGameplay/Source/MassGameplayDebug/Private/DebugVisLocationProcessor.cpp
+++ b/Engine/Plugins/Runtime/MassGameplay/Source/MassGameplayDebug/Private/DebugVisLocationProcessor.cpp
@@ -6,6 +6,7 @@
 #include "MassCommonFragments.h"
 #include "Components/HierarchicalInstancedStaticMeshComponent.h"
 #include "MassEntityManager.h"
+#include "MassExecutionContext.h"
 #include "MassMovementFragments.h"
 
 //----------------------------------------------------------------------//
@@ -31,9 +32,9 @@ void UDebugVisLocationProcessor::Execute(FMassEntityManager& EntityManager, FMas
 #if WITH_EDITORONLY_DATA
 	QUICK_SCOPE_CYCLE_COUNTER(DebugVisLocationProcessor_Run);
 
-	EntityQuery.ForEachEntityChunk(EntityManager, Context, [this, World = EntityManager.GetWorld()](FMassExecutionContext& Context)
+	EntityQuery.ForEachEntityChunk(EntityManager, Context, [this](FMassExecutionContext& Context)
 	{
-		UMassDebuggerSubsystem& Debugger = Context.GetMutableSubsystemChecked<UMassDebuggerSubsystem>(World);
+		UMassDebuggerSubsystem& Debugger = Context.GetMutableSubsystemChecked<UMassDebuggerSubsystem>();
 		UMassDebugVisualizationComponent* Visualizer = Debugger.GetVisualizationComponent();
 		check(Visualizer);
 		TArrayView<UHierarchicalInstancedStaticMeshComponent*> VisualDataISMCs = Visualizer->GetVisualDataISMCs();
@@ -93,9 +94,9 @@ void UMassProcessor_UpdateDebugVis::Execute(FMassEntityManager& EntityManager, F
 {
 	QUICK_SCOPE_CYCLE_COUNTER(UMassProcessor_UpdateDebugVis_Run);
 
-	EntityQuery.ForEachEntityChunk(EntityManager, Context, [this, World = EntityManager.GetWorld()](FMassExecutionContext& Context)
+	EntityQuery.ForEachEntityChunk(EntityManager, Context, [this](FMassExecutionContext& Context)
 		{
-			UMassDebuggerSubsystem& Debugger = Context.GetMutableSubsystemChecked<UMassDebuggerSubsystem>(World);
+			UMassDebuggerSubsystem& Debugger = Context.GetMutableSubsystemChecked<UMassDebuggerSubsystem>();
 
 			const int32 NumEntities = Context.GetNumEntities();
 			const TConstArrayView<FTransformFragment> LocationList = Context.GetFragmentView<FTransformFragment>();
diff --git a/Engine/Plugins/Runtime/MassGameplay/Source/MassLOD/Private/MassLODCollectorProcessor.cpp b/Engine/Plugins/Runtime/MassGameplay/Source/MassLOD/Private/MassLODCollectorProcessor.cpp
--- a/Engine/Plugins/Runtime/MassGameplay/Source/MassLOD/Private/MassLODCollectorProcessor.cpp
+++ b/Engine/Plugins/Runtime/MassGameplay/Source/MassLOD/Private/MassLODCollectorProcessor.cpp
@@ -3,6 +3,7 @@
 #include "MassLODCollectorProcessor.h"
 #include "MassLODUtils.h"
 #include "MassCommonFragments.h"
+#include "MassExecutionContext.h"
 #include "Engine/World.h"
 #include "MassSimulationLOD.h"
 
@@ -80,7 +81,7 @@ void UMassLODCollectorProcessor::ExecuteInternal(FMassEntityManager& EntityManag
 
 void UMassLODCollectorProcessor::Execute(FMassEntityManager& EntityManager, FMassExecutionContext& Context)
 {
-	const UMassLODSubsystem& LODSubsystem = Context.GetSubsystemChecked<UMassLODSubsystem>(EntityManager.GetWorld());
+	const UMassLODSubsystem& LODSubsystem = Context.GetSubsystemChecked<UMassLODSubsystem>();
 	const TArray<FViewerInfo>& Viewers = LODSubsystem.GetViewers();
 	Collector.PrepareExecution(Viewers);
 
diff --git a/Engine/Plugins/Runtime/MassGameplay/Source/MassLOD/Private/MassSimulationLOD.cpp b/Engine/Plugins/Runtime/MassGameplay/Source/MassLOD/Private/MassSimulationLOD.cpp
--- a/Engine/Plugins/Runtime/MassGameplay/Source/MassLOD/Private/MassSimulationLOD.cpp
+++ b/Engine/Plugins/Runtime/MassGameplay/Source/MassLOD/Private/MassSimulationLOD.cpp
@@ -3,6 +3,7 @@
 #include "MassSimulationLOD.h"
 
 #include "MassCommonFragments.h"
+#include "MassExecutionContext.h"
 #include "GameFramework/PlayerController.h"
 #include "Engine/World.h"
 #include "VisualLogger/VisualLogger.h"
@@ -119,7 +120,7 @@ void UMassSimulationLODProcessor::Execute(FMassEntityManager& EntityManager, FMa
 	{
 		TRACE_CPUPROFILER_EVENT_SCOPE(TEXT("PrepareExecution"));
 
-		const UMassLODSubsystem& LODSubsystem = Context.GetSubsystemChecked<UMassLODSubsystem>(EntityManager.GetWorld());
+		const UMassLODSubsystem& LODSubsystem = Context.GetSubsystemChecked<UMassLODSubsystem>();
 		const TArray<FViewerInfo>& Viewers = LODSubsystem.GetViewers();
 
 		EntityManager.ForEachSharedFragment<FMassSimulationLODSharedFragment>([&Viewers](FMassSimulationLODSharedFragment& LODSharedFragment)
diff --git a/Engine/Plugins/Runtime/MassGameplay/Source/MassLOD/Public/MassLODCollector.h b/Engine/Plugins/Runtime/MassGameplay/Source/MassLOD/Public/MassLODCollector.h
--- a/Engine/Plugins/Runtime/MassGameplay/Source/MassLOD/Public/MassLODCollector.h
+++ b/Engine/Plugins/Runtime/MassGameplay/Source/MassLOD/Public/MassLODCollector.h
@@ -3,6 +3,7 @@
 #pragma once
 
 #include "MassLODLogic.h"
+#include "MassExecutionContext.h"
 #include "DrawDebugHelpers.h"
 
 
diff --git a/Engine/Plugins/Runtime/MassGameplay/Source/MassLOD/Public/MassLODFragments.h b/Engine/Plugins/Runtime/MassGameplay/Source/MassLOD/Public/MassLODFragments.h
--- a/Engine/Plugins/Runtime/MassGameplay/Source/MassLOD/Public/MassLODFragments.h
+++ b/Engine/Plugins/Runtime/MassGameplay/Source/MassLOD/Public/MassLODFragments.h
@@ -3,7 +3,7 @@
 #pragma once
 
 #include "MassLODTypes.h"
-
+#include "MassExecutionContext.h"
 #include "MassLODFragments.generated.h"
 
 USTRUCT()
diff --git a/Engine/Plugins/Runtime/MassGameplay/Source/MassMovement/Private/MassMovementProcessors.cpp b/Engine/Plugins/Runtime/MassGameplay/Source/MassMovement/Private/MassMovementProcessors.cpp
--- a/Engine/Plugins/Runtime/MassGameplay/Source/MassMovement/Private/MassMovementProcessors.cpp
+++ b/Engine/Plugins/Runtime/MassGameplay/Source/MassMovement/Private/MassMovementProcessors.cpp
@@ -4,6 +4,7 @@
 #include "MassCommonUtils.h"
 #include "MassCommandBuffer.h"
 #include "MassCommonFragments.h"
+#include "MassExecutionContext.h"
 #include "MassMovementFragments.h"
 #include "Math/UnrealMathUtility.h"
 #include "MassSimulationLOD.h"
diff --git a/Engine/Plugins/Runtime/MassGameplay/Source/MassMovement/Private/MassSimpleMovementTrait.cpp b/Engine/Plugins/Runtime/MassGameplay/Source/MassMovement/Private/MassSimpleMovementTrait.cpp
--- a/Engine/Plugins/Runtime/MassGameplay/Source/MassMovement/Private/MassSimpleMovementTrait.cpp
+++ b/Engine/Plugins/Runtime/MassGameplay/Source/MassMovement/Private/MassSimpleMovementTrait.cpp
@@ -5,6 +5,7 @@
 #include "Engine/World.h"
 #include "MassMovementFragments.h"
 #include "MassCommonFragments.h"
+#include "MassExecutionContext.h"
 #include "MassSimulationLOD.h"
 
 
diff --git a/Engine/Plugins/Runtime/MassGameplay/Source/MassMovement/Private/MassVelocityRandomizerTrait.cpp b/Engine/Plugins/Runtime/MassGameplay/Source/MassMovement/Private/MassVelocityRandomizerTrait.cpp
--- a/Engine/Plugins/Runtime/MassGameplay/Source/MassMovement/Private/MassVelocityRandomizerTrait.cpp
+++ b/Engine/Plugins/Runtime/MassGameplay/Source/MassMovement/Private/MassVelocityRandomizerTrait.cpp
@@ -2,6 +2,7 @@
 
 #include "Example/MassVelocityRandomizerTrait.h"
 #include "MassEntityTemplateRegistry.h"
+#include "MassExecutionContext.h"
 #include "Engine/World.h"
 #include "MassMovementFragments.h"
 
diff --git a/Engine/Plugins/Runtime/MassGameplay/Source/MassReplication/Private/MassReplicationFragments.cpp b/Engine/Plugins/Runtime/MassGameplay/Source/MassReplication/Private/MassReplicationFragments.cpp
--- a/Engine/Plugins/Runtime/MassGameplay/Source/MassReplication/Private/MassReplicationFragments.cpp
+++ b/Engine/Plugins/Runtime/MassGameplay/Source/MassReplication/Private/MassReplicationFragments.cpp
@@ -3,6 +3,7 @@
 #include "MassReplicationFragments.h"
 #include "MassReplicationSubsystem.h"
 #include "MassReplicationProcessor.h"
+#include "MassExecutionContext.h"
 #include "Engine/World.h"
 
 //----------------------------------------------------------------------//
@@ -32,9 +33,9 @@ void UMassNetworkIDFragmentInitializer::Execute(FMassEntityManager& EntityManage
 	if (NetMode != NM_Client)
 	{
 #if UE_REPLICATION_COMPILE_SERVER_CODE
-		EntityQuery.ForEachEntityChunk(EntityManager, Context, [World = EntityManager.GetWorld()](FMassExecutionContext& Context)
+		EntityQuery.ForEachEntityChunk(EntityManager, Context, [](FMassExecutionContext& Context)
 			{
-				UMassReplicationSubsystem& ReplicationSubsystem = Context.GetMutableSubsystemChecked<UMassReplicationSubsystem>(World);
+				UMassReplicationSubsystem& ReplicationSubsystem = Context.GetMutableSubsystemChecked<UMassReplicationSubsystem>();
 
 				const TArrayView<FMassNetworkIDFragment> NetworkIDList = Context.GetMutableFragmentView<FMassNetworkIDFragment>();
 				const int32 NumEntities = Context.GetNumEntities();
diff --git a/Engine/Plugins/Runtime/MassGameplay/Source/MassReplication/Private/MassReplicationGridProcessor.cpp b/Engine/Plugins/Runtime/MassGameplay/Source/MassReplication/Private/MassReplicationGridProcessor.cpp
--- a/Engine/Plugins/Runtime/MassGameplay/Source/MassReplication/Private/MassReplicationGridProcessor.cpp
+++ b/Engine/Plugins/Runtime/MassGameplay/Source/MassReplication/Private/MassReplicationGridProcessor.cpp
@@ -5,6 +5,7 @@
 #include "MassReplicationTypes.h"
 #include "MassReplicationFragments.h"
 #include "MassCommonFragments.h"
+#include "MassExecutionContext.h"
 #include "Engine/World.h"
 
 //----------------------------------------------------------------------//
@@ -46,11 +47,9 @@ void UMassReplicationGridProcessor::ConfigureQueries()
 
 void UMassReplicationGridProcessor::Execute(FMassEntityManager& EntityManager, FMassExecutionContext& Context)
 {
-	UWorld* World = EntityManager.GetWorld();
-
-	AddToGridEntityQuery.ForEachEntityChunk(EntityManager, Context, [World](FMassExecutionContext& Context)
+	AddToGridEntityQuery.ForEachEntityChunk(EntityManager, Context, [](FMassExecutionContext& Context)
 	{
-		UMassReplicationSubsystem& ReplicationSubsystem = Context.GetMutableSubsystemChecked<UMassReplicationSubsystem>(World);
+		UMassReplicationSubsystem& ReplicationSubsystem = Context.GetMutableSubsystemChecked<UMassReplicationSubsystem>();
 		FReplicationHashGrid2D& ReplicationGrid = ReplicationSubsystem.GetGridMutable();
 		const int32 NumEntities = Context.GetNumEntities();
 
@@ -72,9 +71,9 @@ void UMassReplicationGridProcessor::Execute(FMassEntityManager& EntityManager, F
 		}
 	});
 
-	UpdateGridEntityQuery.ForEachEntityChunk(EntityManager, Context, [World](FMassExecutionContext& Context)
+	UpdateGridEntityQuery.ForEachEntityChunk(EntityManager, Context, [](FMassExecutionContext& Context)
 	{
-		UMassReplicationSubsystem& ReplicationSubsystem = Context.GetMutableSubsystemChecked<UMassReplicationSubsystem>(World);
+		UMassReplicationSubsystem& ReplicationSubsystem = Context.GetMutableSubsystemChecked<UMassReplicationSubsystem>();
 		FReplicationHashGrid2D& ReplicationGrid = ReplicationSubsystem.GetGridMutable();
 		const int32 NumEntities = Context.GetNumEntities();
 
@@ -103,9 +102,9 @@ void UMassReplicationGridProcessor::Execute(FMassEntityManager& EntityManager, F
 		}
 	});
 
-	RemoveFromGridEntityQuery.ForEachEntityChunk(EntityManager, Context, [World](FMassExecutionContext& Context)
+	RemoveFromGridEntityQuery.ForEachEntityChunk(EntityManager, Context, [](FMassExecutionContext& Context)
 	{
-		UMassReplicationSubsystem& ReplicationSubsystem = Context.GetMutableSubsystemChecked<UMassReplicationSubsystem>(World);
+		UMassReplicationSubsystem& ReplicationSubsystem = Context.GetMutableSubsystemChecked<UMassReplicationSubsystem>();
 		FReplicationHashGrid2D& ReplicationGrid = ReplicationSubsystem.GetGridMutable();
 		const int32 NumEntities = Context.GetNumEntities();
 
@@ -140,9 +139,9 @@ void UMassReplicationGridRemoverProcessor::ConfigureQueries()
 
 void UMassReplicationGridRemoverProcessor::Execute(FMassEntityManager& EntityManager, FMassExecutionContext& Context)
 {
-	EntityQuery.ForEachEntityChunk(EntityManager, Context, [World = EntityManager.GetWorld()](FMassExecutionContext& Context)
+	EntityQuery.ForEachEntityChunk(EntityManager, Context, [](FMassExecutionContext& Context)
 	{
-		UMassReplicationSubsystem& ReplicationSubsystem = Context.GetMutableSubsystemChecked<UMassReplicationSubsystem>(World);
+		UMassReplicationSubsystem& ReplicationSubsystem = Context.GetMutableSubsystemChecked<UMassReplicationSubsystem>();
 		FReplicationHashGrid2D& ReplicationGrid = ReplicationSubsystem.GetGridMutable();
 		const int32 NumEntities = Context.GetNumEntities();
 		const TArrayView<FMassReplicationGridCellLocationFragment> ReplicationCellLocationList = Context.GetMutableFragmentView<FMassReplicationGridCellLocationFragment>();
diff --git a/Engine/Plugins/Runtime/MassGameplay/Source/MassReplication/Private/MassReplicationProcessor.cpp b/Engine/Plugins/Runtime/MassGameplay/Source/MassReplication/Private/MassReplicationProcessor.cpp
--- a/Engine/Plugins/Runtime/MassGameplay/Source/MassReplication/Private/MassReplicationProcessor.cpp
+++ b/Engine/Plugins/Runtime/MassGameplay/Source/MassReplication/Private/MassReplicationProcessor.cpp
@@ -4,6 +4,7 @@
 #include "MassClientBubbleHandler.h"
 #include "MassLODSubsystem.h"
 #include "MassCommonFragments.h"
+#include "MassExecutionContext.h"
 
 namespace UE::Mass::Replication
 {
@@ -168,7 +169,7 @@ void UMassReplicationProcessor::Execute(FMassEntityManager& EntityManager, FMass
 		PrepareExecution(EntityManager);
 	}
 
-	const UMassLODSubsystem& LODSubsystem = Context.GetSubsystemChecked<UMassLODSubsystem>(EntityManager.GetWorld());
+	const UMassLODSubsystem& LODSubsystem = Context.GetSubsystemChecked<UMassLODSubsystem>();
 	const TArray<FViewerInfo>& AllViewersInfo = LODSubsystem.GetViewers();
 	const TArray<FMassClientHandle>& ClientHandles = ReplicationSubsystem->GetClientReplicationHandles();
 	for (const FMassClientHandle ClientHandle : ClientHandles)
diff --git a/Engine/Plugins/Runtime/MassGameplay/Source/MassReplication/Private/MassReplicationTransformHandlers.cpp b/Engine/Plugins/Runtime/MassGameplay/Source/MassReplication/Private/MassReplicationTransformHandlers.cpp
--- a/Engine/Plugins/Runtime/MassGameplay/Source/MassReplication/Private/MassReplicationTransformHandlers.cpp
+++ b/Engine/Plugins/Runtime/MassGameplay/Source/MassReplication/Private/MassReplicationTransformHandlers.cpp
@@ -3,6 +3,7 @@
 #include "MassReplicationTransformHandlers.h"
 #include "MassEntityQuery.h"
 #include "MassCommonFragments.h"
+#include "MassExecutionContext.h"
 #include "AIHelpers.h"
 
 void FMassReplicationProcessorTransformHandlerBase::AddRequirements(FMassEntityQuery& InQuery)
diff --git a/Engine/Plugins/Runtime/MassGameplay/Source/MassReplication/Public/MassClientBubbleHandler.h b/Engine/Plugins/Runtime/MassGameplay/Source/MassReplication/Public/MassClientBubbleHandler.h
--- a/Engine/Plugins/Runtime/MassGameplay/Source/MassReplication/Public/MassClientBubbleHandler.h
+++ b/Engine/Plugins/Runtime/MassGameplay/Source/MassReplication/Public/MassClientBubbleHandler.h
@@ -540,7 +540,7 @@ void TClientBubbleHandlerBase<AgentArrayItem>::PostReplicatedAddEntitiesHelper(c
 		const FMassArchetypeHandle& ArchetypeHandle = MassEntityTemplate->GetArchetype();
 
 
-		FMassExecutionContext ExecContext;
+		FMassExecutionContext ExecContext(EntityManager);
 		FMassEntityQuery Query;
 
 		AddRequirementsForSpawnQuery(Query);
diff --git a/Engine/Plugins/Runtime/MassGameplay/Source/MassRepresentation/Private/MassRepresentationProcessor.cpp b/Engine/Plugins/Runtime/MassGameplay/Source/MassRepresentation/Private/MassRepresentationProcessor.cpp
--- a/Engine/Plugins/Runtime/MassGameplay/Source/MassRepresentation/Private/MassRepresentationProcessor.cpp
+++ b/Engine/Plugins/Runtime/MassGameplay/Source/MassRepresentation/Private/MassRepresentationProcessor.cpp
@@ -50,7 +50,7 @@ void UMassRepresentationProcessor::UpdateRepresentation(FMassExecutionContext& C
 	UMassRepresentationActorManagement* RepresentationActorManagement = RepresentationParams.CachedRepresentationActorManagement;
 	check(RepresentationActorManagement);
 
-	UMassActorSubsystem* MassActorSubsystem = Context.GetMutableSubsystem<UMassActorSubsystem>(RepresentationSubsystem->GetWorld());
+	UMassActorSubsystem* MassActorSubsystem = Context.GetMutableSubsystem<UMassActorSubsystem>();
 
 	FMassEntityManager& CachedEntityManager = Context.GetEntityManagerChecked();
 
diff --git a/Engine/Plugins/Runtime/MassGameplay/Source/MassRepresentation/Private/MassUpdateISMProcessor.cpp b/Engine/Plugins/Runtime/MassGameplay/Source/MassRepresentation/Private/MassUpdateISMProcessor.cpp
--- a/Engine/Plugins/Runtime/MassGameplay/Source/MassRepresentation/Private/MassUpdateISMProcessor.cpp
+++ b/Engine/Plugins/Runtime/MassGameplay/Source/MassRepresentation/Private/MassUpdateISMProcessor.cpp
@@ -6,6 +6,7 @@
 #include "MassEntityManager.h"
 #include "MassRepresentationFragments.h"
 #include "MassCommonFragments.h"
+#include "MassExecutionContext.h"
 #include "MassLODFragments.h"
 #include "Engine/World.h"
 
diff --git a/Engine/Plugins/Runtime/MassGameplay/Source/MassRepresentation/Private/MassVisualizationLODProcessor.cpp b/Engine/Plugins/Runtime/MassGameplay/Source/MassRepresentation/Private/MassVisualizationLODProcessor.cpp
--- a/Engine/Plugins/Runtime/MassGameplay/Source/MassRepresentation/Private/MassVisualizationLODProcessor.cpp
+++ b/Engine/Plugins/Runtime/MassGameplay/Source/MassRepresentation/Private/MassVisualizationLODProcessor.cpp
@@ -1,6 +1,7 @@
 // Copyright Epic Games, Inc. All Rights Reserved.
 
 #include "MassVisualizationLODProcessor.h"
+#include "MassExecutionContext.h"
 
 namespace UE::MassRepresentation
 {
@@ -65,7 +66,7 @@ void UMassVisualizationLODProcessor::Execute(FMassEntityManager& EntityManager,
 
 	{
 		TRACE_CPUPROFILER_EVENT_SCOPE(PrepareExecution)
-		const UMassLODSubsystem& LODSubsystem = Context.GetSubsystemChecked<UMassLODSubsystem>(EntityManager.GetWorld());
+		const UMassLODSubsystem& LODSubsystem = Context.GetSubsystemChecked<UMassLODSubsystem>();
 		const TArray<FViewerInfo>& Viewers = LODSubsystem.GetViewers();
 		EntityManager.ForEachSharedFragment<FMassVisualizationLODSharedFragment>([this, &Viewers](FMassVisualizationLODSharedFragment& LODSharedFragment)
 		{
diff --git a/Engine/Plugins/Runtime/MassGameplay/Source/MassSignals/Private/MassSignalProcessorBase.cpp b/Engine/Plugins/Runtime/MassGameplay/Source/MassSignals/Private/MassSignalProcessorBase.cpp
--- a/Engine/Plugins/Runtime/MassGameplay/Source/MassSignals/Private/MassSignalProcessorBase.cpp
+++ b/Engine/Plugins/Runtime/MassGameplay/Source/MassSignals/Private/MassSignalProcessorBase.cpp
@@ -3,6 +3,7 @@
 #include "MassSignalProcessorBase.h"
 #include "MassSignalSubsystem.h"
 #include "MassArchetypeTypes.h"
+#include "MassExecutionContext.h"
 #include "Engine/World.h"
 #include "Misc/ScopeLock.h"
 
diff --git a/Engine/Plugins/Runtime/MassGameplay/Source/MassSignals/Private/MassSignalSubsystem.cpp b/Engine/Plugins/Runtime/MassGameplay/Source/MassSignals/Private/MassSignalSubsystem.cpp
--- a/Engine/Plugins/Runtime/MassGameplay/Source/MassSignals/Private/MassSignalSubsystem.cpp
+++ b/Engine/Plugins/Runtime/MassGameplay/Source/MassSignals/Private/MassSignalSubsystem.cpp
@@ -1,6 +1,7 @@
 // Copyright Epic Games, Inc. All Rights Reserved.
 
 #include "MassSignalSubsystem.h"
+#include "MassExecutionContext.h"
 #include "MassCommandBuffer.h"
 #include "MassSignalTypes.h"
 #include "VisualLogger/VisualLogger.h"
diff --git a/Engine/Plugins/Runtime/MassGameplay/Source/MassSmartObjects/Private/MassSmartObjectHandler.cpp b/Engine/Plugins/Runtime/MassGameplay/Source/MassSmartObjects/Private/MassSmartObjectHandler.cpp
--- a/Engine/Plugins/Runtime/MassGameplay/Source/MassSmartObjects/Private/MassSmartObjectHandler.cpp
+++ b/Engine/Plugins/Runtime/MassGameplay/Source/MassSmartObjects/Private/MassSmartObjectHandler.cpp
@@ -3,6 +3,7 @@
 #include "MassSmartObjectHandler.h"
 #include "MassSmartObjectBehaviorDefinition.h"
 #include "MassCommonTypes.h"
+#include "MassExecutionContext.h"
 #include "SmartObjectSubsystem.h"
 #include "VisualLogger/VisualLogger.h"
 #include "MassCommandBuffer.h"
diff --git a/Engine/Plugins/Runtime/MassGameplay/Source/MassSmartObjects/Private/MassSmartObjectProcessor.cpp b/Engine/Plugins/Runtime/MassGameplay/Source/MassSmartObjects/Private/MassSmartObjectProcessor.cpp
--- a/Engine/Plugins/Runtime/MassGameplay/Source/MassSmartObjects/Private/MassSmartObjectProcessor.cpp
+++ b/Engine/Plugins/Runtime/MassGameplay/Source/MassSmartObjects/Private/MassSmartObjectProcessor.cpp
@@ -3,6 +3,7 @@
 #include "MassSmartObjectProcessor.h"
 #include "MassCommandBuffer.h"
 #include "MassCommonTypes.h"
+#include "MassExecutionContext.h"
 #include "MassSignalSubsystem.h"
 #include "MassSmartObjectBehaviorDefinition.h"
 #include "MassSmartObjectFragments.h"
@@ -54,10 +55,8 @@ UMassSmartObjectCandidatesFinderProcessor::UMassSmartObjectCandidatesFinderProce
 
 void UMassSmartObjectCandidatesFinderProcessor::Execute(FMassEntityManager& EntityManager, FMassExecutionContext& Context)
 {
-	UWorld* World = EntityManager.GetWorld();
-
-	UMassSignalSubsystem& SignalSubsystem = Context.GetMutableSubsystemChecked<UMassSignalSubsystem>(World);
-	const UZoneGraphAnnotationSubsystem& AnnotationSubsystem = Context.GetSubsystemChecked<UZoneGraphAnnotationSubsystem>(World);
+	UMassSignalSubsystem& SignalSubsystem = Context.GetMutableSubsystemChecked<UMassSignalSubsystem>();
+	const UZoneGraphAnnotationSubsystem& AnnotationSubsystem = Context.GetSubsystemChecked<UZoneGraphAnnotationSubsystem>();
 	
 	// Create filter
 	FSmartObjectRequestFilter Filter;
@@ -87,9 +86,9 @@ void UMassSmartObjectCandidatesFinderProcessor::Execute(FMassEntityManager& Enti
 	};
 
 	// Process world location based requests
-	WorldRequestQuery.ForEachEntityChunk(EntityManager, Context, [this, &Filter, &EntitiesToSignal, &BeginRequestProcessing, &EndRequestProcessing, World](FMassExecutionContext& Context)
+	WorldRequestQuery.ForEachEntityChunk(EntityManager, Context, [this, &Filter, &EntitiesToSignal, &BeginRequestProcessing, &EndRequestProcessing](FMassExecutionContext& Context)
 	{
-		const USmartObjectSubsystem& SmartObjectSubsystem = Context.GetSubsystemChecked<USmartObjectSubsystem>(World);
+		const USmartObjectSubsystem& SmartObjectSubsystem = Context.GetSubsystemChecked<USmartObjectSubsystem>();
 
 		const int32 NumEntities = Context.GetNumEntities();
 		EntitiesToSignal.Reserve(EntitiesToSignal.Num() + NumEntities);
@@ -157,12 +156,12 @@ void UMassSmartObjectCandidatesFinderProcessor::Execute(FMassEntityManager& Enti
 	USmartObjectZoneAnnotations* Annotations = Cast<USmartObjectZoneAnnotations>(AnnotationSubsystem.GetFirstAnnotationForTag(SmartObjectTag));
 
 	LaneRequestQuery.ForEachEntityChunk(EntityManager, Context,
-		[&AnnotationSubsystem, Annotations, &Filter, SmartObjectTag, &EntitiesToSignal, &BeginRequestProcessing, &EndRequestProcessing, World = EntityManager.GetWorld()](FMassExecutionContext& Context)
+		[&AnnotationSubsystem, Annotations, &Filter, SmartObjectTag, &EntitiesToSignal, &BeginRequestProcessing, &EndRequestProcessing](FMassExecutionContext& Context)
 		{
 #if WITH_MASSGAMEPLAY_DEBUG
-			const UZoneGraphSubsystem& ZoneGraphSubsystem = Context.GetSubsystemChecked<UZoneGraphSubsystem>(World);
+			const UZoneGraphSubsystem& ZoneGraphSubsystem = Context.GetSubsystemChecked<UZoneGraphSubsystem>();
 #endif // WITH_MASSGAMEPLAY_DEBUG
-			const USmartObjectSubsystem& SmartObjectSubsystem = Context.GetSubsystemChecked<USmartObjectSubsystem>(World);
+			const USmartObjectSubsystem& SmartObjectSubsystem = Context.GetSubsystemChecked<USmartObjectSubsystem>();
 
 			const int32 NumEntities = Context.GetNumEntities();
 			EntitiesToSignal.Reserve(EntitiesToSignal.Num() + NumEntities);
@@ -321,14 +320,13 @@ UMassSmartObjectTimedBehaviorProcessor::UMassSmartObjectTimedBehaviorProcessor()
 
 void UMassSmartObjectTimedBehaviorProcessor::Execute(FMassEntityManager& EntityManager, FMassExecutionContext& Context)
 {
-	UWorld* World = EntityManager.GetWorld();
 	TArray<FMassEntityHandle> ToRelease;
 
 	QUICK_SCOPE_CYCLE_COUNTER(UMassProcessor_SmartObjectTestBehavior_Run);
 
-	EntityQuery.ForEachEntityChunk(EntityManager, Context, [this, &ToRelease, World](FMassExecutionContext& Context)
+	EntityQuery.ForEachEntityChunk(EntityManager, Context, [this, &ToRelease](FMassExecutionContext& Context)
 	{
-		const USmartObjectSubsystem& SmartObjectSubsystem = Context.GetSubsystemChecked<USmartObjectSubsystem>(World);
+		const USmartObjectSubsystem& SmartObjectSubsystem = Context.GetSubsystemChecked<USmartObjectSubsystem>();
 
 		const int32 NumEntities = Context.GetNumEntities();
 		const TArrayView<FMassSmartObjectUserFragment> UserList = Context.GetMutableFragmentView<FMassSmartObjectUserFragment>();
@@ -377,7 +375,7 @@ void UMassSmartObjectTimedBehaviorProcessor::Execute(FMassEntityManager& EntityM
 
 	if (ToRelease.Num())
 	{
-		UMassSignalSubsystem& SignalSubsystem = Context.GetMutableSubsystemChecked<UMassSignalSubsystem>(World);
+		UMassSignalSubsystem& SignalSubsystem = Context.GetMutableSubsystemChecked<UMassSignalSubsystem>();
 		SignalSubsystem.SignalEntities(UE::Mass::Signals::SmartObjectInteractionDone, ToRelease);
 	}
 }
@@ -401,9 +399,9 @@ void UMassSmartObjectUserFragmentDeinitializer::ConfigureQueries()
 
 void UMassSmartObjectUserFragmentDeinitializer::Execute(FMassEntityManager& EntityManager, FMassExecutionContext& Context)
 {
-	EntityQuery.ForEachEntityChunk(EntityManager, Context, [this, World = EntityManager.GetWorld()](FMassExecutionContext& Context)
+	EntityQuery.ForEachEntityChunk(EntityManager, Context, [this](FMassExecutionContext& Context)
 		{
-			USmartObjectSubsystem& SmartObjectSubsystem = Context.GetMutableSubsystemChecked<USmartObjectSubsystem>(World);
+			USmartObjectSubsystem& SmartObjectSubsystem = Context.GetMutableSubsystemChecked<USmartObjectSubsystem>();
 			const int32 NumEntities = Context.GetNumEntities();
 			const TArrayView<FMassSmartObjectUserFragment> SmartObjectUserFragments = Context.GetMutableFragmentView<FMassSmartObjectUserFragment>();
 
diff --git a/Engine/Plugins/Runtime/MassGameplay/Source/MassSpawner/Private/MassSpawnLocationProcessor.cpp b/Engine/Plugins/Runtime/MassGameplay/Source/MassSpawner/Private/MassSpawnLocationProcessor.cpp
--- a/Engine/Plugins/Runtime/MassGameplay/Source/MassSpawner/Private/MassSpawnLocationProcessor.cpp
+++ b/Engine/Plugins/Runtime/MassGameplay/Source/MassSpawner/Private/MassSpawnLocationProcessor.cpp
@@ -3,6 +3,7 @@
 #include "MassSpawnLocationProcessor.h"
 #include "MassCommonFragments.h"
 #include "MassEntityManager.h"
+#include "MassExecutionContext.h"
 #include "MassSpawnerTypes.h"
 #include "Engine/World.h"
 #include "VisualLogger/VisualLogger.h"
diff --git a/Engine/Plugins/TypedElementsDataStorage/Source/TypedElementsDataStorage/Private/UI/TypedElementSlateWidgetReferenceColumnUpdateProcessor.cpp b/Engine/Plugins/TypedElementsDataStorage/Source/TypedElementsDataStorage/Private/UI/TypedElementSlateWidgetReferenceColumnUpdateProcessor.cpp
--- a/Engine/Plugins/TypedElementsDataStorage/Source/TypedElementsDataStorage/Private/UI/TypedElementSlateWidgetReferenceColumnUpdateProcessor.cpp
+++ b/Engine/Plugins/TypedElementsDataStorage/Source/TypedElementsDataStorage/Private/UI/TypedElementSlateWidgetReferenceColumnUpdateProcessor.cpp
@@ -4,6 +4,7 @@
 
 #include "Elements/Columns/TypedElementSlateWidgetColumns.h"
 #include "MassCommonTypes.h"
+#include "MassExecutionContext.h"
 
 UTypedElementSlateWidgetReferenceColumnUpdateProcessor::UTypedElementSlateWidgetReferenceColumnUpdateProcessor()
 	: ColumnRemovalQuery(*this)
diff --git a/Engine/Plugins/TypedElementsDataStorage/Source/TypedElementsDataStorageUI/Private/TypedElementCounterWidgetProcessor.cpp b/Engine/Plugins/TypedElementsDataStorage/Source/TypedElementsDataStorageUI/Private/TypedElementCounterWidgetProcessor.cpp
--- a/Engine/Plugins/TypedElementsDataStorage/Source/TypedElementsDataStorageUI/Private/TypedElementCounterWidgetProcessor.cpp
+++ b/Engine/Plugins/TypedElementsDataStorage/Source/TypedElementsDataStorageUI/Private/TypedElementCounterWidgetProcessor.cpp
@@ -10,6 +10,7 @@
 #include "TypedElementCounterWidgetConstructor.h"
 #include "TypedElementSubsystems.h"
 #include "MassCommonTypes.h"
+#include "MassExecutionContext.h"
 #include "Widgets/Text/STextBlock.h"
 
 UTypedElementCounterWidgetProcessor::UTypedElementCounterWidgetProcessor()
@@ -33,7 +34,7 @@ void UTypedElementCounterWidgetProcessor::ConfigureQueries()
 void UTypedElementCounterWidgetProcessor::Execute(FMassEntityManager& EntityManager, FMassExecutionContext& Context)
 {
 	ITypedElementDataStorageInterface* DataInterface = 
-		Context.GetMutableSubsystemChecked<UTypedElementDataStorageSubsystem>(EntityManager.GetWorld()).Get();
+		Context.GetMutableSubsystemChecked<UTypedElementDataStorageSubsystem>().Get();
 	checkf(DataInterface, TEXT(
 		"FTypedElementsDataStorageUiModule tried to process widgets before the "
 		"Typed Elements Data Storage interface is available."));
diff --git a/Engine/Source/Developer/AITestSuite/Private/Tests/AITestsCommon.cpp b/Engine/Source/Developer/AITestSuite/Private/Tests/AITestsCommon.cpp
--- a/Engine/Source/Developer/AITestSuite/Private/Tests/AITestsCommon.cpp
+++ b/Engine/Source/Developer/AITestSuite/Private/Tests/AITestsCommon.cpp
@@ -107,6 +107,7 @@ void FAITestBase::TearDown()
 	for (auto AutoDestroyedObject : SpawnedObjects)
 	{
 		AutoDestroyedObject->RemoveFromRoot();
+		AutoDestroyedObject->MarkAsGarbage();
 	}
 	SpawnedObjects.Reset();
 }
