Changed FMassSubsystemAccess::FetchSubsystemInstance to not report missing World when fetching World-related subsystems, and instead return null in such cases.

diff --git a/Engine/Plugins/Runtime/MassEntity/Source/MassEntity/Private/MassSubsystemAccess.cpp b/Engine/Plugins/Runtime/MassEntity/Source/MassEntity/Private/MassSubsystemAccess.cpp
--- a/Engine/Plugins/Runtime/MassEntity/Source/MassEntity/Private/MassSubsystemAccess.cpp
+++ b/Engine/Plugins/Runtime/MassEntity/Source/MassEntity/Private/MassSubsystemAccess.cpp
@@ -26,8 +26,9 @@ USubsystem* FMassSubsystemAccess::FetchSubsystemInstance(UWorld* World, TSubclas
 	check(SubsystemClass);
 	if (SubsystemClass->IsChildOf<UWorldSubsystem>())
 	{
-		check(World);
-		return World->GetSubsystemBase(UE::Mass::Private::ConvertToSubsystemClass<UWorldSubsystem>(SubsystemClass));
+		return World 
+			? World->GetSubsystemBase(UE::Mass::Private::ConvertToSubsystemClass<UWorldSubsystem>(SubsystemClass)) 
+			: nullptr;
 	}
 	if (SubsystemClass->IsChildOf<UEngineSubsystem>())
 	{
@@ -35,13 +36,16 @@ USubsystem* FMassSubsystemAccess::FetchSubsystemInstance(UWorld* World, TSubclas
 	}
 	if (SubsystemClass->IsChildOf<UGameInstanceSubsystem>())
 	{
-		check(World && World->GetGameInstance());
-		return World->GetGameInstance()->GetSubsystemBase(UE::Mass::Private::ConvertToSubsystemClass<UGameInstanceSubsystem>(SubsystemClass));
+		return (World && World->GetGameInstance())
+			? World->GetGameInstance()->GetSubsystemBase(UE::Mass::Private::ConvertToSubsystemClass<UGameInstanceSubsystem>(SubsystemClass))
+			: nullptr;
 	}
 	if (SubsystemClass->IsChildOf<ULocalPlayerSubsystem>())
 	{
-		check(World);
-		return World->GetFirstLocalPlayerFromController()->GetSubsystemBase(UE::Mass::Private::ConvertToSubsystemClass<ULocalPlayerSubsystem>(SubsystemClass));
+		const ULocalPlayer* LocalPlayer = World ? World->GetFirstLocalPlayerFromController() : nullptr;
+		return LocalPlayer
+			? LocalPlayer->GetSubsystemBase(UE::Mass::Private::ConvertToSubsystemClass<ULocalPlayerSubsystem>(SubsystemClass))
+			: nullptr;
 	}
 #if WITH_EDITOR
 	if (SubsystemClass->IsChildOf<UEditorSubsystem>())
