Fixed UMassCompositeProcessor::Populate to not side-step EProcessorExecutionFlags settings of the processors being added

diff --git a/Engine/Plugins/Runtime/MassEntity/Source/MassEntity/Private/MassProcessor.cpp b/Engine/Plugins/Runtime/MassEntity/Source/MassEntity/Private/MassProcessor.cpp
--- a/Engine/Plugins/Runtime/MassEntity/Source/MassEntity/Private/MassProcessor.cpp
+++ b/Engine/Plugins/Runtime/MassEntity/Source/MassEntity/Private/MassProcessor.cpp
@@ -432,6 +432,8 @@ void UMassCompositeProcessor::Populate(TConstArrayView<FMassProcessorOrderInfo>
 {
 	ChildPipeline.Processors.Reset();
 
+	const UWorld* World = GetWorld();
+	const EProcessorExecutionFlags WorldExecutionFlags = World ? UE::Mass::Utils::GetProcessorExecutionFlagsForWold(*World) : EProcessorExecutionFlags::All;
 	const FMassProcessingPhaseConfig& PhaseConfig = GET_MASS_CONFIG_VALUE(GetProcessingPhaseConfig(ProcessingPhase));
 
 	for (const FMassProcessorOrderInfo& ProcessorInfo : OrderedProcessors)
@@ -439,8 +441,10 @@ void UMassCompositeProcessor::Populate(TConstArrayView<FMassProcessorOrderInfo>
 		if (ensureMsgf(ProcessorInfo.NodeType == FMassProcessorOrderInfo::EDependencyNodeType::Processor, TEXT("Encountered unexpected FMassProcessorOrderInfo::EDependencyNodeType while populating %s"), *GetGroupName().ToString()))
 		{
 			checkSlow(ProcessorInfo.Processor);
-			
-			ChildPipeline.AppendProcessor(*ProcessorInfo.Processor);
+			if (ProcessorInfo.Processor->ShouldExecute(WorldExecutionFlags))
+			{
+				ChildPipeline.AppendProcessor(*ProcessorInfo.Processor);
+			}
 		}
 	}
 }
