Made FMassRuntimePipeline::Processors and updated the call sites to use a getters/setters.

diff --git a/Engine/Plugins/Runtime/MassEntity/Source/MassEntity/Private/MassDebugger.cpp b/Engine/Plugins/Runtime/MassEntity/Source/MassEntity/Private/MassDebugger.cpp
--- a/Engine/Plugins/Runtime/MassEntity/Source/MassEntity/Private/MassDebugger.cpp
+++ b/Engine/Plugins/Runtime/MassEntity/Source/MassEntity/Private/MassDebugger.cpp
@@ -314,7 +314,7 @@ TConstArrayView<struct UMassCompositeProcessor::FDependencyNode> FMassDebugger::
 
 TConstArrayView<TObjectPtr<UMassProcessor>> FMassDebugger::GetHostedProcessors(const UMassCompositeProcessor& GraphOwner)
 {
-	return GraphOwner.ChildPipeline.Processors;
+	return GraphOwner.ChildPipeline.GetProcessors();
 }
 
 FString FMassDebugger::GetRequirementsDescription(const FMassFragmentRequirements& Requirements)
diff --git a/Engine/Plugins/Runtime/MassEntity/Source/MassEntity/Private/MassExecutor.cpp b/Engine/Plugins/Runtime/MassEntity/Source/MassEntity/Private/MassExecutor.cpp
--- a/Engine/Plugins/Runtime/MassEntity/Source/MassEntity/Private/MassExecutor.cpp
+++ b/Engine/Plugins/Runtime/MassEntity/Source/MassEntity/Private/MassExecutor.cpp
@@ -13,20 +13,20 @@ void Run(FMassRuntimePipeline& RuntimePipeline, FMassProcessingContext& Processi
 {
 	if (!ensure(ProcessingContext.EntityManager) || 
 		!ensure(ProcessingContext.DeltaSeconds >= 0.f) ||
-		!ensure(RuntimePipeline.Processors.Find(nullptr) == INDEX_NONE))
+		!ensure(RuntimePipeline.GetProcessors().Find(nullptr) == INDEX_NONE))
 	{
 		return;
 	}
 
 	TRACE_CPUPROFILER_EVENT_SCOPE_STR("MassExecutor Run Pipeline")
-	RunProcessorsView(RuntimePipeline.Processors, ProcessingContext);
+	RunProcessorsView(RuntimePipeline.GetMutableProcessors(), ProcessingContext);
 }
 
 void RunSparse(FMassRuntimePipeline& RuntimePipeline, FMassProcessingContext& ProcessingContext, FMassArchetypeHandle Archetype, TConstArrayView<FMassEntityHandle> Entities)
 {
 	if (!ensure(ProcessingContext.EntityManager) ||
-		!ensure(RuntimePipeline.Processors.Find(nullptr) == INDEX_NONE) ||
-		RuntimePipeline.Processors.Num() == 0 ||
+		!ensure(RuntimePipeline.GetProcessors().Find(nullptr) == INDEX_NONE) ||
+		RuntimePipeline.Num() == 0 ||
 		!ensureMsgf(Archetype.IsValid(), TEXT("The Archetype passed in to UE::Mass::Executor::RunSparse is invalid")))
 	{
 		return;
@@ -35,14 +35,14 @@ void RunSparse(FMassRuntimePipeline& RuntimePipeline, FMassProcessingContext& Pr
 	TRACE_CPUPROFILER_EVENT_SCOPE_STR("MassExecutor RunSparseEntities");
 
 	const FMassArchetypeEntityCollection EntityCollection(Archetype, Entities, FMassArchetypeEntityCollection::NoDuplicates);
-	RunProcessorsView(RuntimePipeline.Processors, ProcessingContext, &EntityCollection);
+	RunProcessorsView(RuntimePipeline.GetMutableProcessors(), ProcessingContext, &EntityCollection);
 }
 
 void RunSparse(FMassRuntimePipeline& RuntimePipeline, FMassProcessingContext& ProcessingContext, const FMassArchetypeEntityCollection& EntityCollection)
 {
 	if (!ensure(ProcessingContext.EntityManager) ||
-		!ensure(RuntimePipeline.Processors.Find(nullptr) == INDEX_NONE) ||
-		RuntimePipeline.Processors.Num() == 0 ||
+		!ensure(RuntimePipeline.GetProcessors().Find(nullptr) == INDEX_NONE) ||
+		RuntimePipeline.Num() == 0 ||
 		!ensureMsgf(EntityCollection.GetArchetype().IsValid(), TEXT("The Archetype of EntityCollection passed in to UE::Mass::Executor::RunSparse is invalid")))
 	{
 		return;
@@ -50,7 +50,7 @@ void RunSparse(FMassRuntimePipeline& RuntimePipeline, FMassProcessingContext& Pr
 
 	TRACE_CPUPROFILER_EVENT_SCOPE_STR("MassExecutor RunSparse");
 
-	RunProcessorsView(RuntimePipeline.Processors, ProcessingContext, &EntityCollection);
+	RunProcessorsView(RuntimePipeline.GetMutableProcessors(), ProcessingContext, &EntityCollection);
 }
 
 void Run(UMassProcessor& Processor, FMassProcessingContext& ProcessingContext)
diff --git a/Engine/Plugins/Runtime/MassEntity/Source/MassEntity/Private/MassObserverManager.cpp b/Engine/Plugins/Runtime/MassEntity/Source/MassEntity/Private/MassObserverManager.cpp
--- a/Engine/Plugins/Runtime/MassEntity/Source/MassEntity/Private/MassObserverManager.cpp
+++ b/Engine/Plugins/Runtime/MassEntity/Source/MassEntity/Private/MassObserverManager.cpp
@@ -226,7 +226,7 @@ void FMassObserverManager::HandleFragmentsImpl(FMassProcessingContext& Processin
 		ProcessingContext.AuxData.InitializeAs(Type);
 		FMassRuntimePipeline& Pipeline = (*HandlersContainer).FindChecked(Type);
 
-		UE::Mass::Executor::RunProcessorsView(Pipeline.Processors, ProcessingContext, &EntityCollection);
+		UE::Mass::Executor::RunProcessorsView(Pipeline.GetMutableProcessors(), ProcessingContext, &EntityCollection);
 	}
 }
 
@@ -239,7 +239,7 @@ void FMassObserverManager::HandleSingleEntityImpl(const UScriptStruct& FragmentT
 	ProcessingContext.AuxData.InitializeAs(&FragmentType);
 	FMassRuntimePipeline& Pipeline = (*HandlersContainer).FindChecked(&FragmentType);
 
-	UE::Mass::Executor::RunProcessorsView(Pipeline.Processors, ProcessingContext, &EntityCollection);
+	UE::Mass::Executor::RunProcessorsView(Pipeline.GetMutableProcessors(), ProcessingContext, &EntityCollection);
 }
 
 void FMassObserverManager::AddObserverInstance(const UScriptStruct& FragmentOrTagType, const EMassObservedOperation Operation, UMassProcessor& ObserverProcessor)
diff --git a/Engine/Plugins/Runtime/MassEntity/Source/MassEntity/Private/MassProcessingPhaseManager.cpp b/Engine/Plugins/Runtime/MassEntity/Source/MassEntity/Private/MassProcessingPhaseManager.cpp
--- a/Engine/Plugins/Runtime/MassEntity/Source/MassEntity/Private/MassProcessingPhaseManager.cpp
+++ b/Engine/Plugins/Runtime/MassEntity/Source/MassEntity/Private/MassProcessingPhaseManager.cpp
@@ -145,7 +145,7 @@ void FMassPhaseProcessorConfigurationHelper::Configure(const TSharedPtr<FMassEnt
 	TmpPipeline.CreateFromArray(PhaseConfig.ProcessorCDOs, ProcessorOuter);
 
 	TArray<FMassProcessorOrderInfo> SortedProcessors;
-	FMassProcessorDependencySolver Solver(TmpPipeline.Processors, bIsGameRuntime);
+	FMassProcessorDependencySolver Solver(TmpPipeline.GetMutableProcessors(), bIsGameRuntime);
 
 	Solver.ResolveDependencies(SortedProcessors, EntityManager, OutOptionalResult);
 
diff --git a/Engine/Plugins/Runtime/MassEntity/Source/MassEntity/Private/MassProcessingTypes.cpp b/Engine/Plugins/Runtime/MassEntity/Source/MassEntity/Private/MassProcessingTypes.cpp
--- a/Engine/Plugins/Runtime/MassEntity/Source/MassEntity/Private/MassProcessingTypes.cpp
+++ b/Engine/Plugins/Runtime/MassEntity/Source/MassEntity/Private/MassProcessingTypes.cpp
@@ -45,11 +45,22 @@ FMassProcessingContext::~FMassProcessingContext()
 //----------------------------------------------------------------------//
 void FMassRuntimePipeline::Reset()
 {
+	for (UMassProcessor* Processor : Processors)
+	{
+		if (Processor)
+		{
+			Processor->MarkAsGarbage();
+		}
+	}
 	Processors.Reset();
 }
 
 void FMassRuntimePipeline::Initialize(UObject& Owner)
 {
+	// having nulls in Processors should be rare so we run the "remove all nulls" operation below only if we know 
+	// for sure that there are any nulls to be removed
+	bool bNullsFound = false;
+
 	for (UMassProcessor* Proc : Processors)
 	{
 		if (Proc)
@@ -57,11 +68,21 @@ void FMassRuntimePipeline::Initialize(UObject& Owner)
 			REDIRECT_OBJECT_TO_VLOG(Proc, &Owner);
 			Proc->Initialize(Owner);
 		}
+		else
+		{
+			bNullsFound = true;
+		}
+	}
+
+	if (bNullsFound)
+	{
+		Processors.RemoveAll([](const UMassProcessor* Proc) { return Proc == nullptr; });
 	}
 }
 
 void FMassRuntimePipeline::SetProcessors(TArray<UMassProcessor*>&& InProcessors)
 {
+	Reset();
 	Processors = InProcessors;
 }
 
@@ -203,6 +224,11 @@ void FMassRuntimePipeline::AppendProcessor(TSubclassOf<UMassProcessor> Processor
 	AppendProcessor(*ProcInstance);
 }
 
+void FMassRuntimePipeline::RemoveProcessor(UMassProcessor& InProcessor)
+{
+	Processors.Remove(&InProcessor);
+}
+
 UMassCompositeProcessor* FMassRuntimePipeline::FindTopLevelGroupByName(FName GroupName)
 {
 	for (UMassProcessor* Processor : Processors)
diff --git a/Engine/Plugins/Runtime/MassEntity/Source/MassEntity/Private/MassProcessor.cpp b/Engine/Plugins/Runtime/MassEntity/Source/MassEntity/Private/MassProcessor.cpp
--- a/Engine/Plugins/Runtime/MassEntity/Source/MassEntity/Private/MassProcessor.cpp
+++ b/Engine/Plugins/Runtime/MassEntity/Source/MassEntity/Private/MassProcessor.cpp
@@ -338,7 +338,7 @@ FGraphEventRef UMassCompositeProcessor::DispatchProcessorTasks(const TSharedPtr<
 
 void UMassCompositeProcessor::Execute(FMassEntityManager& EntityManager, FMassExecutionContext& Context)
 {
-	for (UMassProcessor* Proc : ChildPipeline.Processors)
+	for (UMassProcessor* Proc : ChildPipeline.GetMutableProcessors())
 	{
 		check(Proc);
 		Proc->CallExecute(EntityManager, Context);
@@ -347,15 +347,7 @@ void UMassCompositeProcessor::Execute(FMassEntityManager& EntityManager, FMassEx
 
 void UMassCompositeProcessor::Initialize(UObject& Owner)
 {
-	// remove all nulls
-	ChildPipeline.Processors.RemoveAll([](const UMassProcessor* Proc) { return Proc == nullptr; });
-
-	// from this point on we don't expect to have nulls in ChildPipeline.Processors
-	for (UMassProcessor* Proc : ChildPipeline.Processors)
-	{
-		REDIRECT_OBJECT_TO_VLOG(Proc, this);
-		Proc->Initialize(Owner);
-	}
+	ChildPipeline.Initialize(Owner);
 }
 
 void UMassCompositeProcessor::SetProcessors(TArrayView<UMassProcessor*> InProcessorInstances, const TSharedPtr<FMassEntityManager>& EntityManager)
@@ -430,7 +422,7 @@ void UMassCompositeProcessor::BuildFlatProcessingGraph(TConstArrayView<FMassProc
 
 void UMassCompositeProcessor::Populate(TConstArrayView<FMassProcessorOrderInfo> OrderedProcessors)
 {
-	ChildPipeline.Processors.Reset();
+	ChildPipeline.Reset();
 
 	const UWorld* World = GetWorld();
 	const EProcessorExecutionFlags WorldExecutionFlags = World ? UE::Mass::Utils::GetProcessorExecutionFlagsForWold(*World) : EProcessorExecutionFlags::All;
@@ -452,14 +444,14 @@ void UMassCompositeProcessor::Populate(TConstArrayView<FMassProcessorOrderInfo>
 void UMassCompositeProcessor::DebugOutputDescription(FOutputDevice& Ar, int32 Indent) const
 {
 #if WITH_MASSENTITY_DEBUG
-	if (ChildPipeline.Processors.Num() == 0)
+	if (ChildPipeline.Num() == 0)
 	{
 		Ar.Logf(TEXT("%*sGroup %s: []"), Indent, TEXT(""), *GroupName.ToString());
 	}
 	else
 	{
 		Ar.Logf(TEXT("%*sGroup %s:"), Indent, TEXT(""), *GroupName.ToString());
-		for (UMassProcessor* Proc : ChildPipeline.Processors)
+		for (UMassProcessor* Proc : ChildPipeline.GetProcessors())
 		{
 			check(Proc);
 			Ar.Logf(TEXT("\n"));
@@ -472,7 +464,7 @@ void UMassCompositeProcessor::DebugOutputDescription(FOutputDevice& Ar, int32 In
 void UMassCompositeProcessor::SetProcessingPhase(EMassProcessingPhase Phase)
 {
 	Super::SetProcessingPhase(Phase);
-	for (UMassProcessor* Proc : ChildPipeline.Processors)
+	for (UMassProcessor* Proc : ChildPipeline.GetMutableProcessors())
 	{
 		Proc->SetProcessingPhase(Phase);
 	}
diff --git a/Engine/Plugins/Runtime/MassEntity/Source/MassEntity/Public/MassExecutionContext.h b/Engine/Plugins/Runtime/MassEntity/Source/MassEntity/Public/MassExecutionContext.h
--- a/Engine/Plugins/Runtime/MassEntity/Source/MassEntity/Public/MassExecutionContext.h
+++ b/Engine/Plugins/Runtime/MassEntity/Source/MassEntity/Public/MassExecutionContext.h
@@ -316,7 +316,13 @@ public:
 	}
 
 	/** 
-	 * Processes SubsystemRequirements to fetch and process all the indicated subsystems.
+	 * Processes SubsystemRequirements to fetch and process all the indicated subsystems. 
+	 * 
+	 * @param World used to fetch any UWorld-related subsystems. Note that it's perfectly fine to pass a Null World. 
+	 *	The worst case scenario is that any of the UWorldSubsystem subclasses requirements won't be met.
+	 * @oaram SubsystemRequirements indicates all the subsystems that are expected to be accessed. Requesting a subsystem 
+	 *	not indicated by the SubsystemRequirements will result in a failure.
+	 * 
 	 * @return `true` if all required subsystems have been found, `false` otherwise
 	 */
 	bool CacheSubsystemRequirements(const UWorld* World, const FMassSubsystemRequirements& SubsystemRequirements);
diff --git a/Engine/Plugins/Runtime/MassEntity/Source/MassEntity/Public/MassProcessingTypes.h b/Engine/Plugins/Runtime/MassEntity/Source/MassEntity/Public/MassProcessingTypes.h
--- a/Engine/Plugins/Runtime/MassEntity/Source/MassEntity/Public/MassProcessingTypes.h
+++ b/Engine/Plugins/Runtime/MassEntity/Source/MassEntity/Public/MassProcessingTypes.h
@@ -72,9 +72,11 @@ struct MASSENTITY_API FMassRuntimePipeline
 {
 	GENERATED_BODY()
 
+private:
 	UPROPERTY()
 	TArray<TObjectPtr<UMassProcessor>> Processors;
 
+public:
 	void Reset();
 	void Initialize(UObject& Owner);
 	
@@ -107,12 +109,18 @@ struct MASSENTITY_API FMassRuntimePipeline
 	/** Creates an instance of ProcessorClass and adds it to Processors without any additional checks */
 	void AppendProcessor(TSubclassOf<UMassProcessor> ProcessorClass, UObject& InOwner);
 
+	void RemoveProcessor(UMassProcessor& InProcessor);
+
 	/** goes through Processor looking for a UMassCompositeProcessor instance which GroupName matches the one given as the parameter */
 	UMassCompositeProcessor* FindTopLevelGroupByName(const FName GroupName);
 
 	bool HasProcessorOfExactClass(TSubclassOf<UMassProcessor> InClass) const;
 	bool IsEmpty() const { return Processors.IsEmpty();}
 
+	int32 Num() const { return Processors.Num(); }
+	TConstArrayView<TObjectPtr<UMassProcessor>> GetProcessors() const { return Processors; }
+	TArrayView<TObjectPtr<UMassProcessor>> GetMutableProcessors() { return Processors; }
+
 	MASSENTITY_API friend uint32 GetTypeHash(const FMassRuntimePipeline& Instance);
 };
 
