Mass: Fix mutable uses of instanced structs and struct views.

diff --git a/Engine/Plugins/Runtime/MassEntity/Source/MassEntity/Private/MassArchetypeData.cpp b/Engine/Plugins/Runtime/MassEntity/Source/MassEntity/Private/MassArchetypeData.cpp
--- a/Engine/Plugins/Runtime/MassEntity/Source/MassEntity/Private/MassArchetypeData.cpp
+++ b/Engine/Plugins/Runtime/MassEntity/Source/MassEntity/Private/MassArchetypeData.cpp
@@ -399,7 +399,7 @@ void FMassArchetypeData::ExecuteFunction(FMassExecutionContext& RunContext, cons
 			{
 				PrevSharedFragmentValuesHash = SharedFragmentValuesHash;
 				BindConstSharedFragmentRequirements(RunContext, Chunk.GetSharedFragmentValues(), RequirementMapping.ConstSharedFragments);
-				BindSharedFragmentRequirements(RunContext, Chunk.GetSharedFragmentValues(), RequirementMapping.SharedFragments);
+				BindSharedFragmentRequirements(RunContext, Chunk.GetMutableSharedFragmentValues(), RequirementMapping.SharedFragments);
 			}
 
 			checkf((ChunkIterator->SubchunkStart + ChunkLength) <= Chunk.GetNumInstances() && ChunkLength > 0, TEXT("Invalid subchunk, it is going over the number of instances in the chunk or it is empty."));
@@ -433,7 +433,7 @@ void FMassArchetypeData::ExecuteFunction(FMassExecutionContext& RunContext, cons
 			{
 				PrevSharedFragmentValuesHash = SharedFragmentValuesHash;
 				BindConstSharedFragmentRequirements(RunContext, Chunk.GetSharedFragmentValues(), RequirementMapping.ConstSharedFragments);
-				BindSharedFragmentRequirements(RunContext, Chunk.GetSharedFragmentValues(), RequirementMapping.SharedFragments);
+				BindSharedFragmentRequirements(RunContext, Chunk.GetMutableSharedFragmentValues(), RequirementMapping.SharedFragments);
 			}
 
 			RunContext.SetCurrentChunkSerialModificationNumber(Chunk.GetSerialModificationNumber());
@@ -456,7 +456,7 @@ void FMassArchetypeData::ExecutionFunctionForChunk(FMassExecutionContext RunCont
 	if (ChunkLength)
 	{
 		BindConstSharedFragmentRequirements(RunContext, Chunk.GetSharedFragmentValues(), RequirementMapping.ChunkFragments);
-		BindSharedFragmentRequirements(RunContext, Chunk.GetSharedFragmentValues(), RequirementMapping.ChunkFragments);
+		BindSharedFragmentRequirements(RunContext, Chunk.GetMutableSharedFragmentValues(), RequirementMapping.ChunkFragments);
 
 		RunContext.SetCurrentArchetypesTagBitSet(GetTagBitSet());
 		RunContext.SetCurrentChunkSerialModificationNumber(Chunk.GetSerialModificationNumber());
@@ -721,7 +721,7 @@ void FMassArchetypeData::BindConstSharedFragmentRequirements(FMassExecutionConte
 			const int32 FragmentIndex = FragmentsMapping[i];
 
 			check(FragmentIndex != INDEX_NONE || Requirement.Requirement.IsOptional());
-			Requirement.FragmentView = FragmentIndex != INDEX_NONE ? SharedFragmentValues.GetConstSharedFragments()[FragmentIndex] : FConstSharedStruct();
+			Requirement.FragmentView = FragmentIndex != INDEX_NONE ? FConstStructView(SharedFragmentValues.GetConstSharedFragments()[FragmentIndex]) : FConstStructView();
 		}
 	}
 	else
@@ -730,12 +730,12 @@ void FMassArchetypeData::BindConstSharedFragmentRequirements(FMassExecutionConte
 		{
 			const FConstSharedStruct* SharedFragment = SharedFragmentValues.GetConstSharedFragments().FindByPredicate(FStructTypeEqualOperator(Requirement.Requirement.StructType) );
 			check(SharedFragment != nullptr || Requirement.Requirement.IsOptional());
-			Requirement.FragmentView = SharedFragment ? *SharedFragment : FConstSharedStruct();
+			Requirement.FragmentView = SharedFragment ? FConstStructView(*SharedFragment) : FConstStructView();
 		}
 	}
 }
 
-void FMassArchetypeData::BindSharedFragmentRequirements(FMassExecutionContext& RunContext, const FMassArchetypeSharedFragmentValues& SharedFragmentValues, const FMassFragmentIndicesMapping& FragmentsMapping)
+void FMassArchetypeData::BindSharedFragmentRequirements(FMassExecutionContext& RunContext, FMassArchetypeSharedFragmentValues& SharedFragmentValues, const FMassFragmentIndicesMapping& FragmentsMapping)
 {
 	if (FragmentsMapping.Num() > 0)
 	{
@@ -747,16 +747,16 @@ void FMassArchetypeData::BindSharedFragmentRequirements(FMassExecutionContext& R
 			const int32 FragmentIndex = FragmentsMapping[i];
 
 			check(FragmentIndex != INDEX_NONE || Requirement.Requirement.IsOptional());
-			Requirement.FragmentView = FragmentIndex != INDEX_NONE ? SharedFragmentValues.GetSharedFragments()[FragmentIndex] : FSharedStruct();
+			Requirement.FragmentView = FragmentIndex != INDEX_NONE ? FStructView(SharedFragmentValues.GetMutableSharedFragments()[FragmentIndex]) : FStructView();
 		}
 	}
 	else
 	{
 		for (FMassExecutionContext::FSharedFragmentView& Requirement : RunContext.GetMutableSharedRequirements())
 		{
-			const FSharedStruct* SharedFragment = SharedFragmentValues.GetSharedFragments().FindByPredicate(FStructTypeEqualOperator(Requirement.Requirement.StructType));
+			FSharedStruct* SharedFragment = SharedFragmentValues.GetMutableSharedFragments().FindByPredicate(FStructTypeEqualOperator(Requirement.Requirement.StructType));
 			check(SharedFragment != nullptr || Requirement.Requirement.IsOptional());
-			Requirement.FragmentView = SharedFragment ? *SharedFragment : FSharedStruct();
+			Requirement.FragmentView = SharedFragment ? FStructView(*SharedFragment) : FStructView();
 		}
 	}
 }
diff --git a/Engine/Plugins/Runtime/MassEntity/Source/MassEntity/Private/MassArchetypeData.h b/Engine/Plugins/Runtime/MassEntity/Source/MassEntity/Private/MassArchetypeData.h
--- a/Engine/Plugins/Runtime/MassEntity/Source/MassEntity/Private/MassArchetypeData.h
+++ b/Engine/Plugins/Runtime/MassEntity/Source/MassEntity/Private/MassArchetypeData.h
@@ -130,6 +130,7 @@ public:
 	int32 DebugGetChunkFragmentCount() const { return ChunkFragmentData.Num(); }
 #endif // WITH_MASSENTITY_DEBUG
 
+	FMassArchetypeSharedFragmentValues& GetMutableSharedFragmentValues() { return SharedFragmentValues; }
 	const FMassArchetypeSharedFragmentValues& GetSharedFragmentValues() const { return SharedFragmentValues; }
 };
 
@@ -364,7 +365,7 @@ protected:
 	void BindEntityRequirements(FMassExecutionContext& RunContext, const FMassFragmentIndicesMapping& EntityFragmentsMapping, FMassArchetypeChunk& Chunk, const int32 SubchunkStart, const int32 SubchunkLength);
 	void BindChunkFragmentRequirements(FMassExecutionContext& RunContext, const FMassFragmentIndicesMapping& ChunkFragmentsMapping, FMassArchetypeChunk& Chunk);
 	void BindConstSharedFragmentRequirements(FMassExecutionContext& RunContext, const FMassArchetypeSharedFragmentValues& SharedFragmentValues, const FMassFragmentIndicesMapping& ChunkFragmentsMapping);
-	void BindSharedFragmentRequirements(FMassExecutionContext& RunContext, const FMassArchetypeSharedFragmentValues& SharedFragmentValues, const FMassFragmentIndicesMapping& ChunkFragmentsMapping);
+	void BindSharedFragmentRequirements(FMassExecutionContext& RunContext, FMassArchetypeSharedFragmentValues& SharedFragmentValues, const FMassFragmentIndicesMapping& ChunkFragmentsMapping);
 
 private:
 	int32 AddEntityInternal(FMassEntityHandle Entity, const FMassArchetypeSharedFragmentValues& InSharedFragmentValues);
diff --git a/Engine/Plugins/Runtime/MassEntity/Source/MassEntity/Private/MassEntityView.cpp b/Engine/Plugins/Runtime/MassEntity/Source/MassEntity/Private/MassEntityView.cpp
--- a/Engine/Plugins/Runtime/MassEntity/Source/MassEntity/Private/MassEntityView.cpp
+++ b/Engine/Plugins/Runtime/MassEntity/Source/MassEntity/Private/MassEntityView.cpp
@@ -58,14 +58,16 @@ const void* FMassEntityView::GetConstSharedFragmentPtrChecked(const UScriptStruc
 void* FMassEntityView::GetSharedFragmentPtr(const UScriptStruct& FragmentType) const
 {
 	const FSharedStruct* SharedFragment = Archetype->GetSharedFragmentValues(Entity).GetSharedFragments().FindByPredicate(FStructTypeEqualOperator(&FragmentType));
-	return (SharedFragment != nullptr) ? SharedFragment->GetMutableMemory() : nullptr;
+	// @todo: fix constness, should the getter be non-const too?
+	return (SharedFragment != nullptr) ? const_cast<uint8*>(SharedFragment->GetMemory()) : nullptr;
 }
 
 void* FMassEntityView::GetSharedFragmentPtrChecked(const UScriptStruct& FragmentType) const
 {
 	const FSharedStruct* SharedFragment = Archetype->GetSharedFragmentValues(Entity).GetSharedFragments().FindByPredicate(FStructTypeEqualOperator(&FragmentType));
 	check(SharedFragment != nullptr);
-	return SharedFragment->GetMutableMemory();
+	// @todo: fix constness, should the getter be non-const too?
+	return const_cast<uint8*>(SharedFragment->GetMemory());
 }
 
 bool FMassEntityView::HasTag(const UScriptStruct& TagType) const
diff --git a/Engine/Plugins/Runtime/MassEntity/Source/MassEntity/Public/MassEntityManager.h b/Engine/Plugins/Runtime/MassEntity/Source/MassEntity/Public/MassEntityManager.h
--- a/Engine/Plugins/Runtime/MassEntity/Source/MassEntity/Public/MassEntityManager.h
+++ b/Engine/Plugins/Runtime/MassEntity/Source/MassEntity/Public/MassEntityManager.h
@@ -383,7 +383,7 @@ public:
 	void ForEachSharedFragment(TFunction< void(T& /*SharedFragment*/) > ExecuteFunction)
 	{
 		FStructTypeEqualOperator Predicate(T::StaticStruct());
-		for (const FSharedStruct& Struct : SharedFragments)
+		for (FSharedStruct& Struct : SharedFragments)
 		{
 			if (Predicate(Struct))
 			{
diff --git a/Engine/Plugins/Runtime/MassEntity/Source/MassEntity/Public/MassEntityTypes.h b/Engine/Plugins/Runtime/MassEntity/Source/MassEntity/Public/MassEntityTypes.h
--- a/Engine/Plugins/Runtime/MassEntity/Source/MassEntity/Public/MassEntityTypes.h
+++ b/Engine/Plugins/Runtime/MassEntity/Source/MassEntity/Public/MassEntityTypes.h
@@ -276,6 +276,11 @@ struct MASSENTITY_API FMassArchetypeSharedFragmentValues
 		return ConstSharedFragments;
 	}
 
+	FORCEINLINE TArray<FSharedStruct>& GetMutableSharedFragments()
+	{
+		return SharedFragments;
+	}
+	
 	FORCEINLINE const TArray<FSharedStruct>& GetSharedFragments() const
 	{
 		return SharedFragments;
diff --git a/Engine/Plugins/Runtime/MassEntity/Source/MassEntity/Public/MassExecutionContext.h b/Engine/Plugins/Runtime/MassEntity/Source/MassEntity/Public/MassExecutionContext.h
--- a/Engine/Plugins/Runtime/MassEntity/Source/MassEntity/Public/MassExecutionContext.h
+++ b/Engine/Plugins/Runtime/MassEntity/Source/MassEntity/Public/MassExecutionContext.h
@@ -131,17 +131,17 @@ public:
 	int32 GetChunkSerialModificationNumber() const { return ChunkSerialModificationNumber; }
 
 	template<typename T>
-	T* GetMutableChunkFragmentPtr() const
+	T* GetMutableChunkFragmentPtr()
 	{
 		static_assert(TIsDerivedFrom<T, FMassChunkFragment>::IsDerived, "Given struct doesn't represent a valid chunk fragment type. Make sure to inherit from FMassChunkFragment or one of its child-types.");
 
 		const UScriptStruct* Type = T::StaticStruct();
-		const FChunkFragmentView* FoundChunkFragmentData = ChunkFragmentViews.FindByPredicate([Type](const FChunkFragmentView& Element) { return Element.Requirement.StructType == Type; } );
+		FChunkFragmentView* FoundChunkFragmentData = ChunkFragmentViews.FindByPredicate([Type](const FChunkFragmentView& Element) { return Element.Requirement.StructType == Type; } );
 		return FoundChunkFragmentData ? FoundChunkFragmentData->FragmentView.GetMutablePtr<T>() : static_cast<T*>(nullptr);
 	}
 	
 	template<typename T>
-	T& GetMutableChunkFragment() const
+	T& GetMutableChunkFragment()
 	{
 		T* ChunkFragment = GetMutableChunkFragmentPtr<T>();
 		checkf(ChunkFragment, TEXT("Chunk Fragment requirement not found: %s"), *T::StaticStruct()->GetName());
@@ -151,13 +151,17 @@ public:
 	template<typename T>
 	const T* GetChunkFragmentPtr() const
 	{
-		return GetMutableChunkFragmentPtr<T>();
+		static_assert(TIsDerivedFrom<T, FMassChunkFragment>::IsDerived, "Given struct doesn't represent a valid chunk fragment type. Make sure to inherit from FMassChunkFragment or one of its child-types.");
+
+		const UScriptStruct* Type = T::StaticStruct();
+		const FChunkFragmentView* FoundChunkFragmentData = ChunkFragmentViews.FindByPredicate([Type](const FChunkFragmentView& Element) { return Element.Requirement.StructType == Type; } );
+		return FoundChunkFragmentData ? FoundChunkFragmentData->FragmentView.GetPtr<T>() : static_cast<const T*>(nullptr);
 	}
 	
 	template<typename T>
 	const T& GetChunkFragment() const
 	{
-		const T* ChunkFragment = GetMutableChunkFragmentPtr<T>();
+		const T* ChunkFragment = GetChunkFragmentPtr<T>();
 		checkf(ChunkFragment, TEXT("Chunk Fragment requirement not found: %s"), *T::StaticStruct()->GetName());
 		return *ChunkFragment;
 	}
@@ -182,26 +186,29 @@ public:
 
 
 	template<typename T>
-	T* GetMutableSharedFragmentPtr() const
+	T* GetMutableSharedFragmentPtr()
 	{
 		static_assert(TIsDerivedFrom<T, FMassSharedFragment>::IsDerived, "Given struct doesn't represent a valid shared fragment type. Make sure to inherit from FMassSharedFragment or one of its child-types.");
 
-		const FSharedFragmentView* FoundSharedFragmentData = SharedFragmentViews.FindByPredicate([](const FSharedFragmentView& Element) { return Element.Requirement.StructType == T::StaticStruct(); });
+		FSharedFragmentView* FoundSharedFragmentData = SharedFragmentViews.FindByPredicate([](const FSharedFragmentView& Element) { return Element.Requirement.StructType == T::StaticStruct(); });
 		return FoundSharedFragmentData ? FoundSharedFragmentData->FragmentView.GetMutablePtr<T>() : static_cast<T*>(nullptr);
 	}
 
 	template<typename T>
-	T& GetMutableSharedFragment() const
+	const T* GetSharedFragmentPtr() const
 	{
-		T* SharedFragment = GetMutableSharedFragmentPtr<T>();
-		checkf(SharedFragment, TEXT("Shared Fragment requirement not found: %s"), *T::StaticStruct()->GetName());
-		return *SharedFragment;
+		static_assert(TIsDerivedFrom<T, FMassSharedFragment>::IsDerived, "Given struct doesn't represent a valid shared fragment type. Make sure to inherit from FMassSharedFragment or one of its child-types.");
+
+		const FSharedFragmentView* FoundSharedFragmentData = SharedFragmentViews.FindByPredicate([](const FSharedFragmentView& Element) { return Element.Requirement.StructType == T::StaticStruct(); });
+		return FoundSharedFragmentData ? FoundSharedFragmentData->FragmentView.GetPtr<T>() : static_cast<const T*>(nullptr);
 	}
 
 	template<typename T>
-	const T* GetSharedFragmentPtr() const
+	T& GetMutableSharedFragment()
 	{
-		return GetMutableSharedFragmentPtr<T>();
+		T* SharedFragment = GetMutableSharedFragmentPtr<T>();
+		checkf(SharedFragment, TEXT("Shared Fragment requirement not found: %s"), *T::StaticStruct()->GetName());
+		return *SharedFragment;
 	}
 
 	template<typename T>
diff --git a/Engine/Plugins/Runtime/MassGameplay/Source/MassActors/Private/MassActorSpawnerSubsystem.cpp b/Engine/Plugins/Runtime/MassGameplay/Source/MassActors/Private/MassActorSpawnerSubsystem.cpp
--- a/Engine/Plugins/Runtime/MassGameplay/Source/MassActors/Private/MassActorSpawnerSubsystem.cpp
+++ b/Engine/Plugins/Runtime/MassGameplay/Source/MassActors/Private/MassActorSpawnerSubsystem.cpp
@@ -120,7 +120,7 @@ FMassActorSpawnRequestHandle UMassActorSpawnerSubsystem::RequestActorSpawnIntern
 	if (!SpawnRequests.IsValidIndex(Index))
 	{
 		checkf(SpawnRequests.Num() == Index, TEXT("This case should only be when we need to grow the array of one element."));
-		SpawnRequests.Add(SpawnRequestView);
+		SpawnRequests.Emplace(SpawnRequestView);
 	}
 	else
 	{
diff --git a/Engine/Plugins/Runtime/MassGameplay/Source/MassSpawner/Private/MassEntityTemplateRegistry.cpp b/Engine/Plugins/Runtime/MassGameplay/Source/MassSpawner/Private/MassEntityTemplateRegistry.cpp
--- a/Engine/Plugins/Runtime/MassGameplay/Source/MassSpawner/Private/MassEntityTemplateRegistry.cpp
+++ b/Engine/Plugins/Runtime/MassGameplay/Source/MassSpawner/Private/MassEntityTemplateRegistry.cpp
@@ -17,7 +17,7 @@
 
 namespace FTemplateRegistryHelpers
 {
-	uint32 CalcHash(const FInstancedStruct& StructInstance)
+	uint32 CalcHash(const FConstStructView StructInstance)
 	{
 		const UScriptStruct* Type = StructInstance.GetScriptStruct();
 		const uint8* Memory = StructInstance.GetMemory();
@@ -75,7 +75,7 @@ UMassEntityTemplateRegistry::FStructToTemplateBuilderDelegate& UMassEntityTempla
 	return StructBasedBuilders.FindOrAdd(&DataType);
 }
 
-const FMassEntityTemplate* UMassEntityTemplateRegistry::FindOrBuildStructTemplate(const FInstancedStruct& StructInstance)
+const FMassEntityTemplate* UMassEntityTemplateRegistry::FindOrBuildStructTemplate(const FConstStructView StructInstance)
 {
 	// thou shall not call this function on CDO
 	check(HasAnyFlags(RF_ClassDefaultObject) == false);
@@ -128,7 +128,7 @@ const FMassEntityTemplate* UMassEntityTemplateRegistry::FindOrBuildStructTemplat
 	return NewTemplate;
 }
 
-bool UMassEntityTemplateRegistry::BuildTemplateImpl(const FStructToTemplateBuilderDelegate& Builder, const FInstancedStruct& StructInstance, FMassEntityTemplate& OutTemplate)
+bool UMassEntityTemplateRegistry::BuildTemplateImpl(const FStructToTemplateBuilderDelegate& Builder, const FConstStructView StructInstance, FMassEntityTemplate& OutTemplate)
 {
 	UWorld* World = GetWorld();
 	FMassEntityTemplateBuildContext Context(OutTemplate);
diff --git a/Engine/Plugins/Runtime/MassGameplay/Source/MassSpawner/Public/MassEntityTemplate.h b/Engine/Plugins/Runtime/MassGameplay/Source/MassSpawner/Public/MassEntityTemplate.h
--- a/Engine/Plugins/Runtime/MassGameplay/Source/MassSpawner/Public/MassEntityTemplate.h
+++ b/Engine/Plugins/Runtime/MassGameplay/Source/MassSpawner/Public/MassEntityTemplate.h
@@ -112,11 +112,11 @@ struct MASSSPAWNER_API FMassEntityTemplate
 		if (!Composition.Fragments.Contains(*FragmentType))
 		{
 			Composition.Fragments.Add(*FragmentType);
-			InitialFragmentValues.Add(Fragment);
+			InitialFragmentValues.Emplace(Fragment);
 		}
 		else if (!InitialFragmentValues.ContainsByPredicate(FStructTypeEqualOperator(FragmentType)))
 		{
-			InitialFragmentValues.Add(Fragment);
+			InitialFragmentValues.Emplace(Fragment);
 		}
 	}
 
@@ -128,7 +128,7 @@ struct MASSSPAWNER_API FMassEntityTemplate
 		{
 			Composition.Fragments.Add<T>();
 		}
-		else if (const FInstancedStruct* Fragment = InitialFragmentValues.FindByPredicate(FStructTypeEqualOperator(T::StaticStruct())))
+		else if (FInstancedStruct* Fragment = InitialFragmentValues.FindByPredicate(FStructTypeEqualOperator(T::StaticStruct())))
 		{
 			return Fragment->template GetMutable<T>();
 		}
diff --git a/Engine/Plugins/Runtime/MassGameplay/Source/MassSpawner/Public/MassEntityTemplateRegistry.h b/Engine/Plugins/Runtime/MassGameplay/Source/MassSpawner/Public/MassEntityTemplateRegistry.h
--- a/Engine/Plugins/Runtime/MassGameplay/Source/MassSpawner/Public/MassEntityTemplateRegistry.h
+++ b/Engine/Plugins/Runtime/MassGameplay/Source/MassSpawner/Public/MassEntityTemplateRegistry.h
@@ -204,7 +204,7 @@ class MASSSPAWNER_API UMassEntityTemplateRegistry : public UObject
 	GENERATED_BODY()
 public:
 	// @todo consider TFunction instead
-	DECLARE_DELEGATE_ThreeParams(FStructToTemplateBuilderDelegate, const UWorld* /*World*/, const FInstancedStruct& /*InStructInstance*/, FMassEntityTemplateBuildContext& /*BuildContext*/);
+	DECLARE_DELEGATE_ThreeParams(FStructToTemplateBuilderDelegate, const UWorld* /*World*/, const FConstStructView /*InStructInstance*/, FMassEntityTemplateBuildContext& /*BuildContext*/);
 
 	UMassEntityTemplateRegistry() = default;
 	virtual void BeginDestroy() override;
@@ -212,7 +212,7 @@ public:
 
 	static FStructToTemplateBuilderDelegate& FindOrAdd(const UScriptStruct& DataType);
 
-	const FMassEntityTemplate* FindOrBuildStructTemplate(const FInstancedStruct& StructInstance);
+	const FMassEntityTemplate* FindOrBuildStructTemplate(const FConstStructView StructInstance);
 
 	/** Removes all the cached template instances */
 	void DebugReset();
@@ -226,7 +226,7 @@ public:
 
 protected:
 	/** @return true if a template has been built, false otherwise */
-	bool BuildTemplateImpl(const FStructToTemplateBuilderDelegate& Builder, const FInstancedStruct& StructInstance, FMassEntityTemplate& OutTemplate);
+	bool BuildTemplateImpl(const FStructToTemplateBuilderDelegate& Builder, const FConstStructView StructInstance, FMassEntityTemplate& OutTemplate);
 
 protected:
 	static TMap<const UScriptStruct*, FStructToTemplateBuilderDelegate> StructBasedBuilders;
