MassDebugger changes supporting World-less MassEntityManagers

diff --git a/Engine/Plugins/Runtime/MassEntity/Source/MassEntityDebugger/Private/MassDebuggerModel.cpp b/Engine/Plugins/Runtime/MassEntity/Source/MassEntityDebugger/Private/MassDebuggerModel.cpp
--- a/Engine/Plugins/Runtime/MassEntity/Source/MassEntityDebugger/Private/MassDebuggerModel.cpp
+++ b/Engine/Plugins/Runtime/MassEntity/Source/MassEntityDebugger/Private/MassDebuggerModel.cpp
@@ -48,6 +48,18 @@ namespace UE::Mass::Debugger::Private
 	{
 		return PointerHash(&Processor);
 	}
+
+	/** We're ignoring all the CDO processors (since as such are not being run at runtime) as well ass processors owned
+	 *  by a CDO, for the very same reason. */
+	bool IsDebuggableProcessor(const UWorld* ContextWorld, const UMassProcessor& Processor)
+	{
+		return IsValid(&Processor)
+			&& Processor.HasAnyFlags(RF_ClassDefaultObject) == false
+			&& Processor.GetWorld() == ContextWorld 
+			// checking ContextWorld is a cheaper way of supporting the declared behavior, since if there is a world then
+			// the processors are definitelly not CDO owned (by design). Is there is no world we need to check specifically.
+			&& (ContextWorld != nullptr || Processor.GetOuter()->HasAnyFlags(RF_ClassDefaultObject) == false);
+	}
 } // namespace UE::Mass::Debugger::Private
 
 //----------------------------------------------------------------------//
@@ -69,7 +81,7 @@ FString FMassDebuggerEnvironment::GetDisplayName() const
 #endif // WITH_MASSENTITY_DEBUG
 	
 	const UWorld* WorldPtr = World.Get();
-	DisplayName += WorldPtr ? WorldPtr->GetPathName() : TEXT("Invalid");
+	DisplayName += WorldPtr ? WorldPtr->GetPathName() : TEXT("No World");
 	return DisplayName;
 }
 
@@ -441,13 +453,12 @@ void FMassDebuggerModel::CacheProcessorsData(const TMap<FMassArchetypeHandle, TS
 
 	if (EntityManager)
 	{
-		check(World);
 		for (FThreadSafeObjectIterator It(UMassProcessor::StaticClass()); It; ++It)
 		{
 			UMassProcessor* Processor = Cast<UMassProcessor>(*It);
-			if (Processor && (Processor->HasAnyFlags(RF_ClassDefaultObject) == false)
+			if (Processor 
 				&& Cast<UMassCompositeProcessor>(Processor) == nullptr
-				&& Processor->GetWorld() == World)
+				&& UE::Mass::Debugger::Private::IsDebuggableProcessor(World, *Processor))
 			{
 				CachedProcessors.Add(MakeShareable(new FMassDebuggerProcessorData(*EntityManager, *Processor, InTransientArchetypesMap)));
 			}
@@ -458,9 +469,9 @@ void FMassDebuggerModel::CacheProcessorsData(const TMap<FMassArchetypeHandle, TS
 		for (FThreadSafeObjectIterator It(UMassProcessor::StaticClass()); It; ++It)
 		{
 			UMassProcessor* Processor = Cast<UMassProcessor>(*It);
-			if (Processor && (Processor->HasAnyFlags(RF_ClassDefaultObject) == false)
-				&& Cast<UMassCompositeProcessor>(Processor) == nullptr 
-				&& (World == nullptr || Processor->GetWorld() == World))
+			if (Processor 
+				&& Cast<UMassCompositeProcessor>(Processor) == nullptr
+				&& UE::Mass::Debugger::Private::IsDebuggableProcessor(World, *Processor))
 			{
 				CachedProcessors.Add(MakeShareable(new FMassDebuggerProcessorData(*Processor)));
 			}
@@ -481,8 +492,7 @@ void FMassDebuggerModel::CacheProcessingGraphs()
 	for (FThreadSafeObjectIterator It(UMassProcessor::StaticClass()); It; ++It)
 	{
 		UMassCompositeProcessor* Processor = Cast<UMassCompositeProcessor>(*It);
-		if (Processor && (Processor->HasAnyFlags(RF_ClassDefaultObject) == false)
-			&& (World == nullptr || Processor->GetWorld() == World))
+		if (Processor && UE::Mass::Debugger::Private::IsDebuggableProcessor(World, *Processor))
 		{
 			CachedProcessingGraphs.Add(MakeShareable(new FMassDebuggerProcessingGraph(*this, *Processor)));
 		}
@@ -598,7 +608,7 @@ void FMassDebuggerModel::MarkAsStale()
 
 bool FMassDebuggerModel::IsStale() const 
 { 
-	return Environment.IsValid() && (Environment->IsWorldValid() == false); 
+	return Environment.IsValid() == false || (Environment->NeedsValidWorld() == true && Environment->IsWorldValid() == false);
 }
 
 const TSharedPtr<FMassDebuggerProcessorData>& FMassDebuggerModel::GetProcessorDataChecked(const UMassProcessor& Processor) const
diff --git a/Engine/Plugins/Runtime/MassEntity/Source/MassEntityDebugger/Private/MassDebuggerModel.h b/Engine/Plugins/Runtime/MassEntity/Source/MassEntityDebugger/Private/MassDebuggerModel.h
--- a/Engine/Plugins/Runtime/MassEntity/Source/MassEntityDebugger/Private/MassDebuggerModel.h
+++ b/Engine/Plugins/Runtime/MassEntity/Source/MassEntityDebugger/Private/MassDebuggerModel.h
@@ -136,7 +136,7 @@ struct FMassDebuggerProcessingGraph
 struct FMassDebuggerEnvironment
 {
 	explicit FMassDebuggerEnvironment(const TSharedRef<const FMassEntityManager>& InEntityManager)
-		: EntityManager(InEntityManager), World(InEntityManager->GetWorld())
+		: EntityManager(InEntityManager), World(InEntityManager->GetWorld()), bNeedsValidWorld(InEntityManager->GetWorld() != nullptr)
 	{}
 
 	bool operator==(const FMassDebuggerEnvironment& Other) const { return EntityManager == Other.EntityManager; }
@@ -144,9 +144,11 @@ struct FMassDebuggerEnvironment
 	FString GetDisplayName() const;
 	const FMassEntityManager* GetEntityManager() const;
 	bool IsWorldValid() const { return World.IsValid(); }
+	bool NeedsValidWorld() const { return bNeedsValidWorld; }
 	
 	TWeakPtr<const FMassEntityManager> EntityManager;
 	TWeakObjectPtr<UWorld> World;
+	bool bNeedsValidWorld = false;
 };
 
 
diff --git a/Engine/Plugins/Runtime/MassEntity/Source/MassEntityDebugger/Private/SMassDebugger.cpp b/Engine/Plugins/Runtime/MassEntity/Source/MassEntityDebugger/Private/SMassDebugger.cpp
--- a/Engine/Plugins/Runtime/MassEntity/Source/MassEntityDebugger/Private/SMassDebugger.cpp
+++ b/Engine/Plugins/Runtime/MassEntity/Source/MassEntityDebugger/Private/SMassDebugger.cpp
@@ -267,7 +267,7 @@ void SMassDebugger::RebuildEnvironmentsList()
 		if (const FMassEntityManager* EntityManagerPtr = WeakEntityManager.Pin().Get())
 		{
 			const UWorld* World = EntityManagerPtr->GetWorld();
-			if (World != nullptr && UE::Mass::Debugger::Private::IsSupportedWorldType(World->WorldType))
+			if (World == nullptr || UE::Mass::Debugger::Private::IsSupportedWorldType(World->WorldType))
 			{
 				EnvironmentsList.Add(MakeShareable(new FMassDebuggerEnvironment(EntityManagerPtr->AsShared())));
 			}
