Added FMassProcessingPhaseManager::TriggerPhase that lets users manually trigger a given phase. This is needed to be able to tick processing phases outside of world ticking.

diff --git a/Engine/Plugins/Runtime/MassEntity/Source/MassEntity/Private/MassProcessingPhaseManager.cpp b/Engine/Plugins/Runtime/MassEntity/Source/MassEntity/Private/MassProcessingPhaseManager.cpp
--- a/Engine/Plugins/Runtime/MassEntity/Source/MassEntity/Private/MassProcessingPhaseManager.cpp
+++ b/Engine/Plugins/Runtime/MassEntity/Source/MassEntity/Private/MassProcessingPhaseManager.cpp
@@ -151,6 +151,22 @@ void FMassPhaseProcessorConfigurationHelper::Configure(const TSharedPtr<FMassEnt
 
 	PhaseProcessor.Populate(SortedProcessors);
 
+#if WITH_MASSENTITY_DEBUG
+	for (FMassProcessorOrderInfo& ProcessorOrderInfo : SortedProcessors)
+	{
+		TmpPipeline.RemoveProcessor(*ProcessorOrderInfo.Processor);
+	}
+	
+	if (TmpPipeline.Num())
+	{
+		UE_VLOG_UELOG(&PhaseProcessor, LogMass, Verbose, TEXT("Discarding processors due to not having anything to do (no relevant Archetypes):"));
+		for (const UMassProcessor* Processor : TmpPipeline.GetProcessors())
+		{
+			UE_VLOG_UELOG(&PhaseProcessor, LogMass, Verbose, TEXT("\t%s"), *Processor->GetProcessorName());
+		}
+	}
+#endif // WITH_MASSENTITY_DEBUG
+
 	if (Solver.IsSolvingForSingleThread() == false)
 	{
 		PhaseProcessor.BuildFlatProcessingGraph(SortedProcessors);
@@ -208,6 +224,8 @@ void FMassProcessingPhaseManager::Initialize(UObject& InOwner, TConstArrayView<F
 		}
 #endif // WITH_EDITOR
 	}
+
+	bIsAllowedToTick = true;
 }
 
 void FMassProcessingPhaseManager::Deinitialize()
@@ -218,6 +236,18 @@ void FMassProcessingPhaseManager::Deinitialize()
 	}
 }
 
+const FGraphEventRef& FMassProcessingPhaseManager::TriggerPhase(const EMassProcessingPhase Phase, const float DeltaTime, const FGraphEventRef& MyCompletionGraphEvent)
+{
+	check(Phase != EMassProcessingPhase::MAX);
+
+	if (bIsAllowedToTick)
+	{
+		ProcessingPhases[(int)Phase].ExecuteTick(DeltaTime, LEVELTICK_All, ENamedThreads::GameThread, MyCompletionGraphEvent);
+	}
+
+	return MyCompletionGraphEvent;
+}
+
 void FMassProcessingPhaseManager::Start(UWorld& World)
 {
 	UMassEntitySubsystem* EntitySubsystem = UWorld::GetSubsystem<UMassEntitySubsystem>(&World);
@@ -234,13 +264,16 @@ void FMassProcessingPhaseManager::Start(UWorld& World)
 
 void FMassProcessingPhaseManager::Start(const TSharedPtr<FMassEntityManager>& InEntityManager)
 {
-	UWorld* World = InEntityManager->GetWorld();
-	check(World);
 	EntityManager = InEntityManager;
 
 	OnNewArchetypeHandle = EntityManager->GetOnNewArchetypeEvent().AddRaw(this, &FMassProcessingPhaseManager::OnNewArchetype);
 
-	EnableTickFunctions(*World);
+	if (UWorld* World = InEntityManager->GetWorld())
+	{
+		EnableTickFunctions(*World);
+	}
+
+	bIsAllowedToTick = true;
 }
 
 void FMassProcessingPhaseManager::AddReferencedObjects(FReferenceCollector& Collector)
@@ -286,6 +319,8 @@ void FMassProcessingPhaseManager::EnableTickFunctions(const UWorld& World)
 
 void FMassProcessingPhaseManager::Stop()
 {
+	bIsAllowedToTick = false;
+
 	if (EntityManager)
 	{
 		EntityManager->GetOnNewArchetypeEvent().Remove(OnNewArchetypeHandle);
diff --git a/Engine/Plugins/Runtime/MassEntity/Source/MassEntity/Public/MassProcessingPhaseManager.h b/Engine/Plugins/Runtime/MassEntity/Source/MassEntity/Public/MassProcessingPhaseManager.h
--- a/Engine/Plugins/Runtime/MassEntity/Source/MassEntity/Public/MassProcessingPhaseManager.h
+++ b/Engine/Plugins/Runtime/MassEntity/Source/MassEntity/Public/MassProcessingPhaseManager.h
@@ -141,13 +141,15 @@ public:
 	/** Needs to be called before destruction, ideally before owner's BeginDestroy (a FGCObject's limitation) */
 	void Deinitialize();
 
+	const FGraphEventRef& TriggerPhase(const EMassProcessingPhase Phase, const float DeltaTime, const FGraphEventRef& MyCompletionGraphEvent);
+
 	/** 
 	 *  Stores EntityManager associated with given world's MassEntitySubsystem and kicks off phase ticking.
 	 */
 	void Start(UWorld& World);
 	
 	/**
-	 *  Stores InEntityManager as the entity manager and kicks off phase ticking.
+	 *  Stores InEntityManager as the entity manager. It also kicks off phase ticking if the given InEntityManager is tied to a UWorld.
 	 */
 	void Start(const TSharedPtr<FMassEntityManager>& InEntityManager);
 	void Stop();
@@ -204,4 +206,6 @@ protected:
 	TWeakObjectPtr<UObject> Owner;
 
 	FDelegateHandle OnNewArchetypeHandle;
+
+	bool bIsAllowedToTick = false;
 };
