Tiny change that missed the last change to the file

diff --git a/Engine/Plugins/Runtime/MassEntity/Source/MassEntity/Private/MassProcessingPhaseManager.cpp b/Engine/Plugins/Runtime/MassEntity/Source/MassEntity/Private/MassProcessingPhaseManager.cpp
--- a/Engine/Plugins/Runtime/MassEntity/Source/MassEntity/Private/MassProcessingPhaseManager.cpp
+++ b/Engine/Plugins/Runtime/MassEntity/Source/MassEntity/Private/MassProcessingPhaseManager.cpp
@@ -152,7 +152,7 @@ void FMassPhaseProcessorConfigurationHelper::Configure(const TSharedPtr<FMassEnt
 	PhaseProcessor.Populate(SortedProcessors);
 
 #if WITH_MASSENTITY_DEBUG
-	for (FMassProcessorOrderInfo& ProcessorOrderInfo : SortedProcessors)
+	for (const FMassProcessorOrderInfo& ProcessorOrderInfo : SortedProcessors)
 	{
 		TmpPipeline.RemoveProcessor(*ProcessorOrderInfo.Processor);
 	}
