Added two editor-time Mass subsystems:
* MassEntityEditorSubsystem - instantiating MassEntityManager and supplying world-independent ticking
* MassActorEditorSubsystem - wrapping FMassActorManager for world-less purposes

diff --git a/Engine/Plugins/Runtime/MassEntity/Source/MassEntityEditor/MassEntityEditor.Build.cs b/Engine/Plugins/Runtime/MassEntity/Source/MassEntityEditor/MassEntityEditor.Build.cs
--- a/Engine/Plugins/Runtime/MassEntity/Source/MassEntityEditor/MassEntityEditor.Build.cs
+++ b/Engine/Plugins/Runtime/MassEntity/Source/MassEntityEditor/MassEntityEditor.Build.cs
@@ -7,7 +7,6 @@ namespace UnrealBuildTool.Rules
 		public MassEntityEditor(ReadOnlyTargetRules Target) : base(Target)
 		{
 			PCHUsage = ModuleRules.PCHUsageMode.UseExplicitOrSharedPCHs;
-			//BinariesSubFolder = "NotForLicensees";
 
 			PublicIncludePaths.AddRange(
 			new string[] {
@@ -29,6 +28,7 @@ namespace UnrealBuildTool.Rules
 				"DetailCustomizations",
 				"ComponentVisualizers",
 				"Projects",
+				"EditorSubsystem"
 			}
 			);
 
diff --git a/Engine/Plugins/Runtime/MassEntity/Source/MassEntityEditor/Private/MassEntityEditorSubsystem.cpp b/Engine/Plugins/Runtime/MassEntity/Source/MassEntityEditor/Private/MassEntityEditorSubsystem.cpp
new file mode 100644
--- /dev/null
+++ b/Engine/Plugins/Runtime/MassEntity/Source/MassEntityEditor/Private/MassEntityEditorSubsystem.cpp
@@ -0,0 +1,137 @@
+// Copyright Epic Games, Inc. All Rights Reserved.
+
+#include "MassEntityEditorSubsystem.h"
+#include "MassEntityManager.h"
+#include "MassEntitySettings.h"
+#include "TickableEditorObject.h"
+
+namespace UE::Mass
+{
+struct FMassEditorPhaseTickTask
+{
+	FMassEditorPhaseTickTask(const TSharedRef<FMassProcessingPhaseManager>& InPhaseManager, const EMassProcessingPhase InPhase, const float InDeltaTime)
+		: PhaseManager(InPhaseManager)
+		, Phase(InPhase)
+		, DeltaTime(InDeltaTime)
+	{
+	}
+
+	static TStatId GetStatId()
+	{
+		RETURN_QUICK_DECLARE_CYCLE_STAT(FMassEditorPhaseTickTask, STATGROUP_TaskGraphTasks);
+	}
+
+	static ENamedThreads::Type GetDesiredThread() { return ENamedThreads::GameThread; }
+	static ESubsequentsMode::Type GetSubsequentsMode() { return ESubsequentsMode::TrackSubsequents; }
+	
+	void DoTask(ENamedThreads::Type CurrentThread, const FGraphEventRef& MyCompletionGraphEvent)
+	{
+		TRACE_CPUPROFILER_EVENT_SCOPE(FMassEditorPhaseTickTask);
+
+		PhaseManager->TriggerPhase(Phase, DeltaTime, MyCompletionGraphEvent);
+	}
+private:
+	const TSharedRef<FMassProcessingPhaseManager> PhaseManager;
+	const EMassProcessingPhase Phase = EMassProcessingPhase::MAX;
+	const float DeltaTime = 0.f;
+};
+
+struct FMassEditorTickFunction : FTickableEditorObject
+{
+	FMassEditorTickFunction(UMassEntityEditorSubsystem& InOwner)
+		: Owner(InOwner)
+	{}
+
+	//** FTickableEditorObject Interface
+	virtual void Tick(float DeltaTime) override
+	{
+		Owner.Tick(DeltaTime);
+	}
+
+	virtual ETickableTickType GetTickableTickType() const override 
+	{ 
+		return ETickableTickType::Always; 
+	}
+
+	virtual TStatId GetStatId() const override
+	{
+		RETURN_QUICK_DECLARE_CYCLE_STAT(FMassEditorTickFunction, STATGROUP_Tickables);
+	}
+	//~ END FTickableEditorObject Interface
+
+	UMassEntityEditorSubsystem& Owner;
+};
+
+} // namespace UE::Mass
+
+//////////////////////////////////////////////////////////////////////
+// UMassEntityEditorSubsystem
+
+UMassEntityEditorSubsystem::UMassEntityEditorSubsystem()
+	: EntityManager(MakeShareable(new FMassEntityManager(this)))
+	, PhaseManager(new FMassProcessingPhaseManager())
+{
+
+}
+
+UMassEntityEditorSubsystem::~UMassEntityEditorSubsystem()
+{
+	if (EditorTickFunction)
+	{
+		UE_LOG(LogMass, Warning, TEXT("UMassEntityEditorSubsystem::EditorTickFunction still valid in the destructor indicating that UMassEntityEditorSubsystem::Deinitialize had not been called"));
+		StopAndCleanUp();
+	}
+}
+
+void UMassEntityEditorSubsystem::Initialize(FSubsystemCollectionBase& Collection)
+{
+	// Note that since this is a raw pointer we need to remember to delete it once we're done. It's now done in both 
+	// Deinitialize and the destructor (in case Deinitialize call is missed in an unexpected circumstances).
+	// Originally we wanted to use a TUniquePtr, but it would require to define the UE::Mass::FMassEditorTickFunction type 
+	// in the header file.
+	EditorTickFunction = new UE::Mass::FMassEditorTickFunction(*this);
+
+	EntityManager->Initialize();
+
+	// set up ProcessingPhasesConfig
+	TConstArrayView<FMassProcessingPhaseConfig> MainPhasesConfig = GET_MASS_CONFIG_VALUE(GetProcessingPhasesConfig());
+	for (int PhaseIndex = 0; PhaseIndex < (int)EMassProcessingPhase::MAX; ++PhaseIndex)
+	{
+		ProcessingPhasesConfig[PhaseIndex] = MainPhasesConfig[PhaseIndex];
+	}
+
+	PhaseManager->Initialize(*this, ProcessingPhasesConfig);
+
+	PhaseManager->Start(EntityManager);
+}
+
+void UMassEntityEditorSubsystem::Deinitialize()
+{
+	StopAndCleanUp();
+}
+
+void UMassEntityEditorSubsystem::StopAndCleanUp()
+{
+	if (EditorTickFunction)
+	{
+		delete EditorTickFunction;
+		EditorTickFunction = nullptr;
+	}
+
+	PhaseManager->Stop();
+}
+
+void UMassEntityEditorSubsystem::Tick(float DeltaTime)
+{
+	if (CompletionEvent.IsValid())
+	{
+		CompletionEvent->Wait();
+	}
+
+	for (int PhaseIndex = 0; PhaseIndex < (int)EMassProcessingPhase::MAX; ++PhaseIndex)
+	{
+		const FGraphEventArray Prerequisites = { CompletionEvent };
+		CompletionEvent = TGraphTask<UE::Mass::FMassEditorPhaseTickTask>::CreateTask(&Prerequisites)
+			.ConstructAndDispatchWhenReady(PhaseManager, EMassProcessingPhase(PhaseIndex), DeltaTime);
+	}
+}
diff --git a/Engine/Plugins/Runtime/MassEntity/Source/MassEntityEditor/Public/MassEntityEditorSubsystem.h b/Engine/Plugins/Runtime/MassEntity/Source/MassEntityEditor/Public/MassEntityEditorSubsystem.h
new file mode 100644
--- /dev/null
+++ b/Engine/Plugins/Runtime/MassEntity/Source/MassEntityEditor/Public/MassEntityEditorSubsystem.h
@@ -0,0 +1,45 @@
+// Copyright Epic Games, Inc. All Rights Reserved.
+
+#pragma once
+
+#include "EditorSubsystem.h"
+#include "MassProcessingPhaseManager.h"
+#include "MassEntityEditorSubsystem.generated.h"
+
+
+struct FMassEntityManager;
+namespace UE::Mass
+{
+	struct FMassEditorTickFunction;
+}
+
+UCLASS()
+class MASSENTITYEDITOR_API UMassEntityEditorSubsystem : public UEditorSubsystem
+{
+	GENERATED_BODY()
+
+public:
+	UMassEntityEditorSubsystem();
+	~UMassEntityEditorSubsystem();
+
+	FMassEntityManager& GetMutableEntityManager() { return EntityManager.Get(); }
+
+protected:
+	// USubsystem interface begin
+	virtual void Initialize(FSubsystemCollectionBase& Collection) override;
+	virtual void Deinitialize() override;
+	// USubsystem interface end
+	void StopAndCleanUp();
+
+	friend UE::Mass::FMassEditorTickFunction;
+	void Tick(float DeltaTime);
+
+	UE::Mass::FMassEditorTickFunction* EditorTickFunction = nullptr;
+
+	TSharedRef<FMassEntityManager> EntityManager;
+
+	TSharedRef<FMassProcessingPhaseManager> PhaseManager;
+	FGraphEventRef CompletionEvent;
+
+	FMassProcessingPhaseConfig ProcessingPhasesConfig[(uint8)EMassProcessingPhase::MAX];
+};
diff --git a/Engine/Plugins/Runtime/MassGameplay/Source/MassGameplayEditor/MassGameplayEditor.Build.cs b/Engine/Plugins/Runtime/MassGameplay/Source/MassGameplayEditor/MassGameplayEditor.Build.cs
--- a/Engine/Plugins/Runtime/MassGameplay/Source/MassGameplayEditor/MassGameplayEditor.Build.cs
+++ b/Engine/Plugins/Runtime/MassGameplay/Source/MassGameplayEditor/MassGameplayEditor.Build.cs
@@ -27,9 +27,12 @@ namespace UnrealBuildTool.Rules
 				"PropertyEditor",
 				"MassSpawner",
 				"MassEntity",
+				"MassEntityEditor",
+				"MassActors",
 				"DetailCustomizations",
 				"ComponentVisualizers",
 				"Projects",
+				"EditorSubsystem"
 			}
 			);
 
diff --git a/Engine/Plugins/Runtime/MassGameplay/Source/MassGameplayEditor/Private/MassActorEditorSubsystem.cpp b/Engine/Plugins/Runtime/MassGameplay/Source/MassGameplayEditor/Private/MassActorEditorSubsystem.cpp
new file mode 100644
--- /dev/null
+++ b/Engine/Plugins/Runtime/MassGameplay/Source/MassGameplayEditor/Private/MassActorEditorSubsystem.cpp
@@ -0,0 +1,18 @@
+// Copyright Epic Games, Inc. All Rights Reserved.
+
+#include "MassActorEditorSubsystem.h"
+#include "MassActorSubsystem.h"
+#include "MassEntityManager.h"
+#include "MassEntityEditorSubsystem.h"
+
+
+//----------------------------------------------------------------------//
+//  UMassActorEditorSubsystem
+//----------------------------------------------------------------------//
+void UMassActorEditorSubsystem::Initialize(FSubsystemCollectionBase& Collection)
+{
+	UMassEntityEditorSubsystem* MassEditorEditorSubsystem = Collection.InitializeDependency<UMassEntityEditorSubsystem>();
+	check(MassEditorEditorSubsystem);
+	FMassEntityManager& MassEntityManager = MassEditorEditorSubsystem->GetMutableEntityManager();
+	ActorManager = MakeShareable(new FMassActorManager(MassEntityManager.AsShared()));
+}
diff --git a/Engine/Plugins/Runtime/MassGameplay/Source/MassGameplayEditor/Public/MassActorEditorSubsystem.h b/Engine/Plugins/Runtime/MassGameplay/Source/MassGameplayEditor/Public/MassActorEditorSubsystem.h
new file mode 100644
--- /dev/null
+++ b/Engine/Plugins/Runtime/MassGameplay/Source/MassGameplayEditor/Public/MassActorEditorSubsystem.h
@@ -0,0 +1,23 @@
+// Copyright Epic Games, Inc. All Rights Reserved.
+
+#pragma once
+
+#include "EditorSubsystem.h"
+#include "MassActorEditorSubsystem.generated.h"
+
+
+struct FMassActorManager;
+
+UCLASS()
+class MASSGAMEPLAYEDITOR_API UMassActorEditorSubsystem : public UEditorSubsystem
+{
+	GENERATED_BODY()
+
+public:
+	FMassActorManager& GetMutableActorManager() { return *ActorManager.Get(); }
+
+protected:
+	virtual void Initialize(FSubsystemCollectionBase& Collection) override;
+
+	TSharedPtr<FMassActorManager> ActorManager;
+};
