Separated Add and Set operations in FMassCommandBuffer::Flush for a clearer distinction between the two and ordering.

diff --git a/Engine/Plugins/Runtime/MassEntity/Source/MassEntity/Private/MassCommandBuffer.cpp b/Engine/Plugins/Runtime/MassEntity/Source/MassEntity/Private/MassCommandBuffer.cpp
--- a/Engine/Plugins/Runtime/MassEntity/Source/MassEntity/Private/MassCommandBuffer.cpp
+++ b/Engine/Plugins/Runtime/MassEntity/Source/MassEntity/Private/MassCommandBuffer.cpp
@@ -99,7 +99,6 @@ void FMassCommandBuffer::Flush(FMassEntityManager& EntityManager)
 		LLM_SCOPE_BYNAME(TEXT("Mass/FlushCommands"));
 
 		// array used to group commands depending on their operations. Based on EMassCommandOperationType
-		// @todo I'm opened to suggestions on how to better implement this
 		constexpr int32 CommandTypeOrder[] =
 		{
 			MAX_int32 - 1, // None
@@ -107,10 +106,10 @@ void FMassCommandBuffer::Flush(FMassEntityManager& EntityManager)
 			3, // Add
 			1, // Remove
 			2, // ChangeComposition
-			3, // Set
-			4, // Destroy
+			4, // Set
+			5, // Destroy
 		};
-		static_assert((sizeof(CommandTypeOrder) / sizeof(CommandTypeOrder[0])) == (int)EMassCommandOperationType::MAX, "CommandTypeOrder needs to correspond to all EMassCommandOperationType\'s entries");
+		static_assert(UE_ARRAY_COUNT(CommandTypeOrder) == (int)EMassCommandOperationType::MAX, "CommandTypeOrder needs to correspond to all EMassCommandOperationType\'s entries");
 
 		struct FBatchedCommandsSortedIndex
 		{
