PR #10012: Mass Entity: BatchMoveEntitiesToAnotherArchetype fragment vs out ranges order fix (Contributed by [REDACTED_PERSON])

mz: The fix boils down to splitting BatchMoveEntitiesToAnotherArchetype into two sections: adding the new entities to the new archetype (that we do in original order) and the removal from the previous archetype (that we need to do back-to-front to keep the chunk and entity indices relevant). This way the order of entities moved to the new archetype remains the same which is relevant for subsequent operations (like bulk-setting fragment values).
I've also added a "unit" test testing the behavior being fixed.

diff --git a/Engine/Plugins/Runtime/MassEntity/Source/MassEntity/Private/MassArchetypeData.cpp b/Engine/Plugins/Runtime/MassEntity/Source/MassEntity/Private/MassArchetypeData.cpp
--- a/Engine/Plugins/Runtime/MassEntity/Source/MassEntity/Private/MassArchetypeData.cpp
+++ b/Engine/Plugins/Runtime/MassEntity/Source/MassEntity/Private/MassArchetypeData.cpp
@@ -944,13 +944,8 @@ void FMassArchetypeData::BatchMoveEntitiesToAnotherArchetype(const FMassArchetyp
 {
 	check(&NewArchetype != this);
 
-	// Sorting the subchunks info so that subchunks of a given chunk are processed "from the back". Otherwise removing 
-	// a subchunk from the front of the chunk would inevitably invalidate following subchunks' information.
+
 	TArray<FMassArchetypeEntityCollection::FArchetypeEntityRange> Subchunks(EntityCollection.GetRanges());
-	Subchunks.Sort([](const FMassArchetypeEntityCollection::FArchetypeEntityRange& A, const FMassArchetypeEntityCollection::FArchetypeEntityRange& B)
-		{
-			return A.ChunkIndex < B.ChunkIndex || (A.ChunkIndex == B.ChunkIndex && A.SubchunkStart > B.SubchunkStart);
-		});
 
 	const int32 InitialOutEntitiesCount = OutEntitesBeingMoved.Num();
 
@@ -985,16 +980,35 @@ void FMassArchetypeData::BatchMoveEntitiesToAnotherArchetype(const FMassArchetyp
 
 			NumberMoved += ResultSubChunk.Length;
 
-			// @todo consider adding a 'merge sequences' pass at the end of the function since we can end up with 
-			// ranges being right next to each other
 			if (OutNewRanges)
 			{
-				OutNewRanges->Add(ResultSubChunk);
+				// if the new ResultSubChunk is right next to the last stored one then merge them both
+				if (OutNewRanges->Num() && OutNewRanges->Last().IsAdjacentAfter(ResultSubChunk))
+				{
+					OutNewRanges->Last().Length += ResultSubChunk.Length;
+				}
+				else // just add
+				{
+					OutNewRanges->Add(ResultSubChunk);
+				}
 			}
 
 		} while (NumberMoved < EntityRange.Length);
 
-		BatchRemoveEntitiesInternal(EntityRange.ChunkIndex, EntityRange.SubchunkStart, EntityRange.Length);
+	}
+
+	// Sorting the subchunks info so that subchunks of a given chunk are processed "from the back". Otherwise removing 
+	// a subchunk from the front of the chunk would inevitably invalidate following subchunks' information.
+	// Note that we do this after already having added the entities to the new archetype to preserve the order of entites 
+	// as given by the input data.
+	Subchunks.Sort([](const FMassArchetypeEntityCollection::FArchetypeEntityRange& A, const FMassArchetypeEntityCollection::FArchetypeEntityRange& B)
+	{
+		return A.ChunkIndex < B.ChunkIndex || (A.ChunkIndex == B.ChunkIndex && A.SubchunkStart > B.SubchunkStart);
+	});
+
+	for (const FMassArchetypeEntityCollection::FArchetypeEntityRange& Subchunk : Subchunks)
+	{
+		BatchRemoveEntitiesInternal(Subchunk.ChunkIndex, Subchunk.SubchunkStart, Subchunk.Length);
 	}
 
 	for (int i = InitialOutEntitiesCount; i < OutEntitesBeingMoved.Num(); ++i)
diff --git a/Engine/Plugins/Runtime/MassEntity/Source/MassEntity/Public/MassArchetypeTypes.h b/Engine/Plugins/Runtime/MassEntity/Source/MassEntity/Public/MassArchetypeTypes.h
--- a/Engine/Plugins/Runtime/MassEntity/Source/MassEntity/Public/MassArchetypeTypes.h
+++ b/Engine/Plugins/Runtime/MassEntity/Source/MassEntity/Public/MassArchetypeTypes.h
@@ -69,6 +69,12 @@ public:
 		/** Note that we consider invalid-length chunks valid as long as ChunkIndex and SubchunkStart are valid */
 		bool IsSet() const { return ChunkIndex != INDEX_NONE && SubchunkStart >= 0; }
 
+		/** Checks if given InRange comes right after this instance */
+		bool IsAdjacentAfter(const FArchetypeEntityRange& Other) const
+		{
+			return ChunkIndex == Other.ChunkIndex && SubchunkStart + Length == Other.SubchunkStart;
+		}
+
 		bool operator==(const FArchetypeEntityRange& Other) const
 		{
 			return ChunkIndex == Other.ChunkIndex && SubchunkStart == Other.SubchunkStart && Length == Other.Length;
diff --git a/Engine/Plugins/Runtime/MassEntity/Source/MassEntityTestSuite/Private/MassComplexTest.cpp b/Engine/Plugins/Runtime/MassEntity/Source/MassEntityTestSuite/Private/MassComplexTest.cpp
new file mode 100644
--- /dev/null
+++ b/Engine/Plugins/Runtime/MassEntity/Source/MassEntityTestSuite/Private/MassComplexTest.cpp
@@ -0,0 +1,91 @@
+// Copyright Epic Games, Inc. All Rights Reserved.
+
+#include "CoreMinimal.h"
+#include "AITestsCommon.h"
+#include "MassEntityManager.h"
+#include "MassEntityTestTypes.h"
+#include "MassExecutionContext.h"
+
+#define LOCTEXT_NAMESPACE "MassTest"
+
+UE_DISABLE_OPTIMIZATION_SHIP
+
+//-----------------------------------------------------------------------------
+// This file contains tests spanning multiple functionalities.
+//-----------------------------------------------------------------------------
+namespace FMassComplexTest
+{
+//-----------------------------------------------------------------------------
+// Testing the issue first reported here https://forums.unrealengine.com/t/masscommandaddfragmentinstances-adding-fragments-to-incorrect-entities/731341
+// 
+// The issue boiled down to FMassArchetypeData::BatchMoveEntitiesToAnotherArchetype (that's the function that does the 
+// work behind FMassCommandAddFragmentInstances used in this test) was moving entities to the new archetype in a changed 
+// order, while the call site (FMassEntityManager::BatchAddFragmentInstancesForEntities) the order has not changed and 
+// subsequently assigned mismatched values to the freshly moved entities (via FMassArchetypeData::BatchSetFragmentValues)
+// 
+// The fix was as simple as splitting BatchMoveEntitiesToAnotherArchetype into two sections: 
+// 1. adding all the new entities to the new archetype (that we do in original order) 
+// 2. and only then removal from the previous archetype (that we need to do back-to-front to keep the chunk and entity indices relevant). 
+// This way the order of entities moved to the new archetype remains the same as in the input data.
+//-----------------------------------------------------------------------------
+struct FComplex_AddingFragmentInstancesToDiscontinouousEntitiesCollection : FEntityTestBase
+{
+	virtual bool InstantTest() override
+	{
+		constexpr int EntitiesCount = 4;
+		constexpr int MiddleIndex = EntitiesCount / 2;
+		CA_ASSUME(EntityManager);
+
+		// The setup of the test.
+		// We create EntitiesCount entities containing only a FTestFragment_Int fragment and then set the value
+		// to subsequent integer values
+		TArray<FMassEntityHandle> EntitiesCreated;
+		EntityManager->BatchCreateEntities(IntsArchetype, EntitiesCount, EntitiesCreated);
+
+		FMassExecutionContext ExecContext(*EntityManager.Get());
+		FMassEntityQuery Query({ FTestFragment_Int::StaticStruct() });
+		Query.ForEachEntityChunk(*EntityManager, ExecContext, [](FMassExecutionContext& Context)
+			{
+				TArrayView<FTestFragment_Int> Ints = Context.GetMutableFragmentView<FTestFragment_Int>();
+				for (int32 i = 0; i < Context.GetNumEntities(); ++i)
+				{
+					Ints[i].Value = i;
+				}
+			});
+
+		for (int i = EntitiesCount - 1; i >= 0; --i)
+		{
+			// skipping one index somewhere in the middle to ensure discontinuity in entities being processed by the command
+			if (i != MiddleIndex)
+			{
+				const FTestFragment_Int& IntFragment = EntityManager->GetFragmentDataChecked<FTestFragment_Int>(EntitiesCreated[i]);
+				EntityManager->Defer().PushCommand<FMassCommandAddFragmentInstances>(EntitiesCreated[i], FTestFragment_Float(IntFragment.Value));
+			}
+		}
+		EntityManager->FlushCommands();
+
+		// by now we should have 1 remaining entity in the original IntsArchetype (no point in testing it, there are 
+		// other tests doing that), and 3 entities in the Ints&Floats archetype, with value of Int fragment and Float 
+		// fragment matching.
+		for (int i = 0; i < EntitiesCount; ++i)
+		{
+			if (i != MiddleIndex)
+			{
+				const FTestFragment_Float& FloatFragment = EntityManager->GetFragmentDataChecked<FTestFragment_Float>(EntitiesCreated[i]);
+				const FTestFragment_Int& IntFragment = EntityManager->GetFragmentDataChecked<FTestFragment_Int>(EntitiesCreated[i]);
+				
+				AITEST_EQUAL("", IntFragment.Value, i);
+				AITEST_EQUAL("", IntFragment.Value, (int)FloatFragment.Value);
+			}
+		}
+
+		return true;
+	}
+};
+IMPLEMENT_AI_INSTANT_TEST(FComplex_AddingFragmentInstancesToDiscontinouousEntitiesCollection, "System.Mass.Complex.AddingFragmentInstancesToDiscontinouousEntitiesCollection");
+
+} // namespace FMassComplexTest
+
+UE_ENABLE_OPTIMIZATION_SHIP
+
+#undef LOCTEXT_NAMESPACE
\ No newline at end of file
