Fixed a trivial comment typo in MassEntityTypes.h

diff --git a/Engine/Plugins/Runtime/MassEntity/Source/MassEntity/Public/MassEntityTypes.h b/Engine/Plugins/Runtime/MassEntity/Source/MassEntity/Public/MassEntityTypes.h
--- a/Engine/Plugins/Runtime/MassEntity/Source/MassEntity/Public/MassEntityTypes.h
+++ b/Engine/Plugins/Runtime/MassEntity/Source/MassEntity/Public/MassEntityTypes.h
@@ -462,7 +462,7 @@ namespace UE::Mass
 	};
 
 	/** 
-	 * The type hosts a statically-typed collection of TArrays, where each TArray is strongly types (i.e. it contains 
+	 * The type hosts a statically-typed collection of TArrays, where each TArray is strongly-typed (i.e. it contains 
 	 * instances of given structs rather than structs wrapped up in FInstancedStruct). This type lets us do batched 
 	 * fragment values setting by simply copying data rather than setting per-instance. 
 	 */
