Extended Mass with the ability to (un)register dymamic processors.

#NoJira

diff --git a/Engine/Plugins/Runtime/MassEntity/Source/MassEntity/Private/MassEntitySettings.cpp b/Engine/Plugins/Runtime/MassEntity/Source/MassEntity/Private/MassEntitySettings.cpp
--- a/Engine/Plugins/Runtime/MassEntity/Source/MassEntity/Private/MassEntitySettings.cpp
+++ b/Engine/Plugins/Runtime/MassEntity/Source/MassEntity/Private/MassEntitySettings.cpp
@@ -61,10 +61,10 @@ void UMassEntitySettings::BuildPhases()
 
 			FMassProcessorDependencySolver::FResult Result;
 			Result.DependencyGraphFileName = PhaseDumpDependencyGraphFileName;
-			FMassPhaseProcessorConfigurationHelper Configurator(*PhaseConfig.PhaseProcessor, PhaseConfig, *this);
+			FMassPhaseProcessorConfigurationHelper Configurator(*PhaseConfig.PhaseProcessor, PhaseConfig, *this, EMassProcessingPhase(i));
 			Configurator.bInitializeCreatedProcessors = false;
 			Configurator.bIsGameRuntime = false;
-			Configurator.Configure(/*EntityManager=*/nullptr, &Result);
+			Configurator.Configure({}, /*EntityManager=*/nullptr, &Result);
 
 			if (Result.PrunedProcessorClasses.Num() > 0)
 			{
diff --git a/Engine/Plugins/Runtime/MassEntity/Source/MassEntity/Private/MassProcessingPhaseManager.cpp b/Engine/Plugins/Runtime/MassEntity/Source/MassEntity/Private/MassProcessingPhaseManager.cpp
--- a/Engine/Plugins/Runtime/MassEntity/Source/MassEntity/Private/MassProcessingPhaseManager.cpp
+++ b/Engine/Plugins/Runtime/MassEntity/Source/MassEntity/Private/MassProcessingPhaseManager.cpp
@@ -139,10 +139,19 @@ void FMassProcessingPhase::Initialize(FMassProcessingPhaseManager& InPhaseManage
 //----------------------------------------------------------------------//
 // FPhaseProcessorConfigurator
 //----------------------------------------------------------------------//
-void FMassPhaseProcessorConfigurationHelper::Configure(const TSharedPtr<FMassEntityManager>& EntityManager, FMassProcessorDependencySolver::FResult* OutOptionalResult)
+void FMassPhaseProcessorConfigurationHelper::Configure(TArrayView<UMassProcessor*> DynamicProcessors,
+	const TSharedPtr<FMassEntityManager>& EntityManager, FMassProcessorDependencySolver::FResult* OutOptionalResult)
 {
 	FMassRuntimePipeline TmpPipeline;
 	TmpPipeline.CreateFromArray(PhaseConfig.ProcessorCDOs, ProcessorOuter);
+	for (UMassProcessor* Processor : DynamicProcessors)
+	{
+		checkf(Processor != nullptr, TEXT("Dynamic processor provided to MASS is null."));
+		if (Processor->GetProcessingPhase() == Phase)
+		{
+			TmpPipeline.AppendProcessor(*Processor);
+		}
+	}
 
 	TArray<FMassProcessorOrderInfo> SortedProcessors;
 	FMassProcessorDependencySolver Solver(TmpPipeline.GetMutableProcessors(), bIsGameRuntime);
@@ -218,9 +227,9 @@ void FMassProcessingPhaseManager::Initialize(UObject& InOwner, TConstArrayView<F
 			// Runtime processing graph are initialized at runtime base on the actual archetypes instantiated at call time.
 			FMassProcessorDependencySolver::FResult Result;
 			Result.DependencyGraphFileName = DependencyGraphFileName;
-			FMassPhaseProcessorConfigurationHelper Configurator(*PhaseProcessor, ProcessingPhasesConfig[PhaseAsInt], InOwner);
+			FMassPhaseProcessorConfigurationHelper Configurator(*PhaseProcessor, ProcessingPhasesConfig[PhaseAsInt], InOwner, EMassProcessingPhase(PhaseAsInt));
 			Configurator.bIsGameRuntime = false;
-			Configurator.Configure(/*EntityManager=*/nullptr, &Result);
+			Configurator.Configure({}, /*EntityManager=*/nullptr, &Result);
 		}
 #endif // WITH_EDITOR
 	}
@@ -352,22 +361,24 @@ void FMassProcessingPhaseManager::OnPhaseStart(const FMassProcessingPhase& Phase
 	const int32 PhaseAsInt = int32(Phase.Phase);
 	if (Owner.IsValid()
 		&& ensure(Phase.Phase != EMassProcessingPhase::MAX)
-		&& ProcessingGraphBuildStates[PhaseAsInt].bNewArchetypes
+		&& (ProcessingGraphBuildStates[PhaseAsInt].bNewArchetypes || ProcessingGraphBuildStates[PhaseAsInt].bProcessorsNeedRebuild)
 		// if not a valid index then we're not able to recalculate dependencies 
 		&& ensure(ProcessingPhasesConfig.IsValidIndex(PhaseAsInt)))
 	{
 		TRACE_CPUPROFILER_EVENT_SCOPE_STR("Mass Rebuild Phase Graph");
 
 		FPhaseGraphBuildState& GraphBuildState = ProcessingGraphBuildStates[PhaseAsInt];
-		if (GraphBuildState.bInitialized == false || FMassProcessorDependencySolver::IsResultUpToDate(GraphBuildState.LastResult, EntityManager) == false)
+		if (GraphBuildState.bInitialized == false 
+			|| ProcessingGraphBuildStates[PhaseAsInt].bProcessorsNeedRebuild
+			|| FMassProcessorDependencySolver::IsResultUpToDate(GraphBuildState.LastResult, EntityManager) == false)
 		{
 			UMassCompositeProcessor* PhaseProcessor = ProcessingPhases[PhaseAsInt].PhaseProcessor;
 			check(PhaseProcessor);
 
 			GraphBuildState.LastResult.Reset();
 
-			FMassPhaseProcessorConfigurationHelper Configurator(*PhaseProcessor, ProcessingPhasesConfig[PhaseAsInt], *Owner.Get());
-			Configurator.Configure(EntityManager, &GraphBuildState.LastResult);
+			FMassPhaseProcessorConfigurationHelper Configurator(*PhaseProcessor, ProcessingPhasesConfig[PhaseAsInt], *Owner.Get(), Phase.Phase);
+			Configurator.Configure(DynamicProcessors, EntityManager, &GraphBuildState.LastResult);
 
 			GraphBuildState.bInitialized = true;
 
@@ -383,6 +394,7 @@ void FMassProcessingPhaseManager::OnPhaseStart(const FMassProcessingPhase& Phase
 #endif // WITH_MASSENTITY_DEBUG
 		}
 
+		ProcessingGraphBuildStates[PhaseAsInt].bProcessorsNeedRebuild = false;
 		ProcessingGraphBuildStates[PhaseAsInt].bNewArchetypes = false;
 	}
 }
@@ -411,6 +423,26 @@ FString FMassProcessingPhaseManager::GetName() const
 	return GetNameSafe(Owner.Get()) + TEXT("_MassProcessingPhaseManager");
 }
 
+void FMassProcessingPhaseManager::RegisterDynamicProcessor(UMassProcessor& Processor)
+{
+	DynamicProcessors.Add(&Processor);
+	ProcessingGraphBuildStates[uint32(Processor.GetProcessingPhase())].bProcessorsNeedRebuild = true;
+}
+
+void FMassProcessingPhaseManager::UnregisterDynamicProcessor(UMassProcessor& Processor)
+{
+	int32 Index;
+	if (DynamicProcessors.Find(&Processor, Index))
+	{
+		DynamicProcessors.RemoveAt(Index);
+		ProcessingGraphBuildStates[uint32(Processor.GetProcessingPhase())].bProcessorsNeedRebuild = true;
+	}
+	else
+	{
+		checkf(false, TEXT("Unable to remove Processor '%s', as it was never added or already removed."), *Processor.GetName());
+	}
+}
+
 void FMassProcessingPhaseManager::OnNewArchetype(const FMassArchetypeHandle& NewArchetype)
 {
 	for (FPhaseGraphBuildState& GraphBuildState : ProcessingGraphBuildStates)
diff --git a/Engine/Plugins/Runtime/MassEntity/Source/MassEntity/Public/MassProcessingPhaseManager.h b/Engine/Plugins/Runtime/MassEntity/Source/MassEntity/Public/MassProcessingPhaseManager.h
--- a/Engine/Plugins/Runtime/MassEntity/Source/MassEntity/Public/MassProcessingPhaseManager.h
+++ b/Engine/Plugins/Runtime/MassEntity/Source/MassEntity/Public/MassProcessingPhaseManager.h
@@ -95,16 +95,18 @@ private:
 
 struct MASSENTITY_API FMassPhaseProcessorConfigurationHelper
 {
-	FMassPhaseProcessorConfigurationHelper(UMassCompositeProcessor& InOutPhaseProcessor, const FMassProcessingPhaseConfig& InPhaseConfig, UObject& InProcessorOuter)
-		: PhaseProcessor(InOutPhaseProcessor), PhaseConfig(InPhaseConfig), ProcessorOuter(InProcessorOuter)
+	FMassPhaseProcessorConfigurationHelper(UMassCompositeProcessor& InOutPhaseProcessor, const FMassProcessingPhaseConfig& InPhaseConfig, UObject& InProcessorOuter, EMassProcessingPhase InPhase)
+		: PhaseProcessor(InOutPhaseProcessor), PhaseConfig(InPhaseConfig), ProcessorOuter(InProcessorOuter), Phase(InPhase)
 	{
 	}
 
-	void Configure(const TSharedPtr<FMassEntityManager>& EntityManager = TSharedPtr<FMassEntityManager>(), FMassProcessorDependencySolver::FResult* OutOptionalResult = nullptr);
+	void Configure(TArrayView<UMassProcessor*> DynamicProcessors, const TSharedPtr<FMassEntityManager>& EntityManager = TSharedPtr<FMassEntityManager>(), 
+		FMassProcessorDependencySolver::FResult* OutOptionalResult = nullptr);
 
 	UMassCompositeProcessor& PhaseProcessor;
 	const FMassProcessingPhaseConfig& PhaseConfig;
 	UObject& ProcessorOuter;
+	EMassProcessingPhase Phase;
 	bool bInitializeCreatedProcessors = true;
 	bool bIsGameRuntime = true;
 };
@@ -157,6 +159,11 @@ public:
 
 	FString GetName() const;
 
+	/** Registers a dynamic processor. This needs to be a fully formed processor and will be slotted in during the next tick. */
+	void RegisterDynamicProcessor(UMassProcessor& Processor);
+	/** Removes a previously registered dynamic processor of throws an assert if not found. */
+	void UnregisterDynamicProcessor(UMassProcessor& Processor);
+
 protected:
 	// FGCObject interface
 	virtual void AddReferencedObjects(FReferenceCollector& Collector) override;
@@ -192,12 +199,14 @@ protected:
 	{
 		FMassProcessorDependencySolver::FResult LastResult;
 		bool bNewArchetypes = true;
+		bool bProcessorsNeedRebuild = true;
 		bool bInitialized = false;
 	};
 
 	FMassProcessingPhase ProcessingPhases[(uint8)EMassProcessingPhase::MAX];
 	FPhaseGraphBuildState ProcessingGraphBuildStates[(uint8)EMassProcessingPhase::MAX];
 	TArray<FMassProcessingPhaseConfig> ProcessingPhasesConfig;
+	TArray<UMassProcessor*> DynamicProcessors;
 
 	TSharedPtr<FMassEntityManager> EntityManager;
 
diff --git a/Engine/Plugins/Runtime/MassEntity/Source/MassEntityEditor/Private/MassEntityEditorSubsystem.cpp b/Engine/Plugins/Runtime/MassEntity/Source/MassEntityEditor/Private/MassEntityEditorSubsystem.cpp
--- a/Engine/Plugins/Runtime/MassEntity/Source/MassEntityEditor/Private/MassEntityEditorSubsystem.cpp
+++ b/Engine/Plugins/Runtime/MassEntity/Source/MassEntityEditor/Private/MassEntityEditorSubsystem.cpp
@@ -123,6 +123,8 @@ void UMassEntityEditorSubsystem::StopAndCleanUp()
 
 void UMassEntityEditorSubsystem::Tick(float DeltaTime)
 {
+	IsProcessing = true;
+
 	if (CompletionEvent.IsValid())
 	{
 		CompletionEvent->Wait();
@@ -136,4 +138,18 @@ void UMassEntityEditorSubsystem::Tick(float DeltaTime)
 		CompletionEvent = TGraphTask<UE::Mass::FMassEditorPhaseTickTask>::CreateTask(&Prerequisites)
 			.ConstructAndDispatchWhenReady(PhaseManager, EMassProcessingPhase(PhaseIndex), DeltaTime);
 	}
+
+	IsProcessing = false;
+}
+
+void UMassEntityEditorSubsystem::RegisterDynamicProcessor(UMassProcessor& Processor)
+{
+	checkf(!IsProcessing, TEXT("Unable to add dynamic processors to Mass during processing."));
+	PhaseManager->RegisterDynamicProcessor(Processor);
+}
+
+void UMassEntityEditorSubsystem::UnregisterDynamicProcessor(UMassProcessor& Processor)
+{
+	checkf(!IsProcessing, TEXT("Unable to remove dynamic processors to Mass during processing."));
+	PhaseManager->UnregisterDynamicProcessor(Processor);
 }
diff --git a/Engine/Plugins/Runtime/MassEntity/Source/MassEntityEditor/Public/MassEntityEditorSubsystem.h b/Engine/Plugins/Runtime/MassEntity/Source/MassEntityEditor/Public/MassEntityEditorSubsystem.h
--- a/Engine/Plugins/Runtime/MassEntity/Source/MassEntityEditor/Public/MassEntityEditorSubsystem.h
+++ b/Engine/Plugins/Runtime/MassEntity/Source/MassEntityEditor/Public/MassEntityEditorSubsystem.h
@@ -28,6 +28,11 @@ public:
 
 	FOnPreTick& GetOnPreTickDelegate() { return OnPreTickDelegate; }
 
+	/** Registers a dynamic processor. This needs to be a fully formed processor and will be slotted in during the next tick. */
+	void RegisterDynamicProcessor(UMassProcessor& Processor);
+	/** Removes a previously registered dynamic processor of throws an assert if not found. */
+	void UnregisterDynamicProcessor(UMassProcessor& Processor);
+
 protected:
 	// USubsystem interface begin
 	virtual void Initialize(FSubsystemCollectionBase& Collection) override;
@@ -49,4 +54,6 @@ protected:
 
 	/** Gets broadcast right before issuing FMassProcessingPhaseManager's phases execution */
 	FOnPreTick OnPreTickDelegate;
+	
+	bool IsProcessing{ false };
 };
diff --git a/Engine/Plugins/Runtime/MassGameplay/Source/MassSimulation/Private/MassSimulationSubsystem.cpp b/Engine/Plugins/Runtime/MassGameplay/Source/MassSimulation/Private/MassSimulationSubsystem.cpp
--- a/Engine/Plugins/Runtime/MassGameplay/Source/MassSimulation/Private/MassSimulationSubsystem.cpp
+++ b/Engine/Plugins/Runtime/MassGameplay/Source/MassSimulation/Private/MassSimulationSubsystem.cpp
@@ -132,10 +132,21 @@ void UMassSimulationSubsystem::OnWorldBeginPlay(UWorld& InWorld)
 	StartSimulation(InWorld);
 }
 
+void UMassSimulationSubsystem::RegisterDynamicProcessor(UMassProcessor& Processor)
+{
+	checkf(!IsDuringMassProcessing(), TEXT("Unable to add dynamic processors to Mass during processing."));
+	PhaseManager.RegisterDynamicProcessor(Processor);
+}
+
+void UMassSimulationSubsystem::UnregisterDynamicProcessor(UMassProcessor& Processor)
+{
+	checkf(!IsDuringMassProcessing(), TEXT("Unable to remove dynamic processors to Mass during processing."));
+	PhaseManager.UnregisterDynamicProcessor(Processor);
+}
+
 void UMassSimulationSubsystem::RebuildTickPipeline()
 {
 	TConstArrayView<FMassProcessingPhaseConfig> ProcessingPhasesConfig = GET_MASS_CONFIG_VALUE(GetProcessingPhasesConfig());
-
 	FString DependencyGraphFileName;
 
 #if WITH_EDITOR
diff --git a/Engine/Plugins/Runtime/MassGameplay/Source/MassSimulation/Public/MassSimulationSubsystem.h b/Engine/Plugins/Runtime/MassGameplay/Source/MassSimulation/Public/MassSimulationSubsystem.h
--- a/Engine/Plugins/Runtime/MassGameplay/Source/MassSimulation/Public/MassSimulationSubsystem.h
+++ b/Engine/Plugins/Runtime/MassGameplay/Source/MassSimulation/Public/MassSimulationSubsystem.h
@@ -28,6 +28,9 @@ public:
 	FMassProcessingPhase::FOnPhaseEvent& GetOnProcessingPhaseFinished(const EMassProcessingPhase Phase);
 	static FOnSimulationStarted& GetOnSimulationStarted() { return OnSimulationStarted; }
 
+	void RegisterDynamicProcessor(UMassProcessor& Processor);
+	void UnregisterDynamicProcessor(UMassProcessor& Processor);
+
 	bool IsSimulationStarted() const { return bSimulationStarted; }
 
 	/** @return whether hosted EntityManager is currently, actively being used for processing purposes. Equivalent to calling FMassEntityManager.IsProcessing() */
