Minor improvements to the recent "dynamic processors" change:
* added deprecation to FMassPhaseProcessorConfigurationHelper::Config since it's a part of the public API
* added newly added DynamicProcessors array to FMassProcessingPhaseManager::AddReferencedObjects

diff --git a/Engine/Plugins/Runtime/MassEntity/Source/MassEntity/Private/MassProcessingPhaseManager.cpp b/Engine/Plugins/Runtime/MassEntity/Source/MassEntity/Private/MassProcessingPhaseManager.cpp
--- a/Engine/Plugins/Runtime/MassEntity/Source/MassEntity/Private/MassProcessingPhaseManager.cpp
+++ b/Engine/Plugins/Runtime/MassEntity/Source/MassEntity/Private/MassProcessingPhaseManager.cpp
@@ -294,6 +294,14 @@ void FMassProcessingPhaseManager::AddReferencedObjects(FReferenceCollector& Coll
 			Collector.AddReferencedObject(Phase.PhaseProcessor);
 		}
 	}
+
+	for (UMassProcessor* DynamicProcessor : DynamicProcessors)
+	{
+		if (DynamicProcessor)
+		{
+			Collector.AddReferencedObject(DynamicProcessor);
+		}
+	}
 }
 
 void FMassProcessingPhaseManager::EnableTickFunctions(const UWorld& World)
diff --git a/Engine/Plugins/Runtime/MassEntity/Source/MassEntity/Public/MassProcessingPhaseManager.h b/Engine/Plugins/Runtime/MassEntity/Source/MassEntity/Public/MassProcessingPhaseManager.h
--- a/Engine/Plugins/Runtime/MassEntity/Source/MassEntity/Public/MassProcessingPhaseManager.h
+++ b/Engine/Plugins/Runtime/MassEntity/Source/MassEntity/Public/MassProcessingPhaseManager.h
@@ -109,6 +109,12 @@ struct MASSENTITY_API FMassPhaseProcessorConfigurationHelper
 	EMassProcessingPhase Phase;
 	bool bInitializeCreatedProcessors = true;
 	bool bIsGameRuntime = true;
+
+	UE_DEPRECATED(5.2, "This flavor of Configure has been deprecated. Use the one requiring the first parameter to be an array view of additional processors")
+	void Configure(const TSharedPtr<FMassEntityManager>& EntityManager = TSharedPtr<FMassEntityManager>(), FMassProcessorDependencySolver::FResult* OutOptionalResult = nullptr)
+	{
+		Configure({}, EntityManager, OutOptionalResult);
+	}
 };
 
 /** 
