Fixed missing test justification in MassComplextTest.cpp

diff --git a/Engine/Plugins/Runtime/MassEntity/Source/MassEntityTestSuite/Private/MassComplexTest.cpp b/Engine/Plugins/Runtime/MassEntity/Source/MassEntityTestSuite/Private/MassComplexTest.cpp
--- a/Engine/Plugins/Runtime/MassEntity/Source/MassEntityTestSuite/Private/MassComplexTest.cpp
+++ b/Engine/Plugins/Runtime/MassEntity/Source/MassEntityTestSuite/Private/MassComplexTest.cpp
@@ -1,6 +1,5 @@
 // Copyright Epic Games, Inc. All Rights Reserved.
 
-#include "CoreMinimal.h"
 #include "AITestsCommon.h"
 #include "MassEntityManager.h"
 #include "MassEntityTestTypes.h"
@@ -74,8 +73,8 @@ struct FComplex_AddingFragmentInstancesToDiscontinouousEntitiesCollection : FEnt
 				const FTestFragment_Float& FloatFragment = EntityManager->GetFragmentDataChecked<FTestFragment_Float>(EntitiesCreated[i]);
 				const FTestFragment_Int& IntFragment = EntityManager->GetFragmentDataChecked<FTestFragment_Int>(EntitiesCreated[i]);
 				
-				AITEST_EQUAL("", IntFragment.Value, i);
-				AITEST_EQUAL("", IntFragment.Value, (int)FloatFragment.Value);
+				AITEST_EQUAL("All int fragments are expected to contain the value indicating the order in which the entity has been created.", IntFragment.Value, i);
+				AITEST_EQUAL("The int and float values are expected to remain in sync.", IntFragment.Value, (int)FloatFragment.Value);
 			}
 		}
 
