Cleaned up Mass's tests headers and brought the optimizations disabling back

diff --git a/Engine/Plugins/Runtime/MassEntity/Source/MassEntityTestSuite/Private/MassChunkCollectionTest.cpp b/Engine/Plugins/Runtime/MassEntity/Source/MassEntityTestSuite/Private/MassChunkCollectionTest.cpp
--- a/Engine/Plugins/Runtime/MassEntity/Source/MassEntityTestSuite/Private/MassChunkCollectionTest.cpp
+++ b/Engine/Plugins/Runtime/MassEntity/Source/MassEntityTestSuite/Private/MassChunkCollectionTest.cpp
@@ -6,6 +6,8 @@
 
 #define LOCTEXT_NAMESPACE "MassTest"
 
+UE_DISABLE_OPTIMIZATION_SHIP
+
 //----------------------------------------------------------------------//
 // tests 
 //----------------------------------------------------------------------//
@@ -211,4 +213,6 @@ IMPLEMENT_AI_INSTANT_TEST(FEntityCollection_WithPayload, "System.Mass.EntityColl
 //
 }
 
+UE_ENABLE_OPTIMIZATION_SHIP
+
 #undef LOCTEXT_NAMESPACE
diff --git a/Engine/Plugins/Runtime/MassEntity/Source/MassEntityTestSuite/Private/MassCommandsTest.cpp b/Engine/Plugins/Runtime/MassEntity/Source/MassEntityTestSuite/Private/MassCommandsTest.cpp
--- a/Engine/Plugins/Runtime/MassEntity/Source/MassEntityTestSuite/Private/MassCommandsTest.cpp
+++ b/Engine/Plugins/Runtime/MassEntity/Source/MassEntityTestSuite/Private/MassCommandsTest.cpp
@@ -1,6 +1,5 @@
 // Copyright Epic Games, Inc. All Rights Reserved.
 
-#include "Engine/World.h"
 #include "MassEntityManager.h"
 #include "MassProcessingTypes.h"
 #include "MassEntityTestTypes.h"
@@ -12,6 +11,8 @@
 
 #define LOCTEXT_NAMESPACE "MassTest"
 
+UE_DISABLE_OPTIMIZATION_SHIP
+
 //----------------------------------------------------------------------//
 // tests 
 //----------------------------------------------------------------------//
@@ -243,4 +244,6 @@ IMPLEMENT_AI_INSTANT_TEST(FCommands_DeferredFunction, "System.Mass.Commands.Defe
 #endif // WITH_MASSENTITY_DEBUG
 } // FMassCommandsTest
 
+UE_ENABLE_OPTIMIZATION_SHIP
+
 #undef LOCTEXT_NAMESPACE
diff --git a/Engine/Plugins/Runtime/MassEntity/Source/MassEntityTestSuite/Private/MassCompositeProcessorTest.cpp b/Engine/Plugins/Runtime/MassEntity/Source/MassEntityTestSuite/Private/MassCompositeProcessorTest.cpp
--- a/Engine/Plugins/Runtime/MassEntity/Source/MassEntityTestSuite/Private/MassCompositeProcessorTest.cpp
+++ b/Engine/Plugins/Runtime/MassEntity/Source/MassEntityTestSuite/Private/MassCompositeProcessorTest.cpp
@@ -1,9 +1,6 @@
 // Copyright Epic Games, Inc. All Rights Reserved.
 
-#include "CoreMinimal.h"
 #include "AITestsCommon.h"
-
-#include "Engine/World.h"
 #include "MassEntityManager.h"
 #include "MassExecutionContext.h"
 #include "MassProcessingTypes.h"
@@ -12,6 +9,8 @@
 
 #define LOCTEXT_NAMESPACE "MassTest"
 
+UE_DISABLE_OPTIMIZATION_SHIP
+
 namespace FMassCompositeProcessorTest
 {
 
@@ -67,5 +66,7 @@ IMPLEMENT_AI_INSTANT_TEST(FCompositeProcessorTest_MultipleSubProcessors, "System
 
 } // FMassCompositeProcessorTest
 
+UE_ENABLE_OPTIMIZATION_SHIP
+
 #undef LOCTEXT_NAMESPACE
 
diff --git a/Engine/Plugins/Runtime/MassEntity/Source/MassEntityTestSuite/Private/MassDependencyTest.cpp b/Engine/Plugins/Runtime/MassEntity/Source/MassEntityTestSuite/Private/MassDependencyTest.cpp
--- a/Engine/Plugins/Runtime/MassEntity/Source/MassEntityTestSuite/Private/MassDependencyTest.cpp
+++ b/Engine/Plugins/Runtime/MassEntity/Source/MassEntityTestSuite/Private/MassDependencyTest.cpp
@@ -1,13 +1,13 @@
 // Copyright Epic Games, Inc. All Rights Reserved.
 
-#include "CoreMinimal.h"
 #include "AITestsCommon.h"
-
 #include "MassProcessorDependencySolver.h"
 #include "MassEntityTestTypes.h"
 
 #define LOCTEXT_NAMESPACE "MassTest"
 
+UE_DISABLE_OPTIMIZATION_SHIP
+
 namespace FMassDependencySolverTest
 {
 
@@ -380,4 +380,6 @@ IMPLEMENT_AI_INSTANT_TEST(FCircularDependency, "System.Mass.Dependencies.Circula
 
 } // FMassDependencySolverTest
 
+UE_ENABLE_OPTIMIZATION_SHIP
+
 #undef LOCTEXT_NAMESPACE
diff --git a/Engine/Plugins/Runtime/MassEntity/Source/MassEntityTestSuite/Private/MassEntityTest.cpp b/Engine/Plugins/Runtime/MassEntity/Source/MassEntityTestSuite/Private/MassEntityTest.cpp
--- a/Engine/Plugins/Runtime/MassEntity/Source/MassEntityTestSuite/Private/MassEntityTest.cpp
+++ b/Engine/Plugins/Runtime/MassEntity/Source/MassEntityTestSuite/Private/MassEntityTest.cpp
@@ -1,9 +1,6 @@
 // Copyright Epic Games, Inc. All Rights Reserved.
 
-#include "CoreMinimal.h"
 #include "AITestsCommon.h"
-
-#include "Engine/World.h"
 #include "MassEntityManager.h"
 #include "MassProcessingTypes.h"
 #include "MassEntityTestTypes.h"
@@ -11,6 +8,8 @@
 
 #define LOCTEXT_NAMESPACE "MassTest"
 
+UE_DISABLE_OPTIMIZATION_SHIP
+
 namespace FMassEntityTest
 {
 #if WITH_MASSENTITY_DEBUG
@@ -357,4 +356,6 @@ IMPLEMENT_AI_INSTANT_TEST(FEntityTest_ReserveAPreviouslyBuiltEntity, "System.Mas
 
 } // FMassEntityTestTest
 
+UE_ENABLE_OPTIMIZATION_SHIP
+
 #undef LOCTEXT_NAMESPACE
diff --git a/Engine/Plugins/Runtime/MassEntity/Source/MassEntityTestSuite/Private/MassExecutionTest.cpp b/Engine/Plugins/Runtime/MassEntity/Source/MassEntityTestSuite/Private/MassExecutionTest.cpp
--- a/Engine/Plugins/Runtime/MassEntity/Source/MassEntityTestSuite/Private/MassExecutionTest.cpp
+++ b/Engine/Plugins/Runtime/MassEntity/Source/MassEntityTestSuite/Private/MassExecutionTest.cpp
@@ -1,7 +1,5 @@
 // Copyright Epic Games, Inc. All Rights Reserved.
 
-#include "CoreMinimal.h"
-#include "Engine/World.h"
 #include "MassEntityManager.h"
 #include "MassProcessingTypes.h"
 #include "MassEntityTestTypes.h"
@@ -9,6 +7,8 @@
 
 #define LOCTEXT_NAMESPACE "MassTest"
 
+UE_DISABLE_OPTIMIZATION_SHIP
+
 //----------------------------------------------------------------------//
 // tests 
 //----------------------------------------------------------------------//
@@ -168,4 +168,6 @@ struct FExecution_Sparse : FEntityTestBase
 IMPLEMENT_AI_INSTANT_TEST(FExecution_Sparse, "System.Mass.Execution.Sparse");
 } // FMassExecutionTest
 
+UE_ENABLE_OPTIMIZATION_SHIP
+
 #undef LOCTEXT_NAMESPACE
diff --git a/Engine/Plugins/Runtime/MassEntity/Source/MassEntityTestSuite/Private/MassObserverTest.cpp b/Engine/Plugins/Runtime/MassEntity/Source/MassEntityTestSuite/Private/MassObserverTest.cpp
--- a/Engine/Plugins/Runtime/MassEntity/Source/MassEntityTestSuite/Private/MassObserverTest.cpp
+++ b/Engine/Plugins/Runtime/MassEntity/Source/MassEntityTestSuite/Private/MassObserverTest.cpp
@@ -1,6 +1,5 @@
 // Copyright Epic Games, Inc. All Rights Reserved.
 
-#include "Engine/World.h"
 #include "MassEntityManager.h"
 #include "MassProcessingTypes.h"
 #include "MassEntityTestTypes.h"
@@ -10,6 +9,8 @@
 
 #define LOCTEXT_NAMESPACE "MassTest"
 
+UE_DISABLE_OPTIMIZATION_SHIP
+
 //----------------------------------------------------------------------//
 // tests 
 //----------------------------------------------------------------------//
@@ -213,4 +214,6 @@ IMPLEMENT_AI_INSTANT_TEST(FTag_MultipleArchetypeSwap, "System.Mass.Observer.Tag.
 
 } // FMassObserverTest
 
+UE_ENABLE_OPTIMIZATION_SHIP
+
 #undef LOCTEXT_NAMESPACE
diff --git a/Engine/Plugins/Runtime/MassEntity/Source/MassEntityTestSuite/Private/MassProcessorTest.cpp b/Engine/Plugins/Runtime/MassEntity/Source/MassEntityTestSuite/Private/MassProcessorTest.cpp
--- a/Engine/Plugins/Runtime/MassEntity/Source/MassEntityTestSuite/Private/MassProcessorTest.cpp
+++ b/Engine/Plugins/Runtime/MassEntity/Source/MassEntityTestSuite/Private/MassProcessorTest.cpp
@@ -1,9 +1,6 @@
 // Copyright Epic Games, Inc. All Rights Reserved.
 
-#include "CoreMinimal.h"
 #include "AITestsCommon.h"
-
-#include "Engine/World.h"
 #include "MassEntityManager.h"
 #include "MassExecutionContext.h"
 #include "MassProcessingTypes.h"
@@ -12,6 +9,8 @@
 
 #define LOCTEXT_NAMESPACE "MassTest"
 
+UE_DISABLE_OPTIMIZATION_SHIP
+
 namespace FMassProcessorTest
 {
 
@@ -118,5 +117,7 @@ IMPLEMENT_AI_INSTANT_TEST(FProcessorTest_Requirements, "System.Mass.Processor.Re
 
 } // FMassProcessorTestTest
 
+UE_ENABLE_OPTIMIZATION_SHIP
+
 #undef LOCTEXT_NAMESPACE
 
diff --git a/Engine/Plugins/Runtime/MassEntity/Source/MassEntityTestSuite/Private/MassSystemRequirementTest.cpp b/Engine/Plugins/Runtime/MassEntity/Source/MassEntityTestSuite/Private/MassSystemRequirementTest.cpp
--- a/Engine/Plugins/Runtime/MassEntity/Source/MassEntityTestSuite/Private/MassSystemRequirementTest.cpp
+++ b/Engine/Plugins/Runtime/MassEntity/Source/MassEntityTestSuite/Private/MassSystemRequirementTest.cpp
@@ -1,6 +1,5 @@
 // Copyright Epic Games, Inc. All Rights Reserved.
 
-#include "CoreMinimal.h"
 #include "AITestsCommon.h"
 #include "Engine/Engine.h"
 #include "MassProcessorDependencySolver.h"
diff --git a/Engine/Plugins/Runtime/MassEntity/Source/MassEntityTestSuite/Private/MassUtilsTest.cpp b/Engine/Plugins/Runtime/MassEntity/Source/MassEntityTestSuite/Private/MassUtilsTest.cpp
--- a/Engine/Plugins/Runtime/MassEntity/Source/MassEntityTestSuite/Private/MassUtilsTest.cpp
+++ b/Engine/Plugins/Runtime/MassEntity/Source/MassEntityTestSuite/Private/MassUtilsTest.cpp
@@ -1,6 +1,5 @@
 // Copyright Epic Games, Inc. All Rights Reserved.
 
-#include "Engine/World.h"
 #include "MassEntityManager.h"
 #include "MassProcessingTypes.h"
 #include "MassEntityTestTypes.h"
@@ -11,6 +10,8 @@
 
 #define LOCTEXT_NAMESPACE "MassTest"
 
+UE_DISABLE_OPTIMIZATION_SHIP
+
 //----------------------------------------------------------------------//
 // tests 
 //----------------------------------------------------------------------//
@@ -115,4 +116,6 @@ IMPLEMENT_AI_INSTANT_TEST(FUtilsTest_MultiSort_Payload, "System.Mass.Utils.Multi
 
 } // FMassUtilsTest
 
+UE_ENABLE_OPTIMIZATION_SHIP
+
 #undef LOCTEXT_NAMESPACE
