Added a basic unit test of the World-less Mass processing ticking.
Also, made a minor extension to AITestingSuite for a better support of latent tests.

diff --git a/Engine/Plugins/Runtime/MassEntity/Source/MassEntityTestSuite/Private/MassProcessingPhasesTest.cpp b/Engine/Plugins/Runtime/MassEntity/Source/MassEntityTestSuite/Private/MassProcessingPhasesTest.cpp
new file mode 100644
--- /dev/null
+++ b/Engine/Plugins/Runtime/MassEntity/Source/MassEntityTestSuite/Private/MassProcessingPhasesTest.cpp
@@ -0,0 +1,192 @@
+// Copyright Epic Games, Inc. All Rights Reserved.
+
+#include "AITestsCommon.h"
+#include "MassEntityManager.h"
+#include "MassProcessingPhaseManager.h"
+#include "MassEntityTestTypes.h"
+#include "MassProcessingPhaseManager.h"
+
+#define LOCTEXT_NAMESPACE "MassTest"
+
+UE_DISABLE_OPTIMIZATION_SHIP
+
+namespace FMassProcessingPhasesTest
+{
+
+/** Test-time TaskGraph task for triggering processing phases. */
+struct FMassTestPhaseTickTask
+{
+	FMassTestPhaseTickTask(const TSharedRef<FMassProcessingPhaseManager>& InPhaseManager, const EMassProcessingPhase InPhase, const float InDeltaTime)
+		: PhaseManager(InPhaseManager)
+		, Phase(InPhase)
+		, DeltaTime(InDeltaTime)
+	{
+	}
+
+	static TStatId GetStatId()
+	{
+		RETURN_QUICK_DECLARE_CYCLE_STAT(FMassTestPhaseTickTask, STATGROUP_TaskGraphTasks);
+	}
+
+	static ENamedThreads::Type GetDesiredThread() { return ENamedThreads::GameThread; }
+	static ESubsequentsMode::Type GetSubsequentsMode() { return ESubsequentsMode::TrackSubsequents; }
+
+	void DoTask(ENamedThreads::Type CurrentThread, const FGraphEventRef& MyCompletionGraphEvent)
+	{
+		TRACE_CPUPROFILER_EVENT_SCOPE(FMassTestPhaseTickTask);
+
+		PhaseManager->TriggerPhase(Phase, DeltaTime, MyCompletionGraphEvent);
+	}
+private:
+	const TSharedRef<FMassProcessingPhaseManager> PhaseManager;
+	const EMassProcessingPhase Phase = EMassProcessingPhase::MAX;
+	const float DeltaTime = 0.f;
+};
+
+
+/** The main point of this FMassProcessingPhaseManager extension is to disable world-based ticking, even if a world is available. */
+struct FMassTestProcessingPhaseManager : public FMassProcessingPhaseManager
+{
+	void Start(const TSharedPtr<FMassEntityManager>& InEntityManager)
+	{
+		EntityManager = InEntityManager;
+
+		OnNewArchetypeHandle = EntityManager->GetOnNewArchetypeEvent().AddRaw(this, &FMassTestProcessingPhaseManager::OnNewArchetype);
+
+		// at this point FMassProcessingPhaseManager would call EnableTickFunctions if a world was available
+		// here we're skipping it on purpose
+
+		bIsAllowedToTick = true;
+	}
+
+	void OnNewArchetype(const FMassArchetypeHandle& NewArchetype)
+	{
+		FMassProcessingPhaseManager::OnNewArchetype(NewArchetype);
+	}
+};
+
+
+struct FProcessingPhasesTestBase : FEntityTestBase
+{
+	using Super = FEntityTestBase;
+		
+	TSharedPtr<FMassTestProcessingPhaseManager> PhaseManager;
+	FMassProcessingPhaseConfig PhasesConfig[int(EMassProcessingPhase::MAX)];
+	int32 TickIndex = -1;
+	FGraphEventRef CompletionEvent;
+	float DeltaTime = 1.f/30;
+
+	FProcessingPhasesTestBase()
+	{
+		bMakeWorldEntityManagersOwner = true;
+	}
+
+	virtual bool SetUp() override
+	{
+		if (Super::SetUp())
+		{
+			PhaseManager = MakeShareable(new FMassTestProcessingPhaseManager());
+
+			EntityManager->Initialize();
+
+			if (PopulatePhasesConfig())
+			{
+				TickIndex = -1;
+
+				UWorld* World = FAITestHelpers::GetWorld();
+				check(World);
+
+				PhaseManager->Initialize(*World, PhasesConfig);
+
+				PhaseManager->Start(EntityManager);
+
+				return true;
+			}
+		}
+
+		return false;
+	}
+
+	virtual bool Update() override
+	{
+		if (CompletionEvent.IsValid())
+		{
+			CompletionEvent->Wait();
+		}
+
+		for (int PhaseIndex = 0; PhaseIndex < (int)EMassProcessingPhase::MAX; ++PhaseIndex)
+		{
+			const FGraphEventArray Prerequisites = { CompletionEvent };
+			CompletionEvent = TGraphTask<FMassTestPhaseTickTask>::CreateTask(&Prerequisites)
+				.ConstructAndDispatchWhenReady(PhaseManager.ToSharedRef(), EMassProcessingPhase(PhaseIndex), DeltaTime);
+		}
+
+		++TickIndex;
+		return false;
+	}
+
+	virtual void TearDown() override
+	{
+		if (CompletionEvent.IsValid())
+		{
+			CompletionEvent->Wait();
+		}
+
+		PhaseManager->Stop();
+		PhaseManager.Reset();
+		Super::TearDown();
+	}
+
+	virtual void VerifyLatentResults() override
+	{
+		// we need to make sure all the phases are done before attempting to test results
+		if (CompletionEvent.IsValid())
+		{
+			CompletionEvent->Wait();
+		}
+	}
+
+	virtual bool PopulatePhasesConfig() = 0;
+};
+
+/** this test is here to make sure that the set up that other tests rely on actually works, i.e. the phases are getting ticked at all*/
+struct FTestSetupTest : FProcessingPhasesTestBase
+{
+	using Super = FProcessingPhasesTestBase;
+
+	virtual bool SetUp() override
+	{
+		if (Super::SetUp())
+		{
+			UMassTestStaticCounterProcessor::StaticCounter = -1;
+			return true;
+		}
+		return false;
+	}
+
+	virtual bool PopulatePhasesConfig()
+	{
+		PhasesConfig[0].ProcessorCDOs.Add(GetMutableDefault<UMassTestStaticCounterProcessor>());
+		return true;
+	}
+
+	virtual bool Update() override
+	{
+		Super::Update();
+		return TickIndex >= 3;
+	}
+
+	virtual void VerifyLatentResults() override
+	{
+		Super::VerifyLatentResults();
+		AITEST_EQUAL_LATENT("Expecting the UMassTestStaticCounterProcessor getting ticked as many times as the test ticked", UMassTestStaticCounterProcessor::StaticCounter, TickIndex);
+	}
+};
+IMPLEMENT_AI_LATENT_TEST(FTestSetupTest, "System.Mass.ProcessingPhases.SetupTest");
+
+} // FMassProcessingPhasesTest
+
+UE_ENABLE_OPTIMIZATION_SHIP
+
+#undef LOCTEXT_NAMESPACE
+
diff --git a/Engine/Plugins/Runtime/MassEntity/Source/MassEntityTestSuite/Private/MassTestTypes.cpp b/Engine/Plugins/Runtime/MassEntity/Source/MassEntityTestSuite/Private/MassTestTypes.cpp
--- a/Engine/Plugins/Runtime/MassEntity/Source/MassEntityTestSuite/Private/MassTestTypes.cpp
+++ b/Engine/Plugins/Runtime/MassEntity/Source/MassEntityTestSuite/Private/MassTestTypes.cpp
@@ -6,6 +6,9 @@
 #include "MassExecutionContext.h"
 #include "Engine/World.h"
 
+
+UE_DISABLE_OPTIMIZATION_SHIP
+
 //----------------------------------------------------------------------//
 // Test bases 
 //----------------------------------------------------------------------//
@@ -79,6 +82,15 @@ UMassTestProcessor_FloatsInts::UMassTestProcessor_FloatsInts()
 	};
 }
 
+ int UMassTestStaticCounterProcessor::StaticCounter = 0;
+ UMassTestStaticCounterProcessor::UMassTestStaticCounterProcessor()
+ {
+#if WITH_EDITORONLY_DATA
+	 bCanShowUpInSettings = false;
+#endif // WITH_EDITORONLY_DATA
+	 bAutoRegisterWithProcessingPhases = false;
+	 ExecutionFlags = int32(EProcessorExecutionFlags::All);
+ }
 //----------------------------------------------------------------------//
 // UMassTestWorldSubsystem
 //----------------------------------------------------------------------//
@@ -94,3 +106,4 @@ int32 UMassTestWorldSubsystem::Read() const
 	return Number;
 }
 
+UE_ENABLE_OPTIMIZATION_SHIP
diff --git a/Engine/Plugins/Runtime/MassEntity/Source/MassEntityTestSuite/Public/MassEntityTestTypes.h b/Engine/Plugins/Runtime/MassEntity/Source/MassEntityTestSuite/Public/MassEntityTestTypes.h
--- a/Engine/Plugins/Runtime/MassEntity/Source/MassEntityTestSuite/Public/MassEntityTestTypes.h
+++ b/Engine/Plugins/Runtime/MassEntity/Source/MassEntityTestSuite/Public/MassEntityTestTypes.h
@@ -191,6 +191,23 @@ public:
 	UMassTestProcessor_FloatsInts();
 };
 
+UCLASS()
+class MASSENTITYTESTSUITE_API UMassTestStaticCounterProcessor : public UMassProcessor
+{
+	GENERATED_BODY()
+
+public:
+	UMassTestStaticCounterProcessor();
+	virtual void Execute(FMassEntityManager& EntityManager, FMassExecutionContext& Context) override
+	{
+		++StaticCounter;
+	}
+	virtual void ConfigureQueries() override {}
+	virtual bool ShouldAllowQueryBasedPruning(const bool bRuntimeMode = true) const { return false; }
+
+	static int StaticCounter;
+};
+
 struct MASSENTITYTESTSUITE_API FExecutionTestBase : FAITestBase
 {
 	TSharedPtr<FMassEntityManager> EntityManager;
diff --git a/Engine/Source/Developer/AITestSuite/Private/Tests/AITestsCommon.cpp b/Engine/Source/Developer/AITestSuite/Private/Tests/AITestsCommon.cpp
--- a/Engine/Source/Developer/AITestSuite/Private/Tests/AITestsCommon.cpp
+++ b/Engine/Source/Developer/AITestSuite/Private/Tests/AITestsCommon.cpp
@@ -55,6 +55,16 @@ bool FAITestCommand_PerformTest::Update()
 	return AITest == nullptr || AITest->Update();
 }
 
+bool FAITestCommand_VerifyTestResults::Update()
+{
+	if (AITest)
+	{
+		AITest->VerifyLatentResults();
+	}
+	// signal "done"
+	return true; 
+}
+
 bool FAITestCommand_TearDownTest::Update()
 {
 	if (AITest)
diff --git a/Engine/Source/Developer/AITestSuite/Public/AITestsCommon.h b/Engine/Source/Developer/AITestSuite/Public/AITestsCommon.h
--- a/Engine/Source/Developer/AITestSuite/Public/AITestsCommon.h
+++ b/Engine/Source/Developer/AITestSuite/Public/AITestsCommon.h
@@ -72,6 +72,8 @@ public:
 	virtual bool SetUp() { return true; }
 	/** @return true to indicate that the test is done. */
 	virtual bool Update() { return false; } 
+	/** lets the Test instance test the results. Use AITEST_*_LATENT macros */
+	virtual void VerifyLatentResults() {}
 	/** @return false to indicate an issue with test execution. Will signal to automation framework this test instance failed. */
 	virtual bool InstantTest() { return false;}
 	// it's essential that overriding functions call the super-implementation. Otherwise the check in ~FAITestBase will fail.
@@ -80,6 +82,7 @@ public:
 
 DEFINE_EXPORTED_LATENT_AUTOMATION_COMMAND_ONE_PARAMETER(AITESTSUITE_API, FAITestCommand_SetUpTest, FAITestBase*, AITest);
 DEFINE_EXPORTED_LATENT_AUTOMATION_COMMAND_ONE_PARAMETER(AITESTSUITE_API, FAITestCommand_PerformTest, FAITestBase*, AITest);
+DEFINE_EXPORTED_LATENT_AUTOMATION_COMMAND_ONE_PARAMETER(AITESTSUITE_API, FAITestCommand_VerifyTestResults, FAITestBase*, AITest);
 DEFINE_EXPORTED_LATENT_AUTOMATION_COMMAND_ONE_PARAMETER(AITESTSUITE_API, FAITestCommand_TearDownTest, FAITestBase*, AITest);
 
 // @note that TestClass needs to derive from FAITestBase
@@ -94,6 +97,8 @@ DEFINE_EXPORTED_LATENT_AUTOMATION_COMMAND_ONE_PARAMETER(AITESTSUITE_API, FAITest
 		ADD_LATENT_AUTOMATION_COMMAND(FAITestCommand_SetUpTest(TestInstance)); \
 		/* run latent command to update */ \
 		ADD_LATENT_AUTOMATION_COMMAND(FAITestCommand_PerformTest(TestInstance)); \
+		/* let the Test instance verify the results, calls VerifyLatentResults */ \
+		ADD_LATENT_AUTOMATION_COMMAND(FAITestCommand_VerifyTestResults(TestInstance)); \
 		/* run latent command to tear down */ \
 		ADD_LATENT_AUTOMATION_COMMAND(FAITestCommand_TearDownTest(TestInstance)); \
 		return true; \
