Switched FMTTestBase tests over to the recently added FAITestBase functionality (utilizing VerifyLatentResults to test results).

diff --git a/Engine/Plugins/Runtime/MassEntity/Source/MassEntityTestSuite/Private/MassMultiThreadTest.cpp b/Engine/Plugins/Runtime/MassEntity/Source/MassEntityTestSuite/Private/MassMultiThreadTest.cpp
--- a/Engine/Plugins/Runtime/MassEntity/Source/MassEntityTestSuite/Private/MassMultiThreadTest.cpp
+++ b/Engine/Plugins/Runtime/MassEntity/Source/MassEntityTestSuite/Private/MassMultiThreadTest.cpp
@@ -10,6 +10,8 @@
 
 #define LOCTEXT_NAMESPACE "MassTest"
 
+UE_DISABLE_OPTIMIZATION_SHIP
+
 //----------------------------------------------------------------------//
 // tests 
 //----------------------------------------------------------------------//
@@ -46,14 +48,11 @@ struct FMTTestBase : FEntityTestBase
 
 		if (FinishEvent->IsComplete())
 		{
-			VerifyResults();
 			// signal that we're done with this test
 			return true;
 		}
 		return false;
 	}
-
-	virtual void VerifyResults() = 0;
 };
 
 struct FMTTrivial : FMTTestBase
@@ -88,7 +87,7 @@ struct FMTTrivial : FMTTestBase
 		return true;
 	}
 	
-	virtual void VerifyResults()
+	virtual void VerifyLatentResults() override
 	{
 		AITEST_EQUAL_LATENT("Expected to process all the created entities.", NumToCreate, NumProcessed);
 	}
@@ -168,7 +167,7 @@ struct FMTBasic : FMTTestBase
 		return true;
 	}
 
-	virtual void VerifyResults()
+	virtual void VerifyLatentResults() override
 	{
 		for (int i = 0; i < Entities.Num(); ++i)
 		{
@@ -181,4 +180,6 @@ IMPLEMENT_AI_LATENT_TEST(FMTBasic, "System.Mass.Multithreading.Basic");
 
 } // FMassMultiThreadingTest
 
+UE_ENABLE_OPTIMIZATION_SHIP
+
 #undef LOCTEXT_NAMESPACE
