Added rudimentary Dynamic Mass Processors unit tests

Moved some types around to reuse them in the newly added tests file.

diff --git a/Engine/Plugins/Runtime/MassEntity/Source/MassEntityTestSuite/Private/MassDynamicProcessorsTest.cpp b/Engine/Plugins/Runtime/MassEntity/Source/MassEntityTestSuite/Private/MassDynamicProcessorsTest.cpp
new file mode 100644
--- /dev/null
+++ b/Engine/Plugins/Runtime/MassEntity/Source/MassEntityTestSuite/Private/MassDynamicProcessorsTest.cpp
@@ -0,0 +1,158 @@
+// Copyright Epic Games, Inc. All Rights Reserved.
+
+#include "UObject/UObjectGlobals.h"
+#include "AITestsCommon.h"
+#include "Engine/World.h"
+#include "MassEntityManager.h"
+#include "MassExecutionContext.h"
+#include "MassProcessingTypes.h"
+#include "MassEntityTestTypes.h"
+#include "MassExecutor.h"
+
+#define LOCTEXT_NAMESPACE "MassTest"
+
+UE_DISABLE_OPTIMIZATION_SHIP
+
+namespace FMassDynamicProcessorsTest
+{
+
+struct FDynamicProcessorsAddTrivial : FProcessingPhasesTestBase
+{
+	using Super = FProcessingPhasesTestBase;
+	int32 AllowedTicksCount = 2;
+	int32 AddDynamicProcessorOnTickIndex = 1;
+	int32 NumberOfTimesTicked = 0;
+	UMassTestProcessorBase* DynamicProcessor = nullptr;
+	TWeakObjectPtr<UMassTestProcessorBase> WeakDynamicProcessor;
+
+	virtual bool PopulatePhasesConfig()
+	{
+		// there are going to be no processors initially
+		return true;
+	}
+
+	virtual bool Update() override
+	{
+		Super::Update();
+
+		if (TickIndex == AddDynamicProcessorOnTickIndex)
+		{
+			// on second tick we're adding a dynamic processor
+			DynamicProcessor = NewObject<UMassTestProcessorBase>(World);
+			DynamicProcessor->ExecutionFunction = [this](FMassEntityManager& InEntitySubsystem, FMassExecutionContext& Context)
+			{
+				++NumberOfTimesTicked;
+			};
+
+			// if we Register the processor at this point we risk that the change will be picked up during the tick we just 
+			// triggered via the Super::Update() call. We need to wait for those to ticks to be processed before we can continue
+			if (CompletionEvent)
+			{
+				CompletionEvent->Wait();
+			}
+			PhaseManager->RegisterDynamicProcessor(*DynamicProcessor);
+			WeakDynamicProcessor = DynamicProcessor;
+		}
+
+		// finishing after two ticks of the added dynamic processor
+		return TickIndex >= AllowedTicksCount + AddDynamicProcessorOnTickIndex;
+	}
+
+	virtual void VerifyLatentResults() override
+	{
+		Super::VerifyLatentResults();
+		AITEST_EQUAL_LATENT("Expecting the dynamic processor to be ticked the predicted number of times", NumberOfTimesTicked, AllowedTicksCount);
+	}
+};
+IMPLEMENT_AI_LATENT_TEST(FDynamicProcessorsAddTrivial, "System.Mass.ProcessingPhases.DynamicProcessors.Add");
+
+
+struct FDynamicProcessorsRemoveTrivial : FDynamicProcessorsAddTrivial
+{
+	using Super = FDynamicProcessorsAddTrivial;
+
+	int32 ArbitraryNumberOfTicksAfterRemoval = 2;
+
+	FDynamicProcessorsRemoveTrivial()
+	{
+		AllowedTicksCount = 1;
+	}
+	
+	virtual bool Update() override
+	{
+		Super::Update();
+
+		if (TickIndex == AddDynamicProcessorOnTickIndex + AllowedTicksCount)
+		{
+			// if we Register the processor at this point we risk that the change will be picked up during the tick we just 
+			// triggered via the Super::Update() call. We need to wait for those to ticks to be processed before we can continue
+			if (CompletionEvent)
+			{
+				CompletionEvent->Wait();
+			}
+			PhaseManager->UnregisterDynamicProcessor(*DynamicProcessor);
+		}
+
+		// finishing after two ticks of the added dynamic processor
+		return TickIndex >= AllowedTicksCount + AddDynamicProcessorOnTickIndex + ArbitraryNumberOfTicksAfterRemoval;
+	}
+};
+IMPLEMENT_AI_LATENT_TEST(FDynamicProcessorsRemoveTrivial, "System.Mass.ProcessingPhases.DynamicProcessors.Remove");
+
+
+struct FDynamicProcessorsAddGCShield : FDynamicProcessorsAddTrivial
+{
+	using Super = FDynamicProcessorsAddTrivial;
+
+	FDynamicProcessorsAddGCShield()
+	{
+		AllowedTicksCount = 5;
+	}
+
+	virtual bool Update() override
+	{
+		CA_ASSUME(GEngine);
+		GEngine->Exec(World, TEXT("obj gc"));
+		return Super::Update();
+	}
+
+	virtual void VerifyLatentResults() override
+	{
+		Super::VerifyLatentResults();
+		AITEST_NOT_NULL_LATENT("Expecting the dynamic processor to be valid regardless of multiple garbage collections", WeakDynamicProcessor.Get());
+	}
+};
+IMPLEMENT_AI_LATENT_TEST(FDynamicProcessorsAddGCShield, "System.Mass.ProcessingPhases.DynamicProcessors.AddGCShield");
+
+
+struct FDynamicProcessorsRemoveGCShield : FDynamicProcessorsRemoveTrivial
+{
+	using Super = FDynamicProcessorsRemoveTrivial;
+
+	FDynamicProcessorsRemoveGCShield()
+	{
+		AllowedTicksCount = 5;
+	}
+
+	virtual bool Update() override
+	{
+		CA_ASSUME(GEngine);
+		CollectGarbage(RF_NoFlags);
+		return Super::Update();
+	}
+
+	virtual void VerifyLatentResults() override
+	{
+		Super::VerifyLatentResults();
+		AITEST_NOT_NULL_LATENT("Expecting the dynamic processor to have been set (test implementation verification)", DynamicProcessor);
+		AITEST_NULL_LATENT("Expecting the dynamic processor to be removed by GC after we unregister it", WeakDynamicProcessor.Get());
+	}
+};
+IMPLEMENT_AI_LATENT_TEST(FDynamicProcessorsRemoveGCShield, "System.Mass.ProcessingPhases.DynamicProcessors.RemoveGCShield");
+
+} // FMassDynamicProcessorsTest
+
+UE_ENABLE_OPTIMIZATION_SHIP
+
+#undef LOCTEXT_NAMESPACE
+
diff --git a/Engine/Plugins/Runtime/MassEntity/Source/MassEntityTestSuite/Private/MassProcessingPhasesTest.cpp b/Engine/Plugins/Runtime/MassEntity/Source/MassEntityTestSuite/Private/MassProcessingPhasesTest.cpp
--- a/Engine/Plugins/Runtime/MassEntity/Source/MassEntityTestSuite/Private/MassProcessingPhasesTest.cpp
+++ b/Engine/Plugins/Runtime/MassEntity/Source/MassEntityTestSuite/Private/MassProcessingPhasesTest.cpp
@@ -8,146 +8,88 @@
 
 #define LOCTEXT_NAMESPACE "MassTest"
 
+
 UE_DISABLE_OPTIMIZATION_SHIP
 
-namespace FMassProcessingPhasesTest
+//----------------------------------------------------------------------//
+// FProcessingPhasesTestBase
+//----------------------------------------------------------------------//
+FProcessingPhasesTestBase::FProcessingPhasesTestBase()
 {
+	bMakeWorldEntityManagersOwner = true;
+}
 
-/** Test-time TaskGraph task for triggering processing phases. */
-struct FMassTestPhaseTickTask
+bool FProcessingPhasesTestBase::SetUp()
 {
-	FMassTestPhaseTickTask(const TSharedRef<FMassProcessingPhaseManager>& InPhaseManager, const EMassProcessingPhase InPhase, const float InDeltaTime)
-		: PhaseManager(InPhaseManager)
-		, Phase(InPhase)
-		, DeltaTime(InDeltaTime)
+	if (Super::SetUp())
 	{
-	}
+		PhaseManager = MakeShareable(new UE::Mass::Testing::FMassTestProcessingPhaseManager());
 
-	static TStatId GetStatId()
-	{
-		RETURN_QUICK_DECLARE_CYCLE_STAT(FMassTestPhaseTickTask, STATGROUP_TaskGraphTasks);
-	}
+		EntityManager->Initialize();
 
-	static ENamedThreads::Type GetDesiredThread() { return ENamedThreads::GameThread; }
-	static ESubsequentsMode::Type GetSubsequentsMode() { return ESubsequentsMode::TrackSubsequents; }
+		if (PopulatePhasesConfig())
+		{
+			TickIndex = -1;
 
-	void DoTask(ENamedThreads::Type CurrentThread, const FGraphEventRef& MyCompletionGraphEvent)
-	{
-		TRACE_CPUPROFILER_EVENT_SCOPE(FMassTestPhaseTickTask);
+			World = FAITestHelpers::GetWorld();
+			check(World);
 
-		PhaseManager->TriggerPhase(Phase, DeltaTime, MyCompletionGraphEvent);
+			PhaseManager->Initialize(*World, PhasesConfig);
+
+			PhaseManager->Start(EntityManager);
+
+			return true;
+		}
 	}
-private:
-	const TSharedRef<FMassProcessingPhaseManager> PhaseManager;
-	const EMassProcessingPhase Phase = EMassProcessingPhase::MAX;
-	const float DeltaTime = 0.f;
-};
 
+	return false;
+}
 
-/** The main point of this FMassProcessingPhaseManager extension is to disable world-based ticking, even if a world is available. */
-struct FMassTestProcessingPhaseManager : public FMassProcessingPhaseManager
+bool FProcessingPhasesTestBase::Update() 
 {
-	void Start(const TSharedPtr<FMassEntityManager>& InEntityManager)
+	if (CompletionEvent.IsValid())
 	{
-		EntityManager = InEntityManager;
-
-		OnNewArchetypeHandle = EntityManager->GetOnNewArchetypeEvent().AddRaw(this, &FMassTestProcessingPhaseManager::OnNewArchetype);
-
-		// at this point FMassProcessingPhaseManager would call EnableTickFunctions if a world was available
-		// here we're skipping it on purpose
-
-		bIsAllowedToTick = true;
+		CompletionEvent->Wait();
 	}
 
-	void OnNewArchetype(const FMassArchetypeHandle& NewArchetype)
+	for (int PhaseIndex = 0; PhaseIndex < (int)EMassProcessingPhase::MAX; ++PhaseIndex)
 	{
-		FMassProcessingPhaseManager::OnNewArchetype(NewArchetype);
+		const FGraphEventArray Prerequisites = { CompletionEvent };
+		CompletionEvent = TGraphTask<UE::Mass::Testing::FMassTestPhaseTickTask>::CreateTask(&Prerequisites)
+			.ConstructAndDispatchWhenReady(PhaseManager.ToSharedRef(), EMassProcessingPhase(PhaseIndex), DeltaTime);
 	}
-};
 
+	++TickIndex;
+	return false;
+}
 
-struct FProcessingPhasesTestBase : FEntityTestBase
+void FProcessingPhasesTestBase::TearDown()
 {
-	using Super = FEntityTestBase;
-		
-	TSharedPtr<FMassTestProcessingPhaseManager> PhaseManager;
-	FMassProcessingPhaseConfig PhasesConfig[int(EMassProcessingPhase::MAX)];
-	int32 TickIndex = -1;
-	FGraphEventRef CompletionEvent;
-	float DeltaTime = 1.f/30;
-
-	FProcessingPhasesTestBase()
+	if (CompletionEvent.IsValid())
 	{
-		bMakeWorldEntityManagersOwner = true;
+		CompletionEvent->Wait();
 	}
 
-	virtual bool SetUp() override
-	{
-		if (Super::SetUp())
-		{
-			PhaseManager = MakeShareable(new FMassTestProcessingPhaseManager());
-
-			EntityManager->Initialize();
-
-			if (PopulatePhasesConfig())
-			{
-				TickIndex = -1;
-
-				UWorld* World = FAITestHelpers::GetWorld();
-				check(World);
-
-				PhaseManager->Initialize(*World, PhasesConfig);
-
-				PhaseManager->Start(EntityManager);
-
-				return true;
-			}
-		}
+	PhaseManager->Stop();
+	PhaseManager = nullptr;
+	Super::TearDown();
+}
 
-		return false;
-	}
-
-	virtual bool Update() override
-	{
-		if (CompletionEvent.IsValid())
-		{
-			CompletionEvent->Wait();
-		}
-
-		for (int PhaseIndex = 0; PhaseIndex < (int)EMassProcessingPhase::MAX; ++PhaseIndex)
-		{
-			const FGraphEventArray Prerequisites = { CompletionEvent };
-			CompletionEvent = TGraphTask<FMassTestPhaseTickTask>::CreateTask(&Prerequisites)
-				.ConstructAndDispatchWhenReady(PhaseManager.ToSharedRef(), EMassProcessingPhase(PhaseIndex), DeltaTime);
-		}
-
-		++TickIndex;
-		return false;
-	}
-
-	virtual void TearDown() override
+void FProcessingPhasesTestBase::VerifyLatentResults()
+{
+	// we need to make sure all the phases are done before attempting to test results
+	if (CompletionEvent.IsValid())
 	{
-		if (CompletionEvent.IsValid())
-		{
-			CompletionEvent->Wait();
-		}
-
-		PhaseManager->Stop();
-		PhaseManager.Reset();
-		Super::TearDown();
+		CompletionEvent->Wait();
 	}
+}
 
-	virtual void VerifyLatentResults() override
-	{
-		// we need to make sure all the phases are done before attempting to test results
-		if (CompletionEvent.IsValid())
-		{
-			CompletionEvent->Wait();
-		}
-	}
 
-	virtual bool PopulatePhasesConfig() = 0;
-};
+//----------------------------------------------------------------------//
+// FMassProcessingPhasesTest
+//----------------------------------------------------------------------//
+namespace FMassProcessingPhasesTest
+{
 
 /** this test is here to make sure that the set up that other tests rely on actually works, i.e. the phases are getting ticked at all*/
 struct FTestSetupTest : FProcessingPhasesTestBase
diff --git a/Engine/Plugins/Runtime/MassEntity/Source/MassEntityTestSuite/Private/MassTestTypes.cpp b/Engine/Plugins/Runtime/MassEntity/Source/MassEntityTestSuite/Private/MassTestTypes.cpp
--- a/Engine/Plugins/Runtime/MassEntity/Source/MassEntityTestSuite/Private/MassTestTypes.cpp
+++ b/Engine/Plugins/Runtime/MassEntity/Source/MassEntityTestSuite/Private/MassTestTypes.cpp
@@ -91,6 +91,7 @@ UMassTestProcessor_FloatsInts::UMassTestProcessor_FloatsInts()
 	 bAutoRegisterWithProcessingPhases = false;
 	 ExecutionFlags = int32(EProcessorExecutionFlags::All);
  }
+
 //----------------------------------------------------------------------//
 // UMassTestWorldSubsystem
 //----------------------------------------------------------------------//
@@ -106,4 +107,61 @@ int32 UMassTestWorldSubsystem::Read() const
 	return Number;
 }
 
+
+namespace UE::Mass::Testing
+{
+//----------------------------------------------------------------------//
+// FMassTestPhaseTickTask
+//----------------------------------------------------------------------//
+FMassTestPhaseTickTask::FMassTestPhaseTickTask(const TSharedRef<FMassProcessingPhaseManager>& InPhaseManager, const EMassProcessingPhase InPhase, const float InDeltaTime)
+	: PhaseManager(InPhaseManager)
+	, Phase(InPhase)
+	, DeltaTime(InDeltaTime)
+{
+}
+
+TStatId FMassTestPhaseTickTask::GetStatId()
+{
+	RETURN_QUICK_DECLARE_CYCLE_STAT(FMassTestPhaseTickTask, STATGROUP_TaskGraphTasks);
+}
+
+ENamedThreads::Type FMassTestPhaseTickTask::GetDesiredThread()
+{ 
+	return ENamedThreads::GameThread; 
+}
+
+ESubsequentsMode::Type FMassTestPhaseTickTask::GetSubsequentsMode()
+{ 
+	return ESubsequentsMode::TrackSubsequents; 
+}
+
+void FMassTestPhaseTickTask::DoTask(ENamedThreads::Type CurrentThread, const FGraphEventRef& MyCompletionGraphEvent)
+{
+	TRACE_CPUPROFILER_EVENT_SCOPE(FMassTestPhaseTickTask);
+	PhaseManager->TriggerPhase(Phase, DeltaTime, MyCompletionGraphEvent);
+}
+
+
+//----------------------------------------------------------------------//
+// FMassTestPhaseTickTask
+//----------------------------------------------------------------------//
+void FMassTestProcessingPhaseManager::Start(const TSharedPtr<FMassEntityManager>& InEntityManager)
+{
+	EntityManager = InEntityManager;
+
+	OnNewArchetypeHandle = EntityManager->GetOnNewArchetypeEvent().AddRaw(this, &FMassTestProcessingPhaseManager::OnNewArchetype);
+
+	// at this point FMassProcessingPhaseManager would call EnableTickFunctions if a world was available
+	// here we're skipping it on purpose
+
+	bIsAllowedToTick = true;
+}
+
+void FMassTestProcessingPhaseManager::OnNewArchetype(const FMassArchetypeHandle& NewArchetype)
+{
+	FMassProcessingPhaseManager::OnNewArchetype(NewArchetype);
+}
+
+} // namespace UE::Mass::Testing
+
 UE_ENABLE_OPTIMIZATION_SHIP
diff --git a/Engine/Plugins/Runtime/MassEntity/Source/MassEntityTestSuite/Public/MassEntityTestTypes.h b/Engine/Plugins/Runtime/MassEntity/Source/MassEntityTestSuite/Public/MassEntityTestTypes.h
--- a/Engine/Plugins/Runtime/MassEntity/Source/MassEntityTestSuite/Public/MassEntityTestTypes.h
+++ b/Engine/Plugins/Runtime/MassEntity/Source/MassEntityTestSuite/Public/MassEntityTestTypes.h
@@ -13,10 +13,16 @@
 #include "Subsystems/GameInstanceSubsystem.h"
 #include "Misc/MTAccessDetector.h"
 #include "MassExternalSubsystemTraits.h"
+#include "MassProcessingPhaseManager.h"
 #include "MassEntityTestTypes.generated.h"
 
 
 struct FMassEntityManager;
+struct FMassProcessingPhaseManager;
+namespace UE::Mass::Testing
+{
+	struct FMassTestProcessingPhaseManager;
+}
 
 USTRUCT()
 struct FTestFragment_Float : public FMassFragment
@@ -123,6 +129,8 @@ public:
 
 	FMassEntityQuery& TestGetQuery() { return EntityQuery; }
 
+	void SetShouldAllowDuplicates(const bool bInShouldAllowDuplicated) { bAllowDuplicates = bInShouldAllowDuplicated; }
+
 protected:
 	FMassEntityQuery EntityQuery;
 };
@@ -230,6 +238,26 @@ struct MASSENTITYTESTSUITE_API FEntityTestBase : FExecutionTestBase
 	virtual bool SetUp() override;
 };
 
+
+struct MASSENTITYTESTSUITE_API FProcessingPhasesTestBase : FEntityTestBase
+{
+	using Super = FEntityTestBase;
+
+	TSharedPtr<UE::Mass::Testing::FMassTestProcessingPhaseManager> PhaseManager;
+	FMassProcessingPhaseConfig PhasesConfig[int(EMassProcessingPhase::MAX)];
+	int32 TickIndex = -1;
+	FGraphEventRef CompletionEvent;
+	float DeltaTime = 1.f / 30;
+	UWorld* World = nullptr;
+
+	FProcessingPhasesTestBase();
+	virtual bool SetUp() override;
+	virtual bool Update() override;
+	virtual void TearDown() override;
+	virtual void VerifyLatentResults() override;
+	virtual bool PopulatePhasesConfig() = 0;
+};
+
 template<typename T>
 void ShuffleDataWithRandomStream(FRandomStream& Rand, TArray<T>& Data)
 {
@@ -336,4 +364,32 @@ struct TMassExternalSubsystemTraits<UMassTestGameInstanceSubsystem>
 		ThreadSafeRead = true,
 		ThreadSafeWrite = false,
 	};
-};
\ No newline at end of file
+};
+
+namespace UE::Mass::Testing
+{
+/** Test-time TaskGraph task for triggering processing phases. */
+struct FMassTestPhaseTickTask
+{
+	FMassTestPhaseTickTask(const TSharedRef<FMassProcessingPhaseManager>& InPhaseManager, const EMassProcessingPhase InPhase, const float InDeltaTime);
+
+	static TStatId GetStatId();
+	static ENamedThreads::Type GetDesiredThread();
+	static ESubsequentsMode::Type GetSubsequentsMode();
+
+	void DoTask(ENamedThreads::Type CurrentThread, const FGraphEventRef& MyCompletionGraphEvent);
+
+private:
+	const TSharedRef<FMassProcessingPhaseManager> PhaseManager;
+	const EMassProcessingPhase Phase = EMassProcessingPhase::MAX;
+	const float DeltaTime = 0.f;
+};
+
+/** The main point of this FMassProcessingPhaseManager extension is to disable world-based ticking, even if a world is available. */
+struct FMassTestProcessingPhaseManager : public FMassProcessingPhaseManager
+{
+	void Start(const TSharedPtr<FMassEntityManager>& InEntityManager);
+	void OnNewArchetype(const FMassArchetypeHandle& NewArchetype);
+};
+
+} // namespace UE::Mass::Testing
