Tiny change to FMassProcessorDependencySolver::CreateNodes with a huge potential.
Laveraging UMassProcessor::ShouldAllowDuplicates (previously named UMassProcessor::AllowDuplicates) to support having multiple instances of a single processor class. The main use case being handling dynamic processor instances, that each can be configured differently.
The change has no impact on the current functionality of CitySample and such since the UMassProcessor::bAllowDuplicates flag is never getting changed from the default `false` in the existing code.

diff --git a/Engine/Plugins/Runtime/MassEntity/Source/MassEntity/Private/MassProcessingTypes.cpp b/Engine/Plugins/Runtime/MassEntity/Source/MassEntity/Private/MassProcessingTypes.cpp
--- a/Engine/Plugins/Runtime/MassEntity/Source/MassEntity/Private/MassProcessingTypes.cpp
+++ b/Engine/Plugins/Runtime/MassEntity/Source/MassEntity/Private/MassProcessingTypes.cpp
@@ -140,7 +140,7 @@ void FMassRuntimePipeline::AppendUniqueRuntimeProcessorCopies(TConstArrayView<co
 	for (const UMassProcessor* Proc : InProcessors)
 	{
 		if (Proc && Proc->ShouldExecute(WorldExecutionFlags)
-			&& (Proc->AllowDuplicates() || (HasProcessorOfExactClass(Proc->GetClass()) == false)))
+			&& (Proc->ShouldAllowMultipleInstances() || (HasProcessorOfExactClass(Proc->GetClass()) == false)))
 		{
 			// unfortunately the const cast is required since NewObject doesn't support const Template object
 			UMassProcessor* ProcCopy = NewObject<UMassProcessor>(&InOwner, Proc->GetClass(), FName(), RF_NoFlags, const_cast<UMassProcessor*>(Proc));
@@ -153,7 +153,7 @@ void FMassRuntimePipeline::AppendUniqueRuntimeProcessorCopies(TConstArrayView<co
 			{
 				UE_VLOG(&InOwner, LogMass, Log, TEXT("Skipping %s due to ExecutionFlags"), *Proc->GetName());
 			}
-			else if (Proc->AllowDuplicates() == false)
+			else if (Proc->ShouldAllowMultipleInstances() == false)
 			{
 				UE_VLOG(&InOwner, LogMass, Log, TEXT("Skipping %s due to it being a duplicate"), *Proc->GetName());
 			}
@@ -183,7 +183,7 @@ void FMassRuntimePipeline::AppendOrOverrideRuntimeProcessorCopies(TConstArrayVie
 			UMassProcessor* ProcCopy = NewObject<UMassProcessor>(&InOwner, Proc->GetClass(), FName(), RF_NoFlags, const_cast<UMassProcessor*>(Proc));
 			check(ProcCopy);
 
-			if (ProcCopy->AllowDuplicates())
+			if (ProcCopy->ShouldAllowMultipleInstances())
 			{
 				// we don't care if there are instances of this class in Processors already
 				Processors.Add(ProcCopy);
diff --git a/Engine/Plugins/Runtime/MassEntity/Source/MassEntity/Private/MassProcessorDependencySolver.cpp b/Engine/Plugins/Runtime/MassEntity/Source/MassEntity/Private/MassProcessorDependencySolver.cpp
--- a/Engine/Plugins/Runtime/MassEntity/Source/MassEntity/Private/MassProcessorDependencySolver.cpp
+++ b/Engine/Plugins/Runtime/MassEntity/Source/MassEntity/Private/MassProcessorDependencySolver.cpp
@@ -357,13 +357,25 @@ void FMassProcessorDependencySolver::CreateSubGroupNames(FName InGroupName, TArr
 int32 FMassProcessorDependencySolver::CreateNodes(UMassProcessor& Processor)
 {
 	check(Processor.GetClass());
-	const FName ProcName = Processor.GetClass()->GetFName();
+	FName ProcName = Processor.GetClass()->GetFName();
 
 	if (const int32* NodeIndexPtr = NodeIndexMap.Find(ProcName))
 	{
-		UE_LOG(LogMass, Warning, TEXT("%s Processor %s already registered. Duplicates are not supported.")
-			, ANSI_TO_TCHAR(__FUNCTION__), *ProcName.ToString());
-		return *NodeIndexPtr;
+		// we can accept another instance of this processor class if the processor itself supports that.
+		// Note that the first instance is added with the class while the subsequent instances are added with 
+		// the instance name. This is done on purpose. The class name is used as dependency, so if at least the first 
+		// processor is placed with the class name, like other processors, it can still be used to influence execution
+		// order via ExecuteBefore and ExecuteAfter.
+		if (Processor.ShouldAllowMultipleInstances())
+		{
+			ProcName = Processor.GetFName();
+		}
+		else
+		{
+			UE_LOG(LogMass, Warning, TEXT("%s Processor %s already registered. Duplicates are not supported.")
+				, ANSI_TO_TCHAR(__FUNCTION__), *ProcName.ToString());
+			return *NodeIndexPtr;
+		}
 	}
 
 	const FMassProcessorExecutionOrder& ExecutionOrder = Processor.GetExecutionOrder();
diff --git a/Engine/Plugins/Runtime/MassEntity/Source/MassEntity/Public/MassProcessor.h b/Engine/Plugins/Runtime/MassEntity/Source/MassEntity/Public/MassProcessor.h
--- a/Engine/Plugins/Runtime/MassEntity/Source/MassEntity/Public/MassProcessor.h
+++ b/Engine/Plugins/Runtime/MassEntity/Source/MassEntity/Public/MassProcessor.h
@@ -56,7 +56,14 @@ public:
 	bool ShouldExecute(const EProcessorExecutionFlags CurrentExecutionFlags) const { return (GetExecutionFlags() & CurrentExecutionFlags) != EProcessorExecutionFlags::None; }
 	void CallExecute(FMassEntityManager& EntityManager, FMassExecutionContext& Context);
 
-	bool AllowDuplicates() const { return bAllowDuplicates; }
+	/** 
+	 * Controls whether there can be multiple instances of a given class in a single FMassRuntimePipeline and during 
+	 * dependency solving. 
+	 */
+	bool ShouldAllowMultipleInstances() const { return bAllowMultipleInstances; }
+
+	UE_DEPRECATED(5.2, "This function is deprecated. Use ShouldAllowDuplicates instead.")
+	bool AllowDuplicates() const { return ShouldAllowMultipleInstances(); }
 
 	virtual void DebugOutputDescription(FOutputDevice& Ar, int32 Indent = 0) const;
 	virtual FString GetProcessorName() const { return GetName(); }
@@ -127,6 +134,11 @@ protected:
 	UPROPERTY(EditDefaultsOnly, Category = Processor, config)
 	bool bAutoRegisterWithProcessingPhases = true;
 
+	/** Meant as a class property, make sure to set it in subclass' constructor. Controls whether there can be multiple
+	 *  instances of a given class in a single FMassRuntimePipeline and during dependency solving. */
+	bool bAllowMultipleInstances = false;
+
+	UE_DEPRECATED(5.2, "This property is deprecated. Use bAllowMultipleInstances instead")
 	/** meant as a class property, make sure to set it in subclass' constructor. Controls whether there can be multiple 
 	 *  instances of a given class in a single FMassRuntimePipeline */
 	bool bAllowDuplicates = false;
diff --git a/Engine/Plugins/Runtime/MassEntity/Source/MassEntityTestSuite/Public/MassEntityTestTypes.h b/Engine/Plugins/Runtime/MassEntity/Source/MassEntityTestSuite/Public/MassEntityTestTypes.h
--- a/Engine/Plugins/Runtime/MassEntity/Source/MassEntityTestSuite/Public/MassEntityTestTypes.h
+++ b/Engine/Plugins/Runtime/MassEntity/Source/MassEntityTestSuite/Public/MassEntityTestTypes.h
@@ -129,7 +129,7 @@ public:
 
 	FMassEntityQuery& TestGetQuery() { return EntityQuery; }
 
-	void SetShouldAllowDuplicates(const bool bInShouldAllowDuplicated) { bAllowDuplicates = bInShouldAllowDuplicated; }
+	void SetShouldAllowMultipleInstances(const bool bInShouldAllowDuplicated) { bAllowMultipleInstances = bInShouldAllowDuplicated; }
 
 protected:
 	FMassEntityQuery EntityQuery;
