Added another dynamic processors unit test, and did some Mass test set up massaging to make the processor test implementation more straightforward (mainly by simplifying the query exec function implementation via UMassTestProcessorBase). This exposed a missing test for composite processors (making sure no work is done when no entities match the processors) as well as for query auto-caching of archetypes. Added both.

diff --git a/Engine/Plugins/Runtime/MassEntity/Source/MassEntityTestSuite/Private/MassCompositeProcessorTest.cpp b/Engine/Plugins/Runtime/MassEntity/Source/MassEntityTestSuite/Private/MassCompositeProcessorTest.cpp
--- a/Engine/Plugins/Runtime/MassEntity/Source/MassEntityTestSuite/Private/MassCompositeProcessorTest.cpp
+++ b/Engine/Plugins/Runtime/MassEntity/Source/MassEntityTestSuite/Private/MassCompositeProcessorTest.cpp
@@ -30,26 +30,94 @@ struct FCompositeProcessorTest_Empty : FEntityTestBase
 };
 IMPLEMENT_AI_INSTANT_TEST(FCompositeProcessorTest_Empty, "System.Mass.Processor.Composite.Empty");
 
+/** 
+ * This test ensures that no processors will do any actual work, if there are no entities matching their requirements. 
+ */
+struct FCompositeProcessorTest_NoWork : FEntityTestBase
+{
+	virtual bool InstantTest() override
+	{
+		CA_ASSUME(EntityManager);
+
+		UMassCompositeProcessor* CompositeProcessor = NewObject<UMassCompositeProcessor>();
+		check(CompositeProcessor);
+
+		int TimesExecuted = 0;
+		FMassExecuteFunction QueryExecFunction = [&TimesExecuted](FMassExecutionContext& Context)
+		{
+			++TimesExecuted;
+		};
+		
+		{
+			TArray<UMassProcessor*> Processors;
+			for (int i = 0; i < 3; ++i)
+			{
+				UMassTestProcessorBase* Processor = NewObject<UMassTestProcessorBase>();
+				Processor->ForEachEntityChunkExecutionFunction = QueryExecFunction;
+				Processors.Add(Processor);
+			}
+
+			CompositeProcessor->SetChildProcessors(MoveTemp(Processors));
+		}
+
+		FMassProcessingContext ProcessingContext(EntityManager, /*DeltaSeconds=*/0.f);
+		UE::Mass::Executor::Run(*CompositeProcessor, ProcessingContext);
+		AITEST_EQUAL("None of the execution functions should have been executed", TimesExecuted, 0);
+
+		// now test there being some entities but of different composition. We create Float entities, but processors
+		// require Ints. 
+		constexpr int32 NumberOfEntitieToCreate = 17;
+		TArray<FMassEntityHandle> EntitiesCreated;
+		EntityManager->BatchCreateEntities(FloatsArchetype, NumberOfEntitieToCreate, EntitiesCreated);
+
+		TimesExecuted = 0;
+		{
+			TArray<UMassProcessor*> Processors;
+			for (int i = 0; i < 3; ++i)
+			{
+				UMassTestProcessorBase* Processor = NewObject<UMassTestProcessorBase>();
+				Processor->ForEachEntityChunkExecutionFunction = QueryExecFunction;
+				Processor->EntityQuery.AddRequirement<FTestFragment_Int>(EMassFragmentAccess::ReadWrite);
+				Processors.Add(Processor);
+			}
+
+			CompositeProcessor->SetChildProcessors(MoveTemp(Processors));
+		}
+
+		UE::Mass::Executor::Run(*CompositeProcessor, ProcessingContext);
+		AITEST_EQUAL("None of the execution functions should have been executed", TimesExecuted, 0);
+
+		return true;
+	}
+};
+IMPLEMENT_AI_INSTANT_TEST(FCompositeProcessorTest_NoWork, "System.Mass.Processor.Composite.NoWork");
 
 struct FCompositeProcessorTest_MultipleSubProcessors : FEntityTestBase
 {
 	virtual bool InstantTest() override
 	{
 		CA_ASSUME(EntityManager);
+		// we create a single entity so that the processors' execution functions would get called.
+		EntityManager->CreateEntity(IntsArchetype);
 
 		UMassCompositeProcessor* CompositeProcessor = NewObject<UMassCompositeProcessor>();
 		check(CompositeProcessor);
 
+		int ExpectedResult = 0;
 		int Result = 0;
 		{
 			TArray<UMassProcessor*> Processors;
 			for (int i = 0; i < 3; ++i)
 			{
 				UMassTestProcessorBase* Processor = NewObject<UMassTestProcessorBase>();
-				Processor->ExecutionFunction = [Processor, &Result, i](FMassEntityManager& InEntitySubsystem, FMassExecutionContext& Context) {
-						check(Processor);
-						Result += (int)FMath::Pow(10.f, float(i));
-					};
+				Processor->EntityQuery.AddRequirement<FTestFragment_Int>(EMassFragmentAccess::ReadOnly);
+				int Value = (int)FMath::Pow(10.f, float(i));
+				Processor->ForEachEntityChunkExecutionFunction = [&Result, Value](FMassExecutionContext& Context)
+				{
+					Result += Value;
+				};
+				ExpectedResult += Value;
+
 				Processors.Add(Processor);
 			}
 
@@ -58,7 +126,7 @@ struct FCompositeProcessorTest_MultipleSubProcessors : FEntityTestBase
 
 		FMassProcessingContext ProcessingContext(EntityManager, /*DeltaSeconds=*/0.f);
 		UE::Mass::Executor::Run(*CompositeProcessor, ProcessingContext);
-		AITEST_EQUAL("All of the child processors should get run", Result, 111);
+		AITEST_EQUAL("All of the child processors should get run", Result, ExpectedResult);
 		return true;
 	}
 };
diff --git a/Engine/Plugins/Runtime/MassEntity/Source/MassEntityTestSuite/Private/MassDependencyTest.cpp b/Engine/Plugins/Runtime/MassEntity/Source/MassEntityTestSuite/Private/MassDependencyTest.cpp
--- a/Engine/Plugins/Runtime/MassEntity/Source/MassEntityTestSuite/Private/MassDependencyTest.cpp
+++ b/Engine/Plugins/Runtime/MassEntity/Source/MassEntityTestSuite/Private/MassDependencyTest.cpp
@@ -336,22 +336,22 @@ struct FCircularDependency : FDependencySolverBase
 	{
 		{
 			UMassTestProcessorBase* Proc = Processors.Add_GetRef(NewObject<UMassTestProcessor_A>());
-			Proc->TestGetQuery().AddRequirement<FTestFragment_Int>(EMassFragmentAccess::ReadWrite);
+			Proc->EntityQuery.AddRequirement<FTestFragment_Int>(EMassFragmentAccess::ReadWrite);
 			Proc->GetMutableExecutionOrder().ExecuteAfter.Add(GetProcessorName<UMassTestProcessor_D>());
 		}
 		{
 			UMassTestProcessorBase* Proc = Processors.Add_GetRef(NewObject<UMassTestProcessor_B>());
-			Proc->TestGetQuery().AddRequirement<FTestFragment_Int>(EMassFragmentAccess::ReadWrite);
+			Proc->EntityQuery.AddRequirement<FTestFragment_Int>(EMassFragmentAccess::ReadWrite);
 			Proc->GetMutableExecutionOrder().ExecuteAfter.Add(GetProcessorName<UMassTestProcessor_A>());
 		}
 		{
 			UMassTestProcessorBase* Proc = Processors.Add_GetRef(NewObject<UMassTestProcessor_C>());
-			Proc->TestGetQuery().AddRequirement<FTestFragment_Int>(EMassFragmentAccess::ReadWrite);
+			Proc->EntityQuery.AddRequirement<FTestFragment_Int>(EMassFragmentAccess::ReadWrite);
 			Proc->GetMutableExecutionOrder().ExecuteAfter.Add(GetProcessorName<UMassTestProcessor_B>());
 		}
 		{
 			UMassTestProcessorBase* Proc = Processors.Add_GetRef(NewObject<UMassTestProcessor_D>());
-			Proc->TestGetQuery().AddRequirement<FTestFragment_Int>(EMassFragmentAccess::ReadWrite);
+			Proc->EntityQuery.AddRequirement<FTestFragment_Int>(EMassFragmentAccess::ReadWrite);
 			Proc->GetMutableExecutionOrder().ExecuteAfter.Add(GetProcessorName<UMassTestProcessor_C>());
 		}
 
diff --git a/Engine/Plugins/Runtime/MassEntity/Source/MassEntityTestSuite/Private/MassDynamicProcessorsTest.cpp b/Engine/Plugins/Runtime/MassEntity/Source/MassEntityTestSuite/Private/MassDynamicProcessorsTest.cpp
--- a/Engine/Plugins/Runtime/MassEntity/Source/MassEntityTestSuite/Private/MassDynamicProcessorsTest.cpp
+++ b/Engine/Plugins/Runtime/MassEntity/Source/MassEntityTestSuite/Private/MassDynamicProcessorsTest.cpp
@@ -25,7 +25,7 @@ struct FDynamicProcessorsAddTrivial : FProcessingPhasesTestBase
 	UMassTestProcessorBase* DynamicProcessor = nullptr;
 	TWeakObjectPtr<UMassTestProcessorBase> WeakDynamicProcessor;
 
-	virtual bool PopulatePhasesConfig()
+	virtual bool PopulatePhasesConfig() override
 	{
 		// there are going to be no processors initially
 		return true;
@@ -39,7 +39,7 @@ struct FDynamicProcessorsAddTrivial : FProcessingPhasesTestBase
 		{
 			// on second tick we're adding a dynamic processor
 			DynamicProcessor = NewObject<UMassTestProcessorBase>(World);
-			DynamicProcessor->ExecutionFunction = [this](FMassEntityManager& InEntitySubsystem, FMassExecutionContext& Context)
+			DynamicProcessor->ExecutionFunction = [this](FMassEntityManager& InEntityManager, FMassExecutionContext& Context)
 			{
 				++NumberOfTimesTicked;
 			};
@@ -150,6 +150,63 @@ struct FDynamicProcessorsRemoveGCShield : FDynamicProcessorsRemoveTrivial
 };
 IMPLEMENT_AI_LATENT_TEST(FDynamicProcessorsRemoveGCShield, "System.Mass.ProcessingPhases.DynamicProcessors.RemoveGCShield");
 
+
+struct FDynamicProcessorMultipleInstances : FProcessingPhasesTestBase
+{
+	using Super = FProcessingPhasesTestBase;
+	int32 AllowedTicksCount = 2;
+	int32 AddDynamicProcessorOnTickIndex = 1;
+	int32 NumberOfProcessorsToInstantiate = 3;
+	// using atomic here since due to how the dynamic processors are being configured for this test will execute 
+	// in parallel so the accumulation needs to be thread-safe
+	std::atomic<int32> AccumulatedValue = 0;
+	int32 ExpectedAccumulatedValue = 0;
+
+	virtual bool PopulatePhasesConfig() override
+	{
+		// there are going to be no processors initially
+		return true;
+	}
+
+	virtual bool Update() override
+	{
+		Super::Update();
+
+		if (TickIndex == 1)
+		{
+			// if we Register the processor at this point we risk that the change will be picked up during the tick we just 
+			// triggered via the Super::Update() call. We need to wait for those to ticks to be processed before we can continue
+			if (CompletionEvent)
+			{
+				CompletionEvent->Wait();
+			}
+
+			for (int i = 0; i < NumberOfProcessorsToInstantiate; ++i)
+			{
+				UMassTestProcessorBase* DynamicProcessor = NewObject<UMassTestProcessorBase>(World);
+				DynamicProcessor->SetShouldAllowMultipleInstances(true);
+				DynamicProcessor->ExecutionFunction = [this, i](FMassEntityManager& InEntityManager, FMassExecutionContext& Context)
+				{
+					AccumulatedValue += (i + 1);
+				};
+				ExpectedAccumulatedValue += (i + 1);
+
+				PhaseManager->RegisterDynamicProcessor(*DynamicProcessor);
+			}
+			ExpectedAccumulatedValue *= AllowedTicksCount;
+		}
+
+		return TickIndex >= AllowedTicksCount + AddDynamicProcessorOnTickIndex;
+	}
+
+	virtual void VerifyLatentResults() override
+	{
+		Super::VerifyLatentResults();
+		AITEST_EQUAL_LATENT("Expecting the accumulated value to match the prediction", AccumulatedValue.load(), ExpectedAccumulatedValue);
+	}
+};
+IMPLEMENT_AI_LATENT_TEST(FDynamicProcessorMultipleInstances, "System.Mass.ProcessingPhases.DynamicProcessors.MultipleInstances");
+
 } // FMassDynamicProcessorsTest
 
 UE_ENABLE_OPTIMIZATION_SHIP
diff --git a/Engine/Plugins/Runtime/MassEntity/Source/MassEntityTestSuite/Private/MassMultiThreadTest.cpp b/Engine/Plugins/Runtime/MassEntity/Source/MassEntityTestSuite/Private/MassMultiThreadTest.cpp
--- a/Engine/Plugins/Runtime/MassEntity/Source/MassEntityTestSuite/Private/MassMultiThreadTest.cpp
+++ b/Engine/Plugins/Runtime/MassEntity/Source/MassEntityTestSuite/Private/MassMultiThreadTest.cpp
@@ -74,13 +74,10 @@ struct FMTTrivial : FMTTestBase
 		{
 			UMassTestProcessorBase* Proc = Processors.Add_GetRef(NewObject<UMassTestProcessor_A>());
 
-			Proc->TestGetQuery().AddRequirement<FTestFragment_Int>(EMassFragmentAccess::ReadOnly);
-			Proc->ExecutionFunction = [this, Proc](FMassEntityManager& InEntitySubsystem, FMassExecutionContext& Context)
+			Proc->EntityQuery.AddRequirement<FTestFragment_Int>(EMassFragmentAccess::ReadOnly);
+			Proc->ForEachEntityChunkExecutionFunction = [this](FMassExecutionContext& Context)
 			{
-				Proc->TestGetQuery().ForEachEntityChunk(InEntitySubsystem, Context, [this](FMassExecutionContext& Context)
-					{
-						NumProcessed += Context.GetNumEntities();
-					});
+				NumProcessed += Context.GetNumEntities();
 			};
 		}
 
@@ -114,53 +111,44 @@ struct FMTBasic : FMTTestBase
 		{
 			UMassTestProcessorBase* Proc = Processors.Add_GetRef(NewObject<UMassTestProcessor_C>());
 			Proc->GetMutableExecutionOrder().ExecuteAfter.Add(GetProcessorName<UMassTestProcessor_B>());
-			Proc->TestGetQuery().AddRequirement<FTestFragment_Float>(EMassFragmentAccess::ReadOnly);
-			Proc->TestGetQuery().AddRequirement<FTestFragment_Int>(EMassFragmentAccess::ReadWrite);
-			Proc->ExecutionFunction = [this, Proc](FMassEntityManager& InEntitySubsystem, FMassExecutionContext& Context)
+			Proc->EntityQuery.AddRequirement<FTestFragment_Float>(EMassFragmentAccess::ReadOnly);
+			Proc->EntityQuery.AddRequirement<FTestFragment_Int>(EMassFragmentAccess::ReadWrite);
+			Proc->ForEachEntityChunkExecutionFunction = [this](FMassExecutionContext& Context)
 			{
-				Proc->TestGetQuery().ForEachEntityChunk(InEntitySubsystem, Context, [this](FMassExecutionContext& Context)
-					{
-						const TArrayView<FTestFragment_Int> IntsList = Context.GetMutableFragmentView<FTestFragment_Int>();
-						const TConstArrayView<FTestFragment_Float> FloatsList = Context.GetFragmentView<FTestFragment_Float>();
-						for (int32 i = 0; i < Context.GetNumEntities(); ++i)
-						{
-							IntsList[i].Value = int(FloatsList[i].Value) + IntsList[i].Value;
-						}
-					});
+				const TArrayView<FTestFragment_Int> IntsList = Context.GetMutableFragmentView<FTestFragment_Int>();
+				const TConstArrayView<FTestFragment_Float> FloatsList = Context.GetFragmentView<FTestFragment_Float>();
+				for (int32 i = 0; i < Context.GetNumEntities(); ++i)
+				{
+					IntsList[i].Value = int(FloatsList[i].Value) + IntsList[i].Value;
+				}
 			};
 		} 
 		{
 			UMassTestProcessorBase* Proc = Processors.Add_GetRef(NewObject<UMassTestProcessor_B>());
 			Proc->GetMutableExecutionOrder().ExecuteAfter.Add(GetProcessorName<UMassTestProcessor_A>());
-			Proc->TestGetQuery().AddRequirement<FTestFragment_Int>(EMassFragmentAccess::ReadOnly);
-			Proc->TestGetQuery().AddRequirement<FTestFragment_Float>(EMassFragmentAccess::ReadWrite);
-			Proc->ExecutionFunction = [this, Proc](FMassEntityManager& InEntitySubsystem, FMassExecutionContext& Context)
+			Proc->EntityQuery.AddRequirement<FTestFragment_Int>(EMassFragmentAccess::ReadOnly);
+			Proc->EntityQuery.AddRequirement<FTestFragment_Float>(EMassFragmentAccess::ReadWrite);
+			Proc->ForEachEntityChunkExecutionFunction = [this](FMassExecutionContext& Context)
 			{
-				Proc->TestGetQuery().ForEachEntityChunk(InEntitySubsystem, Context, [this](FMassExecutionContext& Context)
-					{
-						const TConstArrayView<FTestFragment_Int> IntsList = Context.GetFragmentView<FTestFragment_Int>();
-						const TArrayView<FTestFragment_Float> FloatsList = Context.GetMutableFragmentView<FTestFragment_Float>();
-						for (int32 i = 0; i < Context.GetNumEntities(); ++i)
-						{
-							FloatsList[i].Value = float(IntsList[i].Value * IntsList[i].Value);
-						}
-					});
+				const TConstArrayView<FTestFragment_Int> IntsList = Context.GetFragmentView<FTestFragment_Int>();
+				const TArrayView<FTestFragment_Float> FloatsList = Context.GetMutableFragmentView<FTestFragment_Float>();
+				for (int32 i = 0; i < Context.GetNumEntities(); ++i)
+				{
+					FloatsList[i].Value = float(IntsList[i].Value * IntsList[i].Value);
+				}
 			};
 		} 
 		{
 			UMassTestProcessorBase* Proc = Processors.Add_GetRef(NewObject<UMassTestProcessor_A>());
-			Proc->TestGetQuery().AddRequirement<FTestFragment_Int>(EMassFragmentAccess::ReadWrite);
-			Proc->ExecutionFunction = [this, Proc](FMassEntityManager& InEntitySubsystem, FMassExecutionContext& Context)
+			Proc->EntityQuery.AddRequirement<FTestFragment_Int>(EMassFragmentAccess::ReadWrite);
+			Proc->ForEachEntityChunkExecutionFunction = [this](FMassExecutionContext& Context)
 			{
 				int Index = 0;
-				Proc->TestGetQuery().ForEachEntityChunk(InEntitySubsystem, Context, [this, &Index](FMassExecutionContext& Context)
-					{
-						const TArrayView<FTestFragment_Int> IntsList = Context.GetMutableFragmentView<FTestFragment_Int>();
-						for (int32 i = 0; i < Context.GetNumEntities(); ++i)
-						{
-							IntsList[i].Value = Index++;
-						}
-					});
+				const TArrayView<FTestFragment_Int> IntsList = Context.GetMutableFragmentView<FTestFragment_Int>();
+				for (int32 i = 0; i < Context.GetNumEntities(); ++i)
+				{
+					IntsList[i].Value = Index++;
+				}
 			};
 		}
 
diff --git a/Engine/Plugins/Runtime/MassEntity/Source/MassEntityTestSuite/Private/MassObserverTest.cpp b/Engine/Plugins/Runtime/MassEntity/Source/MassEntityTestSuite/Private/MassObserverTest.cpp
--- a/Engine/Plugins/Runtime/MassEntity/Source/MassEntityTestSuite/Private/MassObserverTest.cpp
+++ b/Engine/Plugins/Runtime/MassEntity/Source/MassEntityTestSuite/Private/MassObserverTest.cpp
@@ -38,14 +38,11 @@ struct FTagBaseOperation : FEntityTestBase
 		if (FEntityTestBase::SetUp())
 		{
 			TagObserver = NewObject<UMassTestProcessorBase>();
-			TagObserver->TestGetQuery().AddRequirement<FTestFragment_Int>(EMassFragmentAccess::ReadOnly);
-			TagObserver->TestGetQuery().AddTagRequirement<FTestTag_A>(EMassFragmentPresence::All);
-			TagObserver->ExecutionFunction = [this](FMassEntityManager& InEntitySubsystem, FMassExecutionContext& Context)
+			TagObserver->EntityQuery.AddRequirement<FTestFragment_Int>(EMassFragmentAccess::ReadOnly);
+			TagObserver->EntityQuery.AddTagRequirement<FTestTag_A>(EMassFragmentPresence::All);
+			TagObserver->ForEachEntityChunkExecutionFunction = [this](FMassExecutionContext& Context)
 			{
-				TagObserver->TestGetQuery().ForEachEntityChunk(InEntitySubsystem, Context, [this](FMassExecutionContext& Context)
-					{
-						AffectedEntities.Append(Context.GetEntities().GetData(), Context.GetEntities().Num());
-					});
+				AffectedEntities.Append(Context.GetEntities().GetData(), Context.GetEntities().Num());
 			};
 
 			return true;
diff --git a/Engine/Plugins/Runtime/MassEntity/Source/MassEntityTestSuite/Private/MassProcessorTest.cpp b/Engine/Plugins/Runtime/MassEntity/Source/MassEntityTestSuite/Private/MassProcessorTest.cpp
--- a/Engine/Plugins/Runtime/MassEntity/Source/MassEntityTestSuite/Private/MassProcessorTest.cpp
+++ b/Engine/Plugins/Runtime/MassEntity/Source/MassEntityTestSuite/Private/MassProcessorTest.cpp
@@ -19,12 +19,9 @@ int32 SimpleProcessorRun(FMassEntityManager& EntityManager)
 {
 	int32 EntityProcessedCount = 0;
 	TProcessor* Processor = NewObject<TProcessor>();
-	Processor->ExecutionFunction = [Processor, &EntityProcessedCount](FMassEntityManager& InEntitySubsystem, FMassExecutionContext& Context) {
-		check(Processor);
-		Processor->TestGetQuery().ForEachEntityChunk(InEntitySubsystem, Context, [Processor, &EntityProcessedCount](FMassExecutionContext& Context)
-			{
-				EntityProcessedCount += Context.GetNumEntities();
-			});
+	Processor->ForEachEntityChunkExecutionFunction = [&EntityProcessedCount](FMassExecutionContext& Context) 
+	{
+		EntityProcessedCount += Context.GetNumEntities();
 	};
 
 	FMassProcessingContext ProcessingContext(EntityManager, /*DeltaSeconds=*/0.f);
diff --git a/Engine/Plugins/Runtime/MassEntity/Source/MassEntityTestSuite/Private/MassQueryTest.cpp b/Engine/Plugins/Runtime/MassEntity/Source/MassEntityTestSuite/Private/MassQueryTest.cpp
--- a/Engine/Plugins/Runtime/MassEntity/Source/MassEntityTestSuite/Private/MassQueryTest.cpp
+++ b/Engine/Plugins/Runtime/MassEntity/Source/MassEntityTestSuite/Private/MassQueryTest.cpp
@@ -33,7 +33,7 @@ struct FQueryTest_ProcessorRequirements : FEntityTestBase
 		CA_ASSUME(EntityManager);
 
 		UMassTestProcessor_Floats* Processor = NewObject<UMassTestProcessor_Floats>();
-		TConstArrayView<FMassFragmentRequirementDescription> Requirements = Processor->TestGetQuery().GetFragmentRequirements();
+		TConstArrayView<FMassFragmentRequirementDescription> Requirements = Processor->EntityQuery.GetFragmentRequirements();
 		
 		AITEST_TRUE("Query should have extracted some requirements from the given Processor", Requirements.Num() > 0);
 		AITEST_TRUE("There should be exactly one requirement", Requirements.Num() == 1);
@@ -75,19 +75,13 @@ struct FQueryTest_FragmentViewBinding : FEntityTestBase
 		AITEST_TRUE("Initial value of the fragment should match expectations", TestedFragment.Value == 0.f);
 
 		UMassTestProcessor_Floats* Processor = NewObject<UMassTestProcessor_Floats>();
-		Processor->ExecutionFunction = [Processor](FMassEntityManager& InEntitySubsystem, FMassExecutionContext& Context) {
-			check(Processor);
-			//FMassEntityQuery Query(*Processor);
-
-			Processor->TestGetQuery().ForEachEntityChunk(InEntitySubsystem, Context, [](FMassExecutionContext& Context)
-				{
-					TArrayView<FTestFragment_Float> Floats = Context.GetMutableFragmentView<FTestFragment_Float>();
-
-					for (int32 i = 0; i < Context.GetNumEntities(); ++i)
-					{
-						Floats[i].Value = 13.f;
-					}
-				});
+		Processor->ForEachEntityChunkExecutionFunction = [](FMassExecutionContext& Context) 
+		{
+			TArrayView<FTestFragment_Float> Floats = Context.GetMutableFragmentView<FTestFragment_Float>();
+			for (int32 i = 0; i < Context.GetNumEntities(); ++i)
+			{
+				Floats[i].Value = 13.f;
+			}
 		};
 
 		FMassProcessingContext ProcessingContext(*EntityManager, /*DeltaSeconds=*/0.f);
@@ -543,6 +537,55 @@ struct FQueryTest_AnyTag : FEntityTestBase
 };
 IMPLEMENT_AI_INSTANT_TEST(FQueryTest_AnyTag, "System.Mass.Query.AnyTag");
 
+
+struct FQueryTest_AutoRecache : FEntityTestBase
+{
+	virtual bool InstantTest() override
+	{
+		CA_ASSUME(EntityManager);
+
+		FMassEntityQuery Query;
+		// at least one fragment requirement needs to be present for the query to be valid
+		Query.AddRequirement<FTestFragment_Int>(EMassFragmentAccess::ReadOnly);
+		
+		int32 EntitiesFound = 0;
+		const FMassExecuteFunction QueryExecFunction = [&EntitiesFound](FMassExecutionContext& Context)
+		{
+			EntitiesFound += Context.GetNumEntities();
+		};
+		
+		FMassExecutionContext ExecutionContext(*EntityManager, /*DeltaSeconds=*/0.f);
+		Query.ForEachEntityChunk(*EntityManager, ExecutionContext, QueryExecFunction);
+		
+		AITEST_EQUAL("No entities have been created so we expect counting to yield 0", EntitiesFound, 0);
+
+		constexpr int32 NumberOfEntitiesMatching = 17;
+		TArray<FMassEntityHandle> MatchingEntities;
+		EntityManager->BatchCreateEntities(IntsArchetype, NumberOfEntitiesMatching, MatchingEntities);
+
+		EntitiesFound = 0;
+		Query.ForEachEntityChunk(*EntityManager, ExecutionContext, QueryExecFunction);
+		AITEST_EQUAL("The number of entities found should match the number of entities created in the matching archetype", EntitiesFound, MatchingEntities.Num());
+
+		// create more entities, but in an archetype not matching the query
+		constexpr int32 NumberOfEntitiesNotMatching = 13;
+		TArray<FMassEntityHandle> NotMatchingEntities;
+		EntityManager->BatchCreateEntities(FloatsArchetype, NumberOfEntitiesNotMatching, NotMatchingEntities);
+		EntitiesFound = 0;
+		Query.ForEachEntityChunk(*EntityManager, ExecutionContext, QueryExecFunction);
+		AITEST_EQUAL("The number of entities found should not change with addition of entities not matching the query", EntitiesFound, MatchingEntities.Num());
+
+		// create some more in another matching archetype
+		EntityManager->BatchCreateEntities(FloatsIntsArchetype, NumberOfEntitiesMatching, MatchingEntities);
+		EntitiesFound = 0;
+		Query.ForEachEntityChunk(*EntityManager, ExecutionContext, QueryExecFunction);
+		AITEST_EQUAL("The total number of entities found should include entities from both matching archetypes", EntitiesFound, MatchingEntities.Num());
+
+		return true;
+	}
+};
+IMPLEMENT_AI_INSTANT_TEST(FQueryTest_AutoRecache, "System.Mass.Query.AutoReCaching");
+
 } // FMassQueryTest
 
 #undef LOCTEXT_NAMESPACE
diff --git a/Engine/Plugins/Runtime/MassEntity/Source/MassEntityTestSuite/Private/MassTestTypes.cpp b/Engine/Plugins/Runtime/MassEntity/Source/MassEntityTestSuite/Private/MassTestTypes.cpp
--- a/Engine/Plugins/Runtime/MassEntity/Source/MassEntityTestSuite/Private/MassTestTypes.cpp
+++ b/Engine/Plugins/Runtime/MassEntity/Source/MassEntityTestSuite/Private/MassTestTypes.cpp
@@ -53,33 +53,32 @@ UMassTestProcessorBase::UMassTestProcessorBase()
 	bAutoRegisterWithProcessingPhases = false;
 	ExecutionFlags = int32(EProcessorExecutionFlags::All);
 
-	ExecutionFunction = [](FMassEntityManager& InEntitySubsystem, FMassExecutionContext& Context) {};
-	RequirementsFunction = [](FMassEntityQuery& Query){};
+	ForEachEntityChunkExecutionFunction = [](FMassExecutionContext& Context) {};
+	ExecutionFunction = [this](FMassEntityManager& InEntitySubsystem, FMassExecutionContext& Context) 
+	{
+		EntityQuery.ForEachEntityChunk(InEntitySubsystem, Context, ForEachEntityChunkExecutionFunction);
+	};
+}
+
+void UMassTestProcessorBase::Execute(FMassEntityManager& EntityManager, FMassExecutionContext& Context)
+{
+	ExecutionFunction(EntityManager, Context);
 }
 
 UMassTestProcessor_Floats::UMassTestProcessor_Floats()
 {
-	RequirementsFunction = [this](FMassEntityQuery& Query)
-	{
-		EntityQuery.AddRequirement<FTestFragment_Float>(EMassFragmentAccess::ReadWrite);
-	};
+	EntityQuery.AddRequirement<FTestFragment_Float>(EMassFragmentAccess::ReadWrite);
 }
 
 UMassTestProcessor_Ints::UMassTestProcessor_Ints()
 {
-	RequirementsFunction = [this](FMassEntityQuery& Query)
-	{
-		EntityQuery.AddRequirement<FTestFragment_Int>(EMassFragmentAccess::ReadWrite);
-	};
+	EntityQuery.AddRequirement<FTestFragment_Int>(EMassFragmentAccess::ReadWrite);
 }
 
 UMassTestProcessor_FloatsInts::UMassTestProcessor_FloatsInts()
 {
-	RequirementsFunction = [this](FMassEntityQuery& Query)
-	{
-		EntityQuery.AddRequirement<FTestFragment_Float>(EMassFragmentAccess::ReadWrite);
-		EntityQuery.AddRequirement<FTestFragment_Int>(EMassFragmentAccess::ReadWrite);
-	};
+	EntityQuery.AddRequirement<FTestFragment_Float>(EMassFragmentAccess::ReadWrite);
+	EntityQuery.AddRequirement<FTestFragment_Int>(EMassFragmentAccess::ReadWrite);
 }
 
  int UMassTestStaticCounterProcessor::StaticCounter = 0;
diff --git a/Engine/Plugins/Runtime/MassEntity/Source/MassEntityTestSuite/Public/MassEntityTestTypes.h b/Engine/Plugins/Runtime/MassEntity/Source/MassEntityTestSuite/Public/MassEntityTestTypes.h
--- a/Engine/Plugins/Runtime/MassEntity/Source/MassEntityTestSuite/Public/MassEntityTestTypes.h
+++ b/Engine/Plugins/Runtime/MassEntity/Source/MassEntityTestSuite/Public/MassEntityTestTypes.h
@@ -114,24 +114,28 @@ class MASSENTITYTESTSUITE_API UMassTestProcessorBase : public UMassProcessor
 public:
 	UMassTestProcessorBase();
 	FMassProcessorExecutionOrder& GetMutableExecutionOrder() { return ExecutionOrder; }
-	virtual void Execute(FMassEntityManager& EntityManager, FMassExecutionContext& Context) override 
-	{
-		ExecutionFunction(EntityManager, Context);
-	}
-	virtual void ConfigureQueries() override
-	{
-		RequirementsFunction(EntityQuery);
-	}
+	virtual void Execute(FMassEntityManager& EntityManager, FMassExecutionContext& Context) override;
 	virtual bool ShouldAllowQueryBasedPruning(const bool bRuntimeMode = true) const { return false; }
+	/** 
+	 * ConfigureQueries is called in PostInitProperties, so there's no point in calling anything here, 
+	 * since the EntityQuery configuration will take place after a processor instance gets created
+	 * (so after PostInitProperties has already been called). Use EntityQuery directly instead. 
+	 */
+	virtual void ConfigureQueries() override {}
+
+	using FExecutionFunction = TFunction<void(FMassEntityManager& EntityManager, FMassExecutionContext& Context)>;
+	FExecutionFunction ExecutionFunction;
 
-	TFunction<void(FMassEntityManager& EntityManager, FMassExecutionContext& Context)> ExecutionFunction;
-	TFunction<void(FMassEntityQuery& Query)> RequirementsFunction;
+	/** 
+	 * By default ExecutionFunction is configured to pass this function over to EntityQuery.ForEachEntityChunk call. 
+	 * Note that this function won't be used if you override ExecutionFunction's default value.
+	 */
+	FMassExecuteFunction ForEachEntityChunkExecutionFunction;
 
-	FMassEntityQuery& TestGetQuery() { return EntityQuery; }
 
 	void SetShouldAllowMultipleInstances(const bool bInShouldAllowDuplicated) { bAllowMultipleInstances = bInShouldAllowDuplicated; }
 
-protected:
+	/** public on purpose, this is a test processor, no worries about access*/
 	FMassEntityQuery EntityQuery;
 };
 
