Removed the last of TRemoveConst and deprecated it.

diff --git a/Engine/Plugins/Experimental/AnimNext/Source/AnimNextInterface/Public/AnimNextInterfaceParam.h b/Engine/Plugins/Experimental/AnimNext/Source/AnimNextInterface/Public/AnimNextInterfaceParam.h
--- a/Engine/Plugins/Experimental/AnimNext/Source/AnimNextInterface/Public/AnimNextInterfaceParam.h
+++ b/Engine/Plugins/Experimental/AnimNext/Source/AnimNextInterface/Public/AnimNextInterfaceParam.h
@@ -449,7 +449,7 @@ public:
 protected:
 	void* GetDataForConstructor(ValueType* InValuePtrToWrap) const
 	{
-		using MutableValueType = typename TRemoveConst<ValueType>::Type;
+		using MutableValueType = std::remove_const_t<ValueType>;
 
 		// Support containers with Num() and GetData()
 		if constexpr (TModels<Private::CSizedContainerWithAccessibleDataAsRawPtr, ValueType>::Value)
diff --git a/Engine/Plugins/Runtime/MassEntity/Source/MassEntity/Public/MassSubsystemAccess.h b/Engine/Plugins/Runtime/MassEntity/Source/MassEntity/Public/MassSubsystemAccess.h
--- a/Engine/Plugins/Runtime/MassEntity/Source/MassEntity/Public/MassSubsystemAccess.h
+++ b/Engine/Plugins/Runtime/MassEntity/Source/MassEntity/Public/MassSubsystemAccess.h
@@ -192,11 +192,11 @@ protected:
 		{
 			if constexpr (DoesRequireWorld<T>())
 			{
-				SystemInstance = FetchSubsystemInstance<typename TRemoveConst<T>::Type>(World.Get());
+				SystemInstance = FetchSubsystemInstance<std::remove_const_t<T>>(World.Get());
 			}
 			else
 			{
-				SystemInstance = FetchSubsystemInstance<typename TRemoveConst<T>::Type>();
+				SystemInstance = FetchSubsystemInstance<std::remove_const_t<T>>();
 			}
 			Subsystems[SystemIndex] = SystemInstance;
 		}
diff --git a/Engine/Source/Runtime/Core/Public/Templates/UnrealTypeTraits.h b/Engine/Source/Runtime/Core/Public/Templates/UnrealTypeTraits.h
--- a/Engine/Source/Runtime/Core/Public/Templates/UnrealTypeTraits.h
+++ b/Engine/Source/Runtime/Core/Public/Templates/UnrealTypeTraits.h
@@ -298,7 +298,7 @@ struct TCallTraitsParamTypeHelper<T*, true>
  *		MutablePtr = const_cast< RemoveConst< ConstPtrType >::Type >( ConstPtr );
  */
 template< class T >
-struct /*UE_DEPRECATED(5.3, "TRemoveConst has been deprecated, please use std::remove_const_t instead.")*/ TRemoveConst
+struct UE_DEPRECATED(5.3, "TRemoveConst has been deprecated, please use std::remove_const_t instead.") TRemoveConst
 {
 	using Type = std::remove_const_t<T>;
 };
