fixed UE::Mass::Utils::GetEntityManager to not crash if no

diff --git a/Engine/Plugins/Runtime/MassEntity/Source/MassEntity/Private/MassEntityUtils.cpp b/Engine/Plugins/Runtime/MassEntity/Source/MassEntity/Private/MassEntityUtils.cpp
--- a/Engine/Plugins/Runtime/MassEntity/Source/MassEntity/Private/MassEntityUtils.cpp
+++ b/Engine/Plugins/Runtime/MassEntity/Source/MassEntity/Private/MassEntityUtils.cpp
@@ -62,8 +62,9 @@ void CreateEntityCollections(const FMassEntityManager& EntityManager, const TCon
 FMassEntityManager* GetEntityManager(const UWorld* World)
 {
 	UMassEntitySubsystem* EntityManager = UWorld::GetSubsystem<UMassEntitySubsystem>(World);
-	check(EntityManager);
-	return &EntityManager->GetMutableEntityManager();
+	return EntityManager
+		? &EntityManager->GetMutableEntityManager()
+		: nullptr;
 }
 
 FMassEntityManager& GetEntityManagerChecked(const UWorld& World)
