[MassEntity] fixed processor execution flag for world for PIE

diff --git a/Engine/Plugins/Runtime/MassEntity/Source/MassEntity/Private/MassEntityUtils.cpp b/Engine/Plugins/Runtime/MassEntity/Source/MassEntity/Private/MassEntityUtils.cpp
--- a/Engine/Plugins/Runtime/MassEntity/Source/MassEntity/Private/MassEntityUtils.cpp
+++ b/Engine/Plugins/Runtime/MassEntity/Source/MassEntity/Private/MassEntityUtils.cpp
@@ -12,7 +12,7 @@ namespace UE::Mass::Utils
 {
 EProcessorExecutionFlags GetProcessorExecutionFlagsForWold(const UWorld& World)
 {
-	if (World.IsEditorWorld())
+	if (World.IsEditorWorld() && !World.IsGameWorld())
 	{
 		return EProcessorExecutionFlags::Editor;
 	}
