Added UE_ENABLE_INCLUDE_ORDER_DEPRECATED_IN_5_2
 to handle removal of MassExecutionContext.h from MassEntityManager.h

diff --git a/Engine/Plugins/Runtime/MassEntity/Source/MassEntity/Public/MassEntityManager.h b/Engine/Plugins/Runtime/MassEntity/Source/MassEntity/Public/MassEntityManager.h
--- a/Engine/Plugins/Runtime/MassEntity/Source/MassEntity/Public/MassEntityManager.h
+++ b/Engine/Plugins/Runtime/MassEntity/Source/MassEntity/Public/MassEntityManager.h
@@ -493,3 +493,7 @@ private:
 
 	bool bInitialized = false;
 };
+
+#if UE_ENABLE_INCLUDE_ORDER_DEPRECATED_IN_5_2
+#include "MassExecutionContext.h"
+#endif
