Fixed FMassArchetypeData::ExecutionFunctionForChunk using wrong bindings for shared fragments

diff --git a/Engine/Plugins/Runtime/MassEntity/Source/MassEntity/Private/MassArchetypeData.cpp b/Engine/Plugins/Runtime/MassEntity/Source/MassEntity/Private/MassArchetypeData.cpp
--- a/Engine/Plugins/Runtime/MassEntity/Source/MassEntity/Private/MassArchetypeData.cpp
+++ b/Engine/Plugins/Runtime/MassEntity/Source/MassEntity/Private/MassArchetypeData.cpp
@@ -455,8 +455,8 @@ void FMassArchetypeData::ExecutionFunctionForChunk(FMassExecutionContext RunCont
 
 	if (ChunkLength)
 	{
-		BindConstSharedFragmentRequirements(RunContext, Chunk.GetSharedFragmentValues(), RequirementMapping.ChunkFragments);
-		BindSharedFragmentRequirements(RunContext, Chunk.GetMutableSharedFragmentValues(), RequirementMapping.ChunkFragments);
+		BindConstSharedFragmentRequirements(RunContext, Chunk.GetSharedFragmentValues(), RequirementMapping.ConstSharedFragments);
+		BindSharedFragmentRequirements(RunContext, Chunk.GetMutableSharedFragmentValues(), RequirementMapping.SharedFragments);
 
 		RunContext.SetCurrentArchetypesTagBitSet(GetTagBitSet());
 		RunContext.SetCurrentChunkSerialModificationNumber(Chunk.GetSerialModificationNumber());
