Fixed log statement to use constant format strings

diff --git a/Engine/Plugins/Runtime/MassEntity/Source/MassEntity/Private/MassProcessor.cpp b/Engine/Plugins/Runtime/MassEntity/Source/MassEntity/Private/MassProcessor.cpp
--- a/Engine/Plugins/Runtime/MassEntity/Source/MassEntity/Private/MassProcessor.cpp
+++ b/Engine/Plugins/Runtime/MassEntity/Source/MassEntity/Private/MassProcessor.cpp
@@ -240,7 +240,7 @@ void UMassProcessor::RegisterQuery(FMassEntityQuery& Query)
 	}
 	else
 	{
-		constexpr TCHAR MessageFormat[] = TEXT("Registering entity query for %s while the query is not given processor's member variable. Skipping.");
+		static constexpr TCHAR MessageFormat[] = TEXT("Registering entity query for %s while the query is not given processor's member variable. Skipping.");
 		checkf(false, MessageFormat, *GetProcessorName());
 		UE_LOG(LogMass, Error, MessageFormat, *GetProcessorName());
 	}
diff --git a/Engine/Source/Runtime/Engine/Private/DemoNetDriver.cpp b/Engine/Source/Runtime/Engine/Private/DemoNetDriver.cpp
--- a/Engine/Source/Runtime/Engine/Private/DemoNetDriver.cpp
+++ b/Engine/Source/Runtime/Engine/Private/DemoNetDriver.cpp
@@ -315,7 +315,7 @@ struct FDemoBudgetLogHelper
 
 		if (DemoNetDriverRecordingPrivate::WarningTimeInterval == 0.f)
 		{
-			UE_LOG(LogDemo, Log, Format, Args...);
+			UE_LOG(LogDemo, Log, TEXT("%s"), *FString::Printf(Format, Args...));
 			return;
 		}
 
