[MassEntity] cache subsystem requirements for queries
#rnx

diff --git a/Engine/Plugins/Runtime/MassEntity/Source/MassEntity/Private/MassEntityQuery.cpp b/Engine/Plugins/Runtime/MassEntity/Source/MassEntity/Private/MassEntityQuery.cpp
--- a/Engine/Plugins/Runtime/MassEntity/Source/MassEntity/Private/MassEntityQuery.cpp
+++ b/Engine/Plugins/Runtime/MassEntity/Source/MassEntity/Private/MassEntityQuery.cpp
@@ -177,47 +177,49 @@ void FMassEntityQuery::ForEachEntityChunk(FMassEntityManager& EntityManager, FMa
 	FScopedSubsystemRequirementsRestore SubsystemRestore(ExecutionContext);
 
 	ExecutionContext.SetSubsystemRequirements(*this);
-
-	// if there's a chunk collection set by the external code - use that
-	if (ExecutionContext.GetEntityCollection().IsSet())
+	if (ExecutionContext.CacheSubsystemRequirements(EntityManager.GetWorld(), *this))
 	{
-		const FMassArchetypeHandle& ArchetypeHandle = ExecutionContext.GetEntityCollection().GetArchetype();
-		// verify the archetype matches requirements
-		if (DoesArchetypeMatchRequirements(ArchetypeHandle) == false)
+		// if there's a chunk collection set by the external code - use that
+		if (ExecutionContext.GetEntityCollection().IsSet())
 		{
-			UE_VLOG_UELOG(EntityManager.GetOwner(), LogMass, Log, TEXT("Attempted to execute FMassEntityQuery with an incompatible Archetype: %s")
-				, *FMassDebugger::GetArchetypeRequirementCompatibilityDescription(*this, ArchetypeHandle));
+			const FMassArchetypeHandle& ArchetypeHandle = ExecutionContext.GetEntityCollection().GetArchetype();
+			// verify the archetype matches requirements
+			if (DoesArchetypeMatchRequirements(ArchetypeHandle) == false)
+			{
+				UE_VLOG_UELOG(EntityManager.GetOwner(), LogMass, Log, TEXT("Attempted to execute FMassEntityQuery with an incompatible Archetype: %s")
+					, *FMassDebugger::GetArchetypeRequirementCompatibilityDescription(*this, ArchetypeHandle));
 
 #if WITH_MASSENTITY_DEBUG
-			EntityManager.GetRequirementAccessDetector().ReleaseAccess(*this);
+				EntityManager.GetRequirementAccessDetector().ReleaseAccess(*this);
 #endif // WITH_MASSENTITY_DEBUG
-			return;
-		}
-		ExecutionContext.SetFragmentRequirements(*this);
-		
-		FMassArchetypeData& ArchetypeData = FMassArchetypeHelper::ArchetypeDataFromHandleChecked(ArchetypeHandle);
-		ArchetypeData.ExecuteFunction(ExecutionContext, ExecuteFunction
-			, GetFragmentMappingForArchetype(ArchetypeHandle)
-			, ExecutionContext.GetEntityCollection().GetRanges());
+				return;
+			}
+			ExecutionContext.SetFragmentRequirements(*this);
+
+			FMassArchetypeData& ArchetypeData = FMassArchetypeHelper::ArchetypeDataFromHandleChecked(ArchetypeHandle);
+			ArchetypeData.ExecuteFunction(ExecutionContext, ExecuteFunction
+				, GetFragmentMappingForArchetype(ArchetypeHandle)
+				, ExecutionContext.GetEntityCollection().GetRanges());
 #if WITH_MASSENTITY_DEBUG
-		NumEntitiesToProcess = ExecutionContext.GetNumEntities();
+			NumEntitiesToProcess = ExecutionContext.GetNumEntities();
 #endif
-	}
-	else
-	{
-		CacheArchetypes(EntityManager);
-		// it's important to set requirements after caching archetypes due to that call potentially sorting the requirements and the order is relevant here.
-		ExecutionContext.SetFragmentRequirements(*this);
-
-		for (int i = 0; i < ValidArchetypes.Num(); ++i)
+		}
+		else
 		{
-			const FMassArchetypeHandle& ArchetypeHandle = ValidArchetypes[i];
-			FMassArchetypeData& ArchetypeData = FMassArchetypeHelper::ArchetypeDataFromHandleChecked(ArchetypeHandle);
-			ArchetypeData.ExecuteFunction(ExecutionContext, ExecuteFunction, ArchetypeFragmentMapping[i], ChunkCondition);
-			ExecutionContext.ClearFragmentViews();
+			CacheArchetypes(EntityManager);
+			// it's important to set requirements after caching archetypes due to that call potentially sorting the requirements and the order is relevant here.
+			ExecutionContext.SetFragmentRequirements(*this);
+
+			for (int i = 0; i < ValidArchetypes.Num(); ++i)
+			{
+				const FMassArchetypeHandle& ArchetypeHandle = ValidArchetypes[i];
+				FMassArchetypeData& ArchetypeData = FMassArchetypeHelper::ArchetypeDataFromHandleChecked(ArchetypeHandle);
+				ArchetypeData.ExecuteFunction(ExecutionContext, ExecuteFunction, ArchetypeFragmentMapping[i], ChunkCondition);
+				ExecutionContext.ClearFragmentViews();
 #if WITH_MASSENTITY_DEBUG
-			NumEntitiesToProcess += ExecutionContext.GetNumEntities();
+				NumEntitiesToProcess += ExecutionContext.GetNumEntities();
 #endif
+			}
 		}
 	}
 
