Added CA_ASSUME(Node.Processor) in two places in MassProcessorDependencySolver, since MSVC static code analysed apparently cannot deduce that on its own.

diff --git a/Engine/Plugins/Runtime/MassEntity/Source/MassEntity/Private/MassProcessorDependencySolver.cpp b/Engine/Plugins/Runtime/MassEntity/Source/MassEntity/Private/MassProcessorDependencySolver.cpp
--- a/Engine/Plugins/Runtime/MassEntity/Source/MassEntity/Private/MassProcessorDependencySolver.cpp
+++ b/Engine/Plugins/Runtime/MassEntity/Source/MassEntity/Private/MassProcessorDependencySolver.cpp
@@ -565,6 +565,7 @@ void FMassProcessorDependencySolver::LogNode(const FNode& Node, int Indent)
 	}
 	else
 	{
+		CA_ASSUME(Node.Processor); // as implied by Node.IsGroup() == false
 		UE_LOG(LogMass, Log, TEXT("%*s%s before:%s after:%s"), Indent, TEXT(""), *Node.Name.ToString()
 			, *NameViewToString(Node.Processor->GetExecutionOrder().ExecuteBefore)
 			, *NameViewToString(Node.Processor->GetExecutionOrder().ExecuteAfter));
@@ -1119,6 +1120,8 @@ void FMassProcessorDependencySolver::ResolveDependencies(TArray<FMassProcessorOr
 	{
 		if (Node.IsGroup() == false)
 		{
+			CA_ASSUME(Node.Processor); // as implied by Node.IsGroup() == false
+
 			// for each processor-representing node we cache information on which archetypes among the once we've created 
 			// above (see the EntityManager.CreateArchetype call in the previous loop) match this processor. 
 			Node.Processor->GetArchetypesMatchingOwnedQueries(*EntityManager.Get(), Node.ValidArchetypes);
@@ -1126,7 +1129,6 @@ void FMassProcessorDependencySolver::ResolveDependencies(TArray<FMassProcessorOr
 			// prune the archetype-less processors
 			if (Node.ValidArchetypes.Num() == 0 && Node.Processor->ShouldAllowQueryBasedPruning(bGameRuntime))
 			{
-				CA_ASSUME(Node.Processor);
 				UE_LOG(LogMass, Verbose, TEXT("\t%s"), *Node.Processor->GetName());
 
 				if (InOutOptionalResult)
