Fixed System.Mass.Serialization.BitSet.OverrideExisting test

diff --git a/Engine/Plugins/Runtime/MassEntity/Source/MassEntityTestSuite/Private/MassSerializationTest.cpp b/Engine/Plugins/Runtime/MassEntity/Source/MassEntityTestSuite/Private/MassSerializationTest.cpp
--- a/Engine/Plugins/Runtime/MassEntity/Source/MassEntityTestSuite/Private/MassSerializationTest.cpp
+++ b/Engine/Plugins/Runtime/MassEntity/Source/MassEntityTestSuite/Private/MassSerializationTest.cpp
@@ -117,13 +117,6 @@ struct FSerialization_BitSetOverride : FExecutionTestBase
 {
 	virtual bool InstantTest() override
 	{
-		FTestStructTracker TestTracker;
-		const int32 FloatIndex = TestTracker.Add<FTestFragment_Float>();
-		const int32 IntIndex = TestTracker.Add<FTestFragment_Int>();
-		FStructTypeBitSet::FBitSetContainer TestBitArray;
-		TestBitArray.AddAtIndex(FloatIndex);
-		TestBitArray.AddAtIndex(IntIndex);
-
 		// create a bitset we're expected to get while loading TestBitArray's contents into FMassFragmentBitSet
 		FMassFragmentBitSet RegularBitSet;
 		RegularBitSet.Add<FTestFragment_Float>();;
@@ -132,7 +125,7 @@ struct FSerialization_BitSetOverride : FExecutionTestBase
 		// save
 		TArray<uint8> Data;
 		FMemoryWriter Writer(Data);
-		TestTracker.Serialize(Writer, TestBitArray);
+		RegularBitSet.Serialize(Writer);
 
 		// load
 		FMemoryReader Reader(Data);
