Change FMassBatchedEntityCommand to use UE_MT_SCOPED_WRITE_ACCESS in its Add methods.

diff --git a/Engine/Plugins/Runtime/MassEntity/Source/MassEntity/Public/MassCommands.h b/Engine/Plugins/Runtime/MassEntity/Source/MassEntity/Public/MassCommands.h
--- a/Engine/Plugins/Runtime/MassEntity/Source/MassEntity/Public/MassCommands.h
+++ b/Engine/Plugins/Runtime/MassEntity/Source/MassEntity/Public/MassCommands.h
@@ -139,14 +139,14 @@ struct FMassBatchedEntityCommand : public FMassBatchedCommand
 
 	void Add(FMassEntityHandle Entity)
 	{
-		UE_MT_SCOPED_READ_ACCESS(EntitiesAccessDetector);
+		UE_MT_SCOPED_WRITE_ACCESS(EntitiesAccessDetector);
 		TargetEntities.Add(Entity);
 		bHasWork = true;
 	}
 
 	void Add(TConstArrayView<FMassEntityHandle> Entities)
 	{
-		UE_MT_SCOPED_READ_ACCESS(EntitiesAccessDetector);
+		UE_MT_SCOPED_WRITE_ACCESS(EntitiesAccessDetector);
 		TargetEntities.Append(Entities.GetData(), Entities.Num());
 		bHasWork = true;
 	}
