Fixed Mass observer processors not respecting ExecutionFlags.

diff --git a/Engine/Plugins/Runtime/MassEntity/Source/MassEntity/Private/MassObserverManager.cpp b/Engine/Plugins/Runtime/MassEntity/Source/MassEntity/Private/MassObserverManager.cpp
--- a/Engine/Plugins/Runtime/MassEntity/Source/MassEntity/Private/MassObserverManager.cpp
+++ b/Engine/Plugins/Runtime/MassEntity/Source/MassEntity/Private/MassObserverManager.cpp
@@ -18,6 +18,8 @@ void SetUpObservers(FMassEntityManager& EntityManager, const TMap<TPointerType,
 	ObservedBitSet.Reset();
 	UObject* Owner = EntityManager.GetOwner();
 	check(Owner);
+	const UWorld* World = Owner->GetWorld();
+	const EProcessorExecutionFlags WorldExecutionFlags = World ? UE::Mass::Utils::GetProcessorExecutionFlagsForWold(*World) : EProcessorExecutionFlags::All;
 
 	for (auto It : RegisteredObserverTypes)
 	{
@@ -31,7 +33,10 @@ void SetUpObservers(FMassEntityManager& EntityManager, const TMap<TPointerType,
 
 		for (const TSubclassOf<UMassProcessor>& ProcessorClass : It.Value.ClassCollection)
 		{
-			Pipeline.AppendProcessor(ProcessorClass, *Owner);
+			if (ProcessorClass->GetDefaultObject<UMassProcessor>()->ShouldExecute(WorldExecutionFlags))
+			{
+				Pipeline.AppendProcessor(ProcessorClass, *Owner);
+			}
 		}
 		Pipeline.Initialize(*Owner);
 	}
