Fixed a loophole that made it possible to run a MassEntityQuery while its cached data was outdated.
Also:
* added subsystem caching as part of FMassEntityQuery::ForEachEntityChunk execution, and an early out when the requirements are not met.
* added a check verifying the data is up to date while fetching cached requirement mappings.
* renamed GetFragmentMappingForArchetype to GetRequirementsMappingForArchetype to better match what the function does
* changed FMassExecutionContext::CacheSubsystemRequirements to not require UWorld parameter (it's not needed and not being used).

diff --git a/Engine/Plugins/Runtime/MassEntity/Source/MassEntity/Private/MassEntityQuery.cpp b/Engine/Plugins/Runtime/MassEntity/Source/MassEntity/Private/MassEntityQuery.cpp
--- a/Engine/Plugins/Runtime/MassEntity/Source/MassEntity/Private/MassEntityQuery.cpp
+++ b/Engine/Plugins/Runtime/MassEntity/Source/MassEntity/Private/MassEntityQuery.cpp
@@ -176,50 +176,58 @@ void FMassEntityQuery::ForEachEntityChunk(FMassEntityManager& EntityManager, FMa
 	};
 	FScopedSubsystemRequirementsRestore SubsystemRestore(ExecutionContext);
 
-	ExecutionContext.SetSubsystemRequirements(*this);
-	if (ExecutionContext.CacheSubsystemRequirements(EntityManager.GetWorld(), *this))
+	if (ExecutionContext.CacheSubsystemRequirements(*this) == false)
 	{
-		// if there's a chunk collection set by the external code - use that
-		if (ExecutionContext.GetEntityCollection().IsSet())
+		// required subsystems are not available, bail out.
+		return;
+	}
+
+	// if there's a chunk collection set by the external code - use that
+	if (ExecutionContext.GetEntityCollection().IsSet())
+	{
+		const FMassArchetypeHandle& ArchetypeHandle = ExecutionContext.GetEntityCollection().GetArchetype();
+		// verify the archetype matches requirements
+		if (DoesArchetypeMatchRequirements(ArchetypeHandle) == false)
 		{
-			const FMassArchetypeHandle& ArchetypeHandle = ExecutionContext.GetEntityCollection().GetArchetype();
-			// verify the archetype matches requirements
-			if (DoesArchetypeMatchRequirements(ArchetypeHandle) == false)
-			{
-				UE_VLOG_UELOG(EntityManager.GetOwner(), LogMass, Log, TEXT("Attempted to execute FMassEntityQuery with an incompatible Archetype: %s")
-					, *FMassDebugger::GetArchetypeRequirementCompatibilityDescription(*this, ArchetypeHandle));
+			UE_VLOG_UELOG(EntityManager.GetOwner(), LogMass, Log, TEXT("Attempted to execute FMassEntityQuery with an incompatible Archetype: %s")
+				, *FMassDebugger::GetArchetypeRequirementCompatibilityDescription(*this, ArchetypeHandle));
 
 #if WITH_MASSENTITY_DEBUG
-				EntityManager.GetRequirementAccessDetector().ReleaseAccess(*this);
+			EntityManager.GetRequirementAccessDetector().ReleaseAccess(*this);
 #endif // WITH_MASSENTITY_DEBUG
-				return;
-			}
-			ExecutionContext.SetFragmentRequirements(*this);
-
-			FMassArchetypeData& ArchetypeData = FMassArchetypeHelper::ArchetypeDataFromHandleChecked(ArchetypeHandle);
-			ArchetypeData.ExecuteFunction(ExecutionContext, ExecuteFunction
-				, GetFragmentMappingForArchetype(ArchetypeHandle)
-				, ExecutionContext.GetEntityCollection().GetRanges());
+			return;
+		}
+		
+		// note that the following function will only resort to verifying that the data is up to date by
+		// checking the version number. In rare cases when it would result in non trivial cost we actually
+		// do need those calculations.
+		CacheArchetypes(EntityManager);
+
+		ExecutionContext.SetFragmentRequirements(*this);
+		
+		FMassArchetypeData& ArchetypeData = FMassArchetypeHelper::ArchetypeDataFromHandleChecked(ArchetypeHandle);
+		ArchetypeData.ExecuteFunction(ExecutionContext, ExecuteFunction
+			, GetRequirementsMappingForArchetype(ArchetypeHandle)
+			, ExecutionContext.GetEntityCollection().GetRanges());
 #if WITH_MASSENTITY_DEBUG
-			NumEntitiesToProcess = ExecutionContext.GetNumEntities();
+		NumEntitiesToProcess = ExecutionContext.GetNumEntities();
 #endif
-		}
-		else
-		{
-			CacheArchetypes(EntityManager);
-			// it's important to set requirements after caching archetypes due to that call potentially sorting the requirements and the order is relevant here.
-			ExecutionContext.SetFragmentRequirements(*this);
+	}
+	else
+	{
+		CacheArchetypes(EntityManager);
+		// it's important to set requirements after caching archetypes due to that call potentially sorting the requirements and the order is relevant here.
+		ExecutionContext.SetFragmentRequirements(*this);
 
-			for (int i = 0; i < ValidArchetypes.Num(); ++i)
-			{
-				const FMassArchetypeHandle& ArchetypeHandle = ValidArchetypes[i];
-				FMassArchetypeData& ArchetypeData = FMassArchetypeHelper::ArchetypeDataFromHandleChecked(ArchetypeHandle);
-				ArchetypeData.ExecuteFunction(ExecutionContext, ExecuteFunction, ArchetypeFragmentMapping[i], ChunkCondition);
-				ExecutionContext.ClearFragmentViews();
+		for (int i = 0; i < ValidArchetypes.Num(); ++i)
+		{
+			const FMassArchetypeHandle& ArchetypeHandle = ValidArchetypes[i];
+			FMassArchetypeData& ArchetypeData = FMassArchetypeHelper::ArchetypeDataFromHandleChecked(ArchetypeHandle);
+			ArchetypeData.ExecuteFunction(ExecutionContext, ExecuteFunction, ArchetypeFragmentMapping[i], ChunkCondition);
+			ExecutionContext.ClearFragmentViews();
 #if WITH_MASSENTITY_DEBUG
-				NumEntitiesToProcess += ExecutionContext.GetNumEntities();
+			NumEntitiesToProcess += ExecutionContext.GetNumEntities();
 #endif
-			}
 		}
 	}
 
@@ -264,9 +272,10 @@ bool FMassEntityQuery::HasMatchingEntities(FMassEntityManager& InEntityManager)
 	return false;
 }
 
-const FMassQueryRequirementIndicesMapping& FMassEntityQuery::GetFragmentMappingForArchetype(const FMassArchetypeHandle ArchetypeHandle) const
+const FMassQueryRequirementIndicesMapping& FMassEntityQuery::GetRequirementsMappingForArchetype(const FMassArchetypeHandle ArchetypeHandle) const
 {
 	static const FMassQueryRequirementIndicesMapping FallbackEmptyMapping;
+	checkf(IncrementalChangesCount == 0, TEXT("Fetching cached fragments mapping while the query's cached data is out of sync!"));
 	const int32 ArchetypeIndex = ValidArchetypes.Find(ArchetypeHandle);
 	return ArchetypeIndex != INDEX_NONE ? ArchetypeFragmentMapping[ArchetypeIndex] : FallbackEmptyMapping;
 }
diff --git a/Engine/Plugins/Runtime/MassEntity/Source/MassEntity/Private/MassExecutionContext.cpp b/Engine/Plugins/Runtime/MassEntity/Source/MassEntity/Private/MassExecutionContext.cpp
--- a/Engine/Plugins/Runtime/MassEntity/Source/MassEntity/Private/MassExecutionContext.cpp
+++ b/Engine/Plugins/Runtime/MassEntity/Source/MassEntity/Private/MassExecutionContext.cpp
@@ -33,9 +33,8 @@ void FMassExecutionContext::ClearExecutionData()
 	ChunkSerialModificationNumber = -1;
 }
 
-bool FMassExecutionContext::CacheSubsystemRequirements(const UWorld* World, const FMassSubsystemRequirements& SubsystemRequirements)
+bool FMassExecutionContext::CacheSubsystemRequirements(const FMassSubsystemRequirements& SubsystemRequirements)
 {
-	check(EntityManager->GetWorld() == World);
 	return SubsystemAccess.CacheSubsystemRequirements(SubsystemRequirements);
 }
 
diff --git a/Engine/Plugins/Runtime/MassEntity/Source/MassEntity/Private/MassProcessor.cpp b/Engine/Plugins/Runtime/MassEntity/Source/MassEntity/Private/MassProcessor.cpp
--- a/Engine/Plugins/Runtime/MassEntity/Source/MassEntity/Private/MassProcessor.cpp
+++ b/Engine/Plugins/Runtime/MassEntity/Source/MassEntity/Private/MassProcessor.cpp
@@ -208,7 +208,7 @@ void UMassProcessor::CallExecute(FMassEntityManager& EntityManager, FMassExecuti
 	// meaning if it fails there's no point in calling Execute.
 	// Note that we're not testing individual queries in OwnedQueries - processors can function just fine with some 
 	// of their queries not having anything to do.
-	if (Context.CacheSubsystemRequirements(GetWorld(), ProcessorRequirements))
+	if (Context.CacheSubsystemRequirements(ProcessorRequirements))
 	{
 		Execute(EntityManager, Context);
 	}
diff --git a/Engine/Plugins/Runtime/MassEntity/Source/MassEntity/Private/MassSubsystemAccess.cpp b/Engine/Plugins/Runtime/MassEntity/Source/MassEntity/Private/MassSubsystemAccess.cpp
--- a/Engine/Plugins/Runtime/MassEntity/Source/MassEntity/Private/MassSubsystemAccess.cpp
+++ b/Engine/Plugins/Runtime/MassEntity/Source/MassEntity/Private/MassSubsystemAccess.cpp
@@ -58,21 +58,19 @@ USubsystem* FMassSubsystemAccess::FetchSubsystemInstance(UWorld* World, TSubclas
 
 bool FMassSubsystemAccess::CacheSubsystemRequirements(const FMassSubsystemRequirements& SubsystemRequirements)
 {
-	if (SubsystemRequirements.IsEmpty())
-	{
-		return true;
-	}
-
 	bool bResult = true;
 
-	for (FMassExternalSubsystemBitSet::FIndexIterator It = SubsystemRequirements.GetRequiredConstSubsystems().GetIndexIterator(); It && bResult; ++It)
-	{
-		bResult = bResult && CacheSubsystem(*It);
-	}
-
-	for (FMassExternalSubsystemBitSet::FIndexIterator It = SubsystemRequirements.GetRequiredMutableSubsystems().GetIndexIterator(); It && bResult; ++It)
+	if (SubsystemRequirements.IsEmpty() == false)
 	{
-		bResult = bResult && CacheSubsystem(*It);
+		for (FMassExternalSubsystemBitSet::FIndexIterator It = SubsystemRequirements.GetRequiredConstSubsystems().GetIndexIterator(); It && bResult; ++It)
+		{
+			bResult = bResult && CacheSubsystem(*It);
+		}
+
+		for (FMassExternalSubsystemBitSet::FIndexIterator It = SubsystemRequirements.GetRequiredMutableSubsystems().GetIndexIterator(); It && bResult; ++It)
+		{
+			bResult = bResult && CacheSubsystem(*It);
+		}
 	}
 
 	if (bResult)
diff --git a/Engine/Plugins/Runtime/MassEntity/Source/MassEntity/Public/MassEntityQuery.h b/Engine/Plugins/Runtime/MassEntity/Source/MassEntity/Public/MassEntityQuery.h
--- a/Engine/Plugins/Runtime/MassEntity/Source/MassEntity/Public/MassEntityQuery.h
+++ b/Engine/Plugins/Runtime/MassEntity/Source/MassEntity/Public/MassEntityQuery.h
@@ -102,7 +102,7 @@ public:
 	 * If ArchetypeHandle is among ValidArchetypes then the function retrieves requirements mapping cached for it,
 	 * otherwise an empty mapping will be returned (and the requirements binding will be done the slow way).
 	 */
-	const FMassQueryRequirementIndicesMapping& GetFragmentMappingForArchetype(const FMassArchetypeHandle ArchetypeHandle) const;
+	const FMassQueryRequirementIndicesMapping& GetRequirementsMappingForArchetype(const FMassArchetypeHandle ArchetypeHandle) const;
 
 	void ExportRequirements(FMassExecutionRequirements& OutRequirements) const;
 
diff --git a/Engine/Plugins/Runtime/MassEntity/Source/MassEntity/Public/MassExecutionContext.h b/Engine/Plugins/Runtime/MassEntity/Source/MassEntity/Public/MassExecutionContext.h
--- a/Engine/Plugins/Runtime/MassEntity/Source/MassEntity/Public/MassExecutionContext.h
+++ b/Engine/Plugins/Runtime/MassEntity/Source/MassEntity/Public/MassExecutionContext.h
@@ -323,16 +323,15 @@ public:
 	}
 
 	/** 
-	 * Processes SubsystemRequirements to fetch and process all the indicated subsystems. 
-	 * 
-	 * @param World used to fetch any UWorld-related subsystems. Note that it's perfectly fine to pass a Null World. 
-	 *	The worst case scenario is that any of the UWorldSubsystem subclasses requirements won't be met.
-	 * @oaram SubsystemRequirements indicates all the subsystems that are expected to be accessed. Requesting a subsystem 
+	 * Processes SubsystemRequirements to fetch and cache all the indicated subsystems. If a UWorld is required to fetch
+	 * a specific subsystem then the one associated with the stored EntityManager will be used.
+	 *
+	 * @param SubsystemRequirements indicates all the subsystems that are expected to be accessed. Requesting a subsystem 
 	 *	not indicated by the SubsystemRequirements will result in a failure.
 	 * 
-	 * @return `true` if all required subsystems have been found, `false` otherwise
+	 * @return `true` if all required subsystems have been found, `false` otherwise.
 	 */
-	bool CacheSubsystemRequirements(const UWorld* World, const FMassSubsystemRequirements& SubsystemRequirements);
+	bool CacheSubsystemRequirements(const FMassSubsystemRequirements& SubsystemRequirements);
 
 protected:
 	void SetSubsystemRequirements(const FMassSubsystemRequirements& SubsystemRequirements)
@@ -421,4 +420,10 @@ public:
 	{
 		return SubsystemAccess.GetSubsystemChecked<T>(SubsystemClass);
 	}
+
+	UE_DEPRECATED(5.2, "This version of CacheSubsystemRequirements is deprecated. Use the one without the UWorld parameter")
+	bool CacheSubsystemRequirements(const UWorld*, const FMassSubsystemRequirements& SubsystemRequirements)
+	{
+		return CacheSubsystemRequirements(SubsystemRequirements);
+	}
 };
