Extended the comment on UMassProcessor.ProcessingPhase to describe better how to set it.

diff --git a/Engine/Plugins/Runtime/MassEntity/Source/MassEntity/Public/MassProcessor.h b/Engine/Plugins/Runtime/MassEntity/Source/MassEntity/Public/MassProcessor.h
--- a/Engine/Plugins/Runtime/MassEntity/Source/MassEntity/Public/MassProcessor.h
+++ b/Engine/Plugins/Runtime/MassEntity/Source/MassEntity/Public/MassProcessor.h
@@ -122,7 +122,9 @@ protected:
 	UPROPERTY(EditAnywhere, Category = "Pipeline", meta = (Bitmask, BitmaskEnum = "/Script/MassEntity.EProcessorExecutionFlags"), config)
 	int32 ExecutionFlags;
 
-	/** Processing phase this processor will be automatically run as part of. */
+	/** Processing phase this processor will be automatically run as part of. Needs to be set before the processor gets
+	 *  registered with MassProcessingPhaseManager, otherwise it will have no effect. This property is usually read via
+	 *  a given class's CDO, so it's recommended to set it in the constructor. */
 	UPROPERTY(EditDefaultsOnly, Category = Processor, config)
 	EMassProcessingPhase ProcessingPhase = EMassProcessingPhase::PrePhysics;
 
