PR #10255: FIX a Mass bug (Contributed by [REDACTED_PERSON])

diff --git a/Engine/Plugins/Runtime/MassEntity/Source/MassEntity/Private/MassEntityManager.cpp b/Engine/Plugins/Runtime/MassEntity/Source/MassEntity/Private/MassEntityManager.cpp
--- a/Engine/Plugins/Runtime/MassEntity/Source/MassEntity/Private/MassEntityManager.cpp
+++ b/Engine/Plugins/Runtime/MassEntity/Source/MassEntity/Private/MassEntityManager.cpp
@@ -974,39 +974,42 @@ void FMassEntityManager::BatchChangeFragmentCompositionForEntities(TConstArrayVi
 			? (CurrentArchetype->GetFragmentBitSet() + FragmentsToAdd - FragmentsToRemove)
 			: (FragmentsToAdd - FragmentsToRemove);
 
-		if (CurrentArchetype && CurrentArchetype->GetFragmentBitSet() != NewFragmentComposition)
+		if (CurrentArchetype)
 		{
-			FMassFragmentBitSet FragmentsAdded = FragmentsToAdd - CurrentArchetype->GetFragmentBitSet();
-			const bool bFragmentsAddedAreObserved = ObserverManager.HasObserversForBitSet(FragmentsAdded, EMassObservedOperation::Add);
-			FMassFragmentBitSet FragmentsRemoved = FragmentsToRemove.GetOverlap(CurrentArchetype->GetFragmentBitSet());
-			if (ObserverManager.HasObserversForBitSet(FragmentsRemoved, EMassObservedOperation::Remove))
+			if (CurrentArchetype->GetFragmentBitSet() != NewFragmentComposition)
 			{
-				ObserverManager.OnCompositionChanged(Collection, FMassArchetypeCompositionDescriptor(MoveTemp(FragmentsRemoved)), EMassObservedOperation::Remove);
-			}
+				FMassFragmentBitSet FragmentsAdded = FragmentsToAdd - CurrentArchetype->GetFragmentBitSet();
+				const bool bFragmentsAddedAreObserved = ObserverManager.HasObserversForBitSet(FragmentsAdded, EMassObservedOperation::Add);
+				FMassFragmentBitSet FragmentsRemoved = FragmentsToRemove.GetOverlap(CurrentArchetype->GetFragmentBitSet());
+				if (ObserverManager.HasObserversForBitSet(FragmentsRemoved, EMassObservedOperation::Remove))
+				{
+					ObserverManager.OnCompositionChanged(Collection, FMassArchetypeCompositionDescriptor(MoveTemp(FragmentsRemoved)), EMassObservedOperation::Remove);
+				}
 
-			FMassArchetypeHandle NewArchetypeHandle = InternalCreateSimilarArchetype(Collection.GetArchetype().DataPtr, NewFragmentComposition);
-			checkSlow(NewArchetypeHandle.IsValid());
+				FMassArchetypeHandle NewArchetypeHandle = InternalCreateSimilarArchetype(Collection.GetArchetype().DataPtr, NewFragmentComposition);
+				checkSlow(NewArchetypeHandle.IsValid());
 
-			// Move the entity over
-			FMassArchetypeEntityCollection::FEntityRangeArray NewArchetypeEntityRanges;
-			TArray<FMassEntityHandle> EntitesBeingMoved;
-			CurrentArchetype->BatchMoveEntitiesToAnotherArchetype(Collection, *NewArchetypeHandle.DataPtr.Get(), EntitesBeingMoved
-				, bFragmentsAddedAreObserved ? &NewArchetypeEntityRanges : nullptr);
+				// Move the entity over
+				FMassArchetypeEntityCollection::FEntityRangeArray NewArchetypeEntityRanges;
+				TArray<FMassEntityHandle> EntitesBeingMoved;
+				CurrentArchetype->BatchMoveEntitiesToAnotherArchetype(Collection, *NewArchetypeHandle.DataPtr.Get(), EntitesBeingMoved
+					, bFragmentsAddedAreObserved ? &NewArchetypeEntityRanges : nullptr);
 
-			for (const FMassEntityHandle& Entity : EntitesBeingMoved)
-			{
-				check(Entities.IsValidIndex(Entity.Index));
+				for (const FMassEntityHandle& Entity : EntitesBeingMoved)
+				{
+					check(Entities.IsValidIndex(Entity.Index));
 
-				FEntityData& EntityData = Entities[Entity.Index];
-				EntityData.CurrentArchetype = NewArchetypeHandle.DataPtr;
-			}
+					FEntityData& EntityData = Entities[Entity.Index];
+					EntityData.CurrentArchetype = NewArchetypeHandle.DataPtr;
+				}
 
-			if (bFragmentsAddedAreObserved)
-			{
-				ObserverManager.OnCompositionChanged(
-					FMassArchetypeEntityCollection(NewArchetypeHandle, MoveTemp(NewArchetypeEntityRanges))
-					, FMassArchetypeCompositionDescriptor(MoveTemp(FragmentsAdded))
-					, EMassObservedOperation::Add);
+				if (bFragmentsAddedAreObserved)
+				{
+					ObserverManager.OnCompositionChanged(
+						FMassArchetypeEntityCollection(NewArchetypeHandle, MoveTemp(NewArchetypeEntityRanges))
+						, FMassArchetypeCompositionDescriptor(MoveTemp(FragmentsAdded))
+						, EMassObservedOperation::Add);
+				}
 			}
 		}
 		else
