PR #10270: Fix ineffective chunk filter for Mass entity queries (Contributed by [REDACTED_PERSON])

diff --git a/Engine/Plugins/Runtime/MassEntity/Source/MassEntity/Private/MassArchetypeData.cpp b/Engine/Plugins/Runtime/MassEntity/Source/MassEntity/Private/MassArchetypeData.cpp
--- a/Engine/Plugins/Runtime/MassEntity/Source/MassEntity/Private/MassArchetypeData.cpp
+++ b/Engine/Plugins/Runtime/MassEntity/Source/MassEntity/Private/MassArchetypeData.cpp
@@ -376,7 +376,7 @@ void FMassArchetypeData::MoveEntityToAnotherArchetype(const FMassEntityHandle En
 	RemoveEntityInternal(AbsoluteIndex);
 }
 
-void FMassArchetypeData::ExecuteFunction(FMassExecutionContext& RunContext, const FMassExecuteFunction& Function, const FMassQueryRequirementIndicesMapping& RequirementMapping, FMassArchetypeEntityCollection::FConstEntityRangeArrayView EntityRangeContainer)
+void FMassArchetypeData::ExecuteFunction(FMassExecutionContext& RunContext, const FMassExecuteFunction& Function, const FMassQueryRequirementIndicesMapping& RequirementMapping, FMassArchetypeEntityCollection::FConstEntityRangeArrayView EntityRangeContainer, const FMassChunkConditionFunction& ChunkCondition)
 {
 	if (GetNumEntities() == 0)
 	{
@@ -406,9 +406,12 @@ void FMassArchetypeData::ExecuteFunction(FMassExecutionContext& RunContext, cons
 
 			RunContext.SetCurrentChunkSerialModificationNumber(Chunk.GetSerialModificationNumber());
 			BindChunkFragmentRequirements(RunContext, RequirementMapping.ChunkFragments, Chunk);
-			BindEntityRequirements(RunContext, RequirementMapping.EntityFragments, Chunk, ChunkIterator->SubchunkStart, ChunkLength);
 
-			Function(RunContext);
+			if (!ChunkCondition || ChunkCondition(RunContext))
+			{
+				BindEntityRequirements(RunContext, RequirementMapping.EntityFragments, Chunk, ChunkIterator->SubchunkStart, ChunkLength);
+				Function(RunContext);
+			}
 		}
 	}
 }
diff --git a/Engine/Plugins/Runtime/MassEntity/Source/MassEntity/Private/MassArchetypeData.h b/Engine/Plugins/Runtime/MassEntity/Source/MassEntity/Private/MassArchetypeData.h
--- a/Engine/Plugins/Runtime/MassEntity/Source/MassEntity/Private/MassArchetypeData.h
+++ b/Engine/Plugins/Runtime/MassEntity/Source/MassEntity/Private/MassArchetypeData.h
@@ -238,7 +238,7 @@ public:
 
 	uint32 GetCreatedArchetypeDataVersion() const { return CreatedArchetypeDataVersion; }
 
-	void ExecuteFunction(FMassExecutionContext& RunContext, const FMassExecuteFunction& Function, const FMassQueryRequirementIndicesMapping& RequirementMapping, FMassArchetypeEntityCollection::FConstEntityRangeArrayView EntityRangeContainer);
+	void ExecuteFunction(FMassExecutionContext& RunContext, const FMassExecuteFunction& Function, const FMassQueryRequirementIndicesMapping& RequirementMapping, FMassArchetypeEntityCollection::FConstEntityRangeArrayView EntityRangeContainer, const FMassChunkConditionFunction& ChunkCondition);
 	void ExecuteFunction(FMassExecutionContext& RunContext, const FMassExecuteFunction& Function, const FMassQueryRequirementIndicesMapping& RequirementMapping, const FMassChunkConditionFunction& ChunkCondition);
 
 	void ExecutionFunctionForChunk(FMassExecutionContext RunContext, const FMassExecuteFunction& Function, const FMassQueryRequirementIndicesMapping& RequirementMapping, const FMassArchetypeEntityCollection::FArchetypeEntityRange& EntityRange, const FMassChunkConditionFunction& ChunkCondition = FMassChunkConditionFunction());
diff --git a/Engine/Plugins/Runtime/MassEntity/Source/MassEntity/Private/MassEntityQuery.cpp b/Engine/Plugins/Runtime/MassEntity/Source/MassEntity/Private/MassEntityQuery.cpp
--- a/Engine/Plugins/Runtime/MassEntity/Source/MassEntity/Private/MassEntityQuery.cpp
+++ b/Engine/Plugins/Runtime/MassEntity/Source/MassEntity/Private/MassEntityQuery.cpp
@@ -208,7 +208,8 @@ void FMassEntityQuery::ForEachEntityChunk(FMassEntityManager& EntityManager, FMa
 		FMassArchetypeData& ArchetypeData = FMassArchetypeHelper::ArchetypeDataFromHandleChecked(ArchetypeHandle);
 		ArchetypeData.ExecuteFunction(ExecutionContext, ExecuteFunction
 			, GetRequirementsMappingForArchetype(ArchetypeHandle)
-			, ExecutionContext.GetEntityCollection().GetRanges());
+			, ExecutionContext.GetEntityCollection().GetRanges()
+			, ChunkCondition);
 #if WITH_MASSENTITY_DEBUG
 		NumEntitiesToProcess = ExecutionContext.GetNumEntities();
 #endif
