Added UMassObserverProcessor.bAutoRegisterWithObserverRegistry, a flag that can be used to disable a given MassObserver class

diff --git a/Engine/Plugins/Runtime/MassEntity/Source/MassEntity/Private/MassObserverProcessor.cpp b/Engine/Plugins/Runtime/MassEntity/Source/MassEntity/Private/MassObserverProcessor.cpp
--- a/Engine/Plugins/Runtime/MassEntity/Source/MassEntity/Private/MassObserverProcessor.cpp
+++ b/Engine/Plugins/Runtime/MassEntity/Source/MassEntity/Private/MassObserverProcessor.cpp
@@ -39,8 +39,11 @@ void UMassObserverProcessor::PostInitProperties()
 
 void UMassObserverProcessor::Register()
 {
-	check(ObservedType);
-	UMassObserverRegistry::GetMutable().RegisterObserver(*ObservedType, Operation, GetClass());
+	if (bAutoRegisterWithObserverRegistry)
+	{
+		check(ObservedType);
+		UMassObserverRegistry::GetMutable().RegisterObserver(*ObservedType, Operation, GetClass());
+	}
 }
 
 
diff --git a/Engine/Plugins/Runtime/MassEntity/Source/MassEntity/Public/MassObserverProcessor.h b/Engine/Plugins/Runtime/MassEntity/Source/MassEntity/Public/MassObserverProcessor.h
--- a/Engine/Plugins/Runtime/MassEntity/Source/MassEntity/Public/MassObserverProcessor.h
+++ b/Engine/Plugins/Runtime/MassEntity/Source/MassEntity/Public/MassObserverProcessor.h
@@ -24,6 +24,9 @@ protected:
 	virtual void Register();
 
 protected:
+	UPROPERTY(EditDefaultsOnly, Category = Processor, config)
+	bool bAutoRegisterWithObserverRegistry = true;
+
 	/** Determines which Fragment or Tag type this given UMassObserverProcessor will be observing */
 	UPROPERTY()
 	TObjectPtr<UScriptStruct> ObservedType = nullptr;
