Fixed the broken System.Mass.Query.System.AccessAPI.
The test was failing on game-subsystems when the world was not a game world. This part of the test has been extracted into a separate test.
The test was also failing on a custom subsystem (i.e. not derived from any of the specialized engine-defined subsystem classes). This part of the test has been removed. Will be revisited once we have an actual use case for such a subsystem.

diff --git a/Engine/Plugins/Runtime/MassEntity/Source/MassEntityTestSuite/Private/MassSystemRequirementTest.cpp b/Engine/Plugins/Runtime/MassEntity/Source/MassEntityTestSuite/Private/MassSystemRequirementTest.cpp
--- a/Engine/Plugins/Runtime/MassEntity/Source/MassEntityTestSuite/Private/MassSystemRequirementTest.cpp
+++ b/Engine/Plugins/Runtime/MassEntity/Source/MassEntityTestSuite/Private/MassSystemRequirementTest.cpp
@@ -11,29 +11,6 @@
 
 UE_DISABLE_OPTIMIZATION_SHIP
 
-//////////////////////////////////////////////////////////////////////////
-// UMassTestCustomSubsystem 
-
-TWeakObjectPtr<UMassTestCustomSubsystem> UMassTestCustomSubsystem::Instance = nullptr;
-
-UMassTestCustomSubsystem* UMassTestCustomSubsystem::Create()
-{
-	ensure(Instance.IsValid() == false);
-	Instance = NewObject<UMassTestCustomSubsystem>();
-	return Instance.Get();
-}
-
-UMassTestCustomSubsystem* UMassTestCustomSubsystem::Get()
-{
-	return Instance.Get();
-}
-
-/** FMassSubsystemAccess::FetchSubsystemInstance specialization for the UMassTestCustomSubsystem */
-template<>
-UMassTestCustomSubsystem* FMassSubsystemAccess::FetchSubsystemInstance<UMassTestCustomSubsystem>()
-{
-	return UMassTestCustomSubsystem::Get();
-}
 
 //////////////////////////////////////////////////////////////////////////
 // Tests
@@ -51,9 +28,7 @@ struct FSystemRequirementTestBase : FEntityTestBase
 	}
 
 	virtual bool SetUp() override
-	{		
-		AddAutoDestroyObject(*UMassTestCustomSubsystem::Create());
-
+	{
 		if (Super::SetUp())
 		{
 			EntityManager->CreateEntity(FloatsArchetype);
@@ -68,7 +43,7 @@ struct FSystemRequirementTestBase : FEntityTestBase
 //////////////////////////////////////////////////////////////////////////
 // The main point of this specific test is to verify that we're getting the same answers from both
 // FMassEntityQuery and FMassSubsystemAccess.
-struct FSubsystemAccessAPI : FSystemRequirementTestBase
+struct FGenericSubsystemAccessAPI : FSystemRequirementTestBase
 {
 	virtual bool InstantTest() override
 	{
@@ -79,91 +54,115 @@ struct FSubsystemAccessAPI : FSystemRequirementTestBase
 
 		EntityQuery.AddSubsystemRequirement<UMassTestWorldSubsystem>(EMassFragmentAccess::ReadWrite);
 		EntityQuery.AddSubsystemRequirement<UMassTestEngineSubsystem>(EMassFragmentAccess::ReadWrite);
-		EntityQuery.AddSubsystemRequirement<UMassTestCustomSubsystem>(EMassFragmentAccess::ReadWrite);
-		EntityQuery.AddSubsystemRequirement<UMassTestLocalPlayerSubsystem>(EMassFragmentAccess::ReadWrite);
-		EntityQuery.AddSubsystemRequirement<UMassTestGameInstanceSubsystem>(EMassFragmentAccess::ReadWrite);
 
 		SubsystemAccessBits.SetSubsystemRequirementBits(EntityQuery.GetRequiredConstSubsystems(), EntityQuery.GetRequiredMutableSubsystems());
 		SubsystemAccessClasses.SetSubsystemRequirements(EntityQuery);
 
 		UMassTestWorldSubsystem* TestWorldSubsystemActual = World->GetSubsystem<UMassTestWorldSubsystem>();
 		UMassTestEngineSubsystem* TestEngineSubsystemActual = GEngine->GetEngineSubsystem<UMassTestEngineSubsystem>();
-		UMassTestCustomSubsystem* TestCustomSubsystemActual = UMassTestCustomSubsystem::Get();
-		ULocalPlayer* LocalPlayer = World->GetFirstLocalPlayerFromController();
-		UMassTestLocalPlayerSubsystem* TestLocalPlayerSubsystemActual = LocalPlayer ? LocalPlayer->GetSubsystem<UMassTestLocalPlayerSubsystem>() : nullptr;
-		UGameInstance* GameInstance = World->GetGameInstance();
-		UMassTestGameInstanceSubsystem* TestGameInstanceSubsystemActual = GameInstance ? GameInstance->GetSubsystem<UMassTestGameInstanceSubsystem>() : nullptr;
 
 		UMassTestWorldSubsystem* AccessWorldSubsystem = nullptr;
 		UMassTestEngineSubsystem* AccessEngineSubsystem = nullptr;
-		UMassTestCustomSubsystem* AccessCustomSubsystem = nullptr;
-		UMassTestLocalPlayerSubsystem* AccessLocalPlayerSubsystem = nullptr;
-		UMassTestGameInstanceSubsystem* AccessGameInstanceSubsystem = nullptr;
 
 		const UMassTestWorldSubsystem* ConstAccessWorldSubsystem = nullptr;
 		const UMassTestEngineSubsystem* ConstAccessEngineSubsystem = nullptr;
-		const UMassTestCustomSubsystem* ConstAccessCustomSubsystem = nullptr;
-		const UMassTestLocalPlayerSubsystem* ConstAccessLocalPlayerSubsystem = nullptr;
-		const UMassTestGameInstanceSubsystem* ConstAccessGameInstanceSubsystem = nullptr;
 
 		FMassExecutionContext ExecutionContext(*EntityManager);
 		EntityQuery.ForEachEntityChunk(*EntityManager, ExecutionContext, [&](FMassExecutionContext& Context)
 			{
 				AccessWorldSubsystem = Context.GetMutableSubsystem<UMassTestWorldSubsystem>();
 				AccessEngineSubsystem = Context.GetMutableSubsystem<UMassTestEngineSubsystem>();
-				AccessCustomSubsystem = Context.GetMutableSubsystem<UMassTestCustomSubsystem>();
-				AccessLocalPlayerSubsystem = Context.GetMutableSubsystem<UMassTestLocalPlayerSubsystem>();
-				AccessGameInstanceSubsystem = Context.GetMutableSubsystem<UMassTestGameInstanceSubsystem>();
 
 				ConstAccessWorldSubsystem = Context.GetSubsystem<UMassTestWorldSubsystem>();
 				ConstAccessEngineSubsystem = Context.GetSubsystem<UMassTestEngineSubsystem>();
-				ConstAccessCustomSubsystem = Context.GetSubsystem<UMassTestCustomSubsystem>();
-				ConstAccessLocalPlayerSubsystem = Context.GetSubsystem<UMassTestLocalPlayerSubsystem>();
-				ConstAccessGameInstanceSubsystem = Context.GetSubsystem<UMassTestGameInstanceSubsystem>();
 			});
 
 		AITEST_NOT_NULL(TEXT("WorldSubsystem: Subsystem Actual is expected to be not NULL"), TestWorldSubsystemActual);
-		AITEST_EQUAL(TEXT("WorldSubsystem: Mutable subsystem fetched is expected to be the same as the Actual"), TestWorldSubsystemActual, AccessWorldSubsystem);
-		AITEST_EQUAL(TEXT("WorldSubsystem: Const subsystem fetched is expected to be the same as the Actual"), TestWorldSubsystemActual, ConstAccessWorldSubsystem);
-		AITEST_EQUAL(TEXT("WorldSubsystem: subsystem fetched via bits-requirements is expected to be the same as the Actual")
+		AITEST_EQUAL(TEXT("WorldSubsystem: Mutable Subsystem fetched is expected to be the same as the Actual"), TestWorldSubsystemActual, AccessWorldSubsystem);
+		AITEST_EQUAL(TEXT("WorldSubsystem: Const Subsystem fetched is expected to be the same as the Actual"), TestWorldSubsystemActual, ConstAccessWorldSubsystem);
+		AITEST_EQUAL(TEXT("WorldSubsystem: Subsystem fetched via bits-requirements is expected to be the same as the Actual")
 			, TestWorldSubsystemActual, SubsystemAccessBits.GetSubsystem<UMassTestWorldSubsystem>());
-		AITEST_EQUAL(TEXT("WorldSubsystem: subsystem fetched via class-requirements is expected to be the same as the Actual")
+		AITEST_EQUAL(TEXT("WorldSubsystem: Subsystem fetched via class-requirements is expected to be the same as the Actual")
 			, TestWorldSubsystemActual, SubsystemAccessClasses.GetSubsystem<UMassTestWorldSubsystem>());
 
 		AITEST_NOT_NULL(TEXT("EngineSubsystem: Subsystem Actual is expected to be not NULL"), TestEngineSubsystemActual);
-		AITEST_EQUAL(TEXT("EngineSubsystem: Mutable subsystem fetched is expected to be the same as the Actual"), TestEngineSubsystemActual, AccessEngineSubsystem);
-		AITEST_EQUAL(TEXT("EngineSubsystem: Const subsystem fetched is expected to be the same as the Actual"), TestEngineSubsystemActual, ConstAccessEngineSubsystem);
-		AITEST_EQUAL(TEXT("EngineSubsystem: subsystem fetched via bits-requirements is expected to be the same as the Actual")
+		AITEST_EQUAL(TEXT("EngineSubsystem: Mutable Subsystem fetched is expected to be the same as the Actual"), TestEngineSubsystemActual, AccessEngineSubsystem);
+		AITEST_EQUAL(TEXT("EngineSubsystem: Const Subsystem fetched is expected to be the same as the Actual"), TestEngineSubsystemActual, ConstAccessEngineSubsystem);
+		AITEST_EQUAL(TEXT("EngineSubsystem: Subsystem fetched via bits-requirements is expected to be the same as the Actual")
 			, TestEngineSubsystemActual, SubsystemAccessBits.GetSubsystem<UMassTestEngineSubsystem>());
-		AITEST_EQUAL(TEXT("EngineSubsystem: subsystem fetched via class-requirements is expected to be the same as the Actual")
+		AITEST_EQUAL(TEXT("EngineSubsystem: Subsystem fetched via class-requirements is expected to be the same as the Actual")
 			, TestEngineSubsystemActual, SubsystemAccessClasses.GetSubsystem<UMassTestEngineSubsystem>());
 
-		AITEST_NOT_NULL(TEXT("CustomSubsystem: Subsystem Actual is expected to be not NULL"), TestCustomSubsystemActual);
-		AITEST_EQUAL(TEXT("CustomSubsystem: Mutable subsystem fetched is expected to be the same as the Actual"), TestCustomSubsystemActual, AccessCustomSubsystem);
-		AITEST_EQUAL(TEXT("CustomSubsystem: Const subsystem fetched is expected to be the same as the Actual"), TestCustomSubsystemActual, ConstAccessCustomSubsystem);
-		AITEST_EQUAL(TEXT("CustomSubsystem: subsystem fetched via bits-requirements is expected to be the same as the Actual")
-			, TestCustomSubsystemActual, SubsystemAccessBits.GetSubsystem<UMassTestCustomSubsystem>());
-		AITEST_EQUAL(TEXT("CustomSubsystem: subsystem fetched via class-requirements is expected to be the same as the Actual")
-			, TestCustomSubsystemActual, SubsystemAccessClasses.GetSubsystem<UMassTestCustomSubsystem>());
-
-		AITEST_EQUAL(TEXT("LocalPlayerSubsystem: Mutable subsystem fetched is expected to be the same as the Actual"), TestLocalPlayerSubsystemActual, AccessLocalPlayerSubsystem);
-		AITEST_EQUAL(TEXT("LocalPlayerSubsystem: Const subsystem fetched is expected to be the same as the Actual"), TestLocalPlayerSubsystemActual, ConstAccessLocalPlayerSubsystem);
-		AITEST_EQUAL(TEXT("LocalPlayerSubsystem: subsystem fetched via bits-requirements is expected to be the same as the Actual")
+		return true;
+	}
+};
+IMPLEMENT_AI_INSTANT_TEST(FGenericSubsystemAccessAPI, "System.Mass.Query.System.AccessAPI.Generic");
+
+struct FGameSubsystemAccessAPI : FSystemRequirementTestBase
+{
+	virtual bool InstantTest() override
+	{
+		UWorld* World = FAITestHelpers::GetWorld();
+		check(World);
+
+		const ULocalPlayer* LocalPlayer = World->GetFirstLocalPlayerFromController();
+		UMassTestLocalPlayerSubsystem* TestLocalPlayerSubsystemActual = LocalPlayer ? LocalPlayer->GetSubsystem<UMassTestLocalPlayerSubsystem>() : nullptr;
+		const UGameInstance* GameInstance = World->GetGameInstance();
+		UMassTestGameInstanceSubsystem* TestGameInstanceSubsystemActual = GameInstance ? GameInstance->GetSubsystem<UMassTestGameInstanceSubsystem>() : nullptr;
+
+		if (World->IsGameWorld() == false)
+		{
+			// no point in running the full test since none of the used system classes should produce any instances
+			// just verify this assumption
+			AITEST_NULL(TEXT("LocalPlayerSubsystem subclass instance is expected to be NULL in non-game worlds"), TestLocalPlayerSubsystemActual);
+			AITEST_NULL(TEXT("GameInstanceSubsystem subclass instance is expected to be NULL in non-game worlds"), TestGameInstanceSubsystemActual);
+
+			return true;
+		}
+
+		FMassSubsystemAccess SubsystemAccessBits(World);
+		FMassSubsystemAccess SubsystemAccessClasses(World);
+
+		EntityQuery.AddSubsystemRequirement<UMassTestLocalPlayerSubsystem>(EMassFragmentAccess::ReadWrite);
+		EntityQuery.AddSubsystemRequirement<UMassTestGameInstanceSubsystem>(EMassFragmentAccess::ReadWrite);
+
+		SubsystemAccessBits.SetSubsystemRequirementBits(EntityQuery.GetRequiredConstSubsystems(), EntityQuery.GetRequiredMutableSubsystems());
+		SubsystemAccessClasses.SetSubsystemRequirements(EntityQuery);
+
+		UMassTestLocalPlayerSubsystem* AccessLocalPlayerSubsystem = nullptr;
+		UMassTestGameInstanceSubsystem* AccessGameInstanceSubsystem = nullptr;
+
+		const UMassTestLocalPlayerSubsystem* ConstAccessLocalPlayerSubsystem = nullptr;
+		const UMassTestGameInstanceSubsystem* ConstAccessGameInstanceSubsystem = nullptr;
+
+		FMassExecutionContext ExecutionContext(*EntityManager);
+		EntityQuery.ForEachEntityChunk(*EntityManager, ExecutionContext, [&](FMassExecutionContext& Context)
+			{
+				AccessLocalPlayerSubsystem = Context.GetMutableSubsystem<UMassTestLocalPlayerSubsystem>();
+				AccessGameInstanceSubsystem = Context.GetMutableSubsystem<UMassTestGameInstanceSubsystem>();
+
+				ConstAccessLocalPlayerSubsystem = Context.GetSubsystem<UMassTestLocalPlayerSubsystem>();
+				ConstAccessGameInstanceSubsystem = Context.GetSubsystem<UMassTestGameInstanceSubsystem>();
+			});
+
+		AITEST_EQUAL(TEXT("LocalPlayerSubsystem: Mutable Subsystem fetched is expected to be the same as the Actual"), TestLocalPlayerSubsystemActual, AccessLocalPlayerSubsystem);
+		AITEST_EQUAL(TEXT("LocalPlayerSubsystem: Const Subsystem fetched is expected to be the same as the Actual"), TestLocalPlayerSubsystemActual, ConstAccessLocalPlayerSubsystem);
+		AITEST_EQUAL(TEXT("LocalPlayerSubsystem: Subsystem fetched via bits-requirements is expected to be the same as the Actual")
 			, TestLocalPlayerSubsystemActual, SubsystemAccessBits.GetSubsystem<UMassTestLocalPlayerSubsystem>());
-		AITEST_EQUAL(TEXT("LocalPlayerSubsystem: subsystem fetched via class-requirements is expected to be the same as the Actual")
+		AITEST_EQUAL(TEXT("LocalPlayerSubsystem: Subsystem fetched via class-requirements is expected to be the same as the Actual")
 			, TestLocalPlayerSubsystemActual, SubsystemAccessClasses.GetSubsystem<UMassTestLocalPlayerSubsystem>());
 
-		AITEST_EQUAL(TEXT("GameInstanceSubsystem: Mutable subsystem fetched is expected to be the same as the Actual"), TestGameInstanceSubsystemActual, AccessGameInstanceSubsystem);
-		AITEST_EQUAL(TEXT("GameInstanceSubsystem: Const subsystem fetched is expected to be the same as the Actual"), TestGameInstanceSubsystemActual, ConstAccessGameInstanceSubsystem);
-		AITEST_EQUAL(TEXT("GameInstanceSubsystem: subsystem fetched via bits-requirements is expected to be the same as the Actual")
+		AITEST_EQUAL(TEXT("GameInstanceSubsystem: Mutable Subsystem fetched is expected to be the same as the Actual"), TestGameInstanceSubsystemActual, AccessGameInstanceSubsystem);
+		AITEST_EQUAL(TEXT("GameInstanceSubsystem: Const Subsystem fetched is expected to be the same as the Actual"), TestGameInstanceSubsystemActual, ConstAccessGameInstanceSubsystem);
+		AITEST_EQUAL(TEXT("GameInstanceSubsystem: Subsystem fetched via bits-requirements is expected to be the same as the Actual")
 			, TestGameInstanceSubsystemActual, SubsystemAccessBits.GetSubsystem<UMassTestGameInstanceSubsystem>());
-		AITEST_EQUAL(TEXT("GameInstanceSubsystem: subsystem fetched via class-requirements is expected to be the same as the Actual")
+		AITEST_EQUAL(TEXT("GameInstanceSubsystem: Subsystem fetched via class-requirements is expected to be the same as the Actual")
 			, TestGameInstanceSubsystemActual, SubsystemAccessClasses.GetSubsystem<UMassTestGameInstanceSubsystem>());
 
 		return true;
 	}
 };
-IMPLEMENT_AI_INSTANT_TEST(FSubsystemAccessAPI, "System.Mass.Query.System.AccessAPI");
+IMPLEMENT_AI_INSTANT_TEST(FGameSubsystemAccessAPI, "System.Mass.Query.System.AccessAPI.Game");
 
 } // FMassSystemRequirementTest
 
diff --git a/Engine/Plugins/Runtime/MassEntity/Source/MassEntityTestSuite/Public/MassEntityTestTypes.h b/Engine/Plugins/Runtime/MassEntity/Source/MassEntityTestSuite/Public/MassEntityTestTypes.h
--- a/Engine/Plugins/Runtime/MassEntity/Source/MassEntityTestSuite/Public/MassEntityTestTypes.h
+++ b/Engine/Plugins/Runtime/MassEntity/Source/MassEntityTestSuite/Public/MassEntityTestTypes.h
@@ -313,29 +313,6 @@ struct TMassExternalSubsystemTraits<UMassTestEngineSubsystem>
 	};
 };
 
-UCLASS()
-class UMassTestCustomSubsystem : public USubsystem
-{
-	GENERATED_BODY()
-
-public:
-	static UMassTestCustomSubsystem* Create();
-	static UMassTestCustomSubsystem* Get();
-private:
-	static TWeakObjectPtr<UMassTestCustomSubsystem> Instance;
-};
-
-template<>
-struct TMassExternalSubsystemTraits<UMassTestCustomSubsystem>
-{
-	enum
-	{
-		GameThreadOnly = false,
-		ThreadSafeRead = true,
-		ThreadSafeWrite = false,
-	};
-};
-
 UCLASS()
 class UMassTestLocalPlayerSubsystem : public ULocalPlayerSubsystem
 {
