Added IsProcessing checks in a couple of Mass entity-creation places.
This change is related to a licensee running into issues when spawning entities at arbitrary moments, clashing with Mass's assumption that the data composition won't be changed while processors are active.
As part of the change I refactored all the `__FUNCTION__` occurences to not need the `ANSI_TO_TCHAR` macro (using "%hs" in format string).

diff --git a/Engine/Plugins/Runtime/MassEntity/Source/MassEntity/Private/MassEntityManager.cpp b/Engine/Plugins/Runtime/MassEntity/Source/MassEntity/Private/MassEntityManager.cpp
--- a/Engine/Plugins/Runtime/MassEntity/Source/MassEntity/Private/MassEntityManager.cpp
+++ b/Engine/Plugins/Runtime/MassEntity/Source/MassEntity/Private/MassEntityManager.cpp
@@ -95,8 +95,8 @@ void FMassEntityManager::Initialize()
 {
 	if (bInitialized)
 	{
-		UE_LOG(LogMass, Log, TEXT("Calling %s on already initialized entity manager owned by %s")
-			, ANSI_TO_TCHAR(__FUNCTION__), *GetNameSafe(Owner.Get()));
+		UE_LOG(LogMass, Log, TEXT("Calling %s on already initialized entity manager owned by %hs")
+			, __FUNCTION__, *GetNameSafe(Owner.Get()));
 		return;
 	}
 
@@ -174,8 +174,8 @@ void FMassEntityManager::Deinitialize()
 	}
 	else
 	{
-		UE_LOG(LogMass, Log, TEXT("Calling %s on already deinitialized entity manager owned by %s")
-			, ANSI_TO_TCHAR(__FUNCTION__), *GetNameSafe(Owner.Get()));
+		UE_LOG(LogMass, Log, TEXT("Calling %hs on already deinitialized entity manager owned by %hs")
+			, __FUNCTION__, *GetNameSafe(Owner.Get()));
 	}
 }
 
@@ -198,7 +198,7 @@ FMassArchetypeHandle FMassEntityManager::CreateArchetype(FMassArchetypeHandle So
 FMassArchetypeHandle FMassEntityManager::CreateArchetype(const TSharedPtr<FMassArchetypeData>& SourceArchetype, const FMassFragmentBitSet& AddedFragments, const FName ArchetypeDebugName)
 {
 	check(SourceArchetype.IsValid());
-	checkf(AddedFragments.IsEmpty() == false, TEXT("%s Adding an empty fragment list to an archetype is not supported."), ANSI_TO_TCHAR(__FUNCTION__));
+	checkf(AddedFragments.IsEmpty() == false, TEXT("%hs Adding an empty fragment list to an archetype is not supported."), __FUNCTION__);
 
 	const FMassArchetypeCompositionDescriptor Composition(AddedFragments + SourceArchetype->GetFragmentBitSet(), SourceArchetype->GetTagBitSet(), SourceArchetype->GetChunkFragmentBitSet(), SourceArchetype->GetSharedFragmentBitSet());
 	return CreateArchetype(Composition, ArchetypeDebugName);
@@ -326,8 +326,8 @@ void FMassEntityManager::InternalAppendFragmentsAndTagsToArchetypeCompositionDes
 		}
 		else
 		{
-			UE_LOG(LogMass, Warning, TEXT("%s: %s is not a valid fragment nor tag type. Ignoring.")
-				, ANSI_TO_TCHAR(__FUNCTION__), *GetNameSafe(Type));
+			UE_LOG(LogMass, Warning, TEXT("%hs: %s is not a valid fragment nor tag type. Ignoring.")
+				, __FUNCTION__, *GetNameSafe(Type));
 		}
 	}
 }
@@ -379,6 +379,7 @@ void FMassEntityManager::DoEntityCompaction(const double TimeAllowed)
 
 FMassEntityHandle FMassEntityManager::CreateEntity(const FMassArchetypeHandle& ArchetypeHandle, const FMassArchetypeSharedFragmentValues& SharedFragmentValues)
 {
+	checkf(IsProcessing() == false, TEXT("Synchronous API function %hs called during mass processing. Use asynchronous API instead."), __FUNCTION__);
 	check(ArchetypeHandle.IsValid());
 
 	const FMassEntityHandle Entity = ReserveEntity();
@@ -388,6 +389,7 @@ FMassEntityHandle FMassEntityManager::CreateEntity(const FMassArchetypeHandle& A
 
 FMassEntityHandle FMassEntityManager::CreateEntity(TConstArrayView<FInstancedStruct> FragmentInstanceList, const FMassArchetypeSharedFragmentValues& SharedFragmentValues, const FName ArchetypeDebugName)
 {
+	checkf(IsProcessing() == false, TEXT("Synchronous API function %hs called during mass processing. Use asynchronous API instead."), __FUNCTION__);
 	check(FragmentInstanceList.Num() > 0);
 
 	const FMassArchetypeHandle& ArchetypeHandle = CreateArchetype(FMassArchetypeCompositionDescriptor(FragmentInstanceList,
@@ -423,6 +425,7 @@ void FMassEntityManager::ReleaseReservedEntity(FMassEntityHandle Entity)
 
 void FMassEntityManager::BuildEntity(FMassEntityHandle Entity, const FMassArchetypeHandle& ArchetypeHandle, const FMassArchetypeSharedFragmentValues& SharedFragmentValues)
 {
+	checkf(IsProcessing() == false, TEXT("Synchronous API function %hs called during mass processing. Use asynchronous API instead."), __FUNCTION__);
 	checkf(!IsEntityBuilt(Entity), TEXT("Expecting an entity that is not already built"));
 	check(ArchetypeHandle.IsValid());
 
@@ -431,6 +434,7 @@ void FMassEntityManager::BuildEntity(FMassEntityHandle Entity, const FMassArchet
 
 void FMassEntityManager::BuildEntity(FMassEntityHandle Entity, TConstArrayView<FInstancedStruct> FragmentInstanceList, const FMassArchetypeSharedFragmentValues& SharedFragmentValues)
 {
+	checkf(IsProcessing() == false, TEXT("Synchronous API function %hs called during mass processing. Use asynchronous API instead."), __FUNCTION__);
 	check(FragmentInstanceList.Num() > 0);
 	checkf(!IsEntityBuilt(Entity), TEXT("Expecting an entity that is not already built"));
 
@@ -456,6 +460,7 @@ void FMassEntityManager::BuildEntity(FMassEntityHandle Entity, TConstArrayView<F
 
 void FMassEntityManager::BatchBuildEntities(const FMassArchetypeEntityCollectionWithPayload& EncodedEntitiesWithPayload, const FMassFragmentBitSet& FragmentsAffected, const FMassArchetypeSharedFragmentValues& SharedFragmentValues, const FName ArchetypeDebugName)
 {
+	checkf(IsProcessing() == false, TEXT("Synchronous API function %hs called during mass processing. Use asynchronous API instead."), __FUNCTION__);
 	check(SharedFragmentValues.IsSorted());
 
 	FMassArchetypeCompositionDescriptor Composition(FragmentsAffected, FMassTagBitSet(), FMassChunkFragmentBitSet(), FMassSharedFragmentBitSet());
@@ -473,6 +478,8 @@ void FMassEntityManager::BatchBuildEntities(const FMassArchetypeEntityCollection
 
 void FMassEntityManager::BatchBuildEntities(const FMassArchetypeEntityCollectionWithPayload& EncodedEntitiesWithPayload, FMassArchetypeCompositionDescriptor&& Composition, const FMassArchetypeSharedFragmentValues& SharedFragmentValues, const FName ArchetypeDebugName)
 {
+	checkf(IsProcessing() == false, TEXT("Synchronous API function %hs called during mass processing. Use asynchronous API instead."), __FUNCTION__);
+
 	TRACE_CPUPROFILER_EVENT_SCOPE(Mass_BatchBuildEntities);
 
 	FMassArchetypeEntityCollection::FEntityRangeArray TargetArchetypeEntityRanges;
@@ -522,6 +529,8 @@ TSharedRef<FMassEntityManager::FEntityCreationContext> FMassEntityManager::Batch
 {
 	TRACE_CPUPROFILER_EVENT_SCOPE(Mass_BatchCreateEntities);
 
+	checkf(IsProcessing() == false, TEXT("Synchronous API function %hs called during mass processing. Use asynchronous API instead."), __FUNCTION__);
+
 	FMassArchetypeData& ArchetypeData = FMassArchetypeHelper::ArchetypeDataFromHandleChecked(ArchetypeHandle);
 	check(Count > 0);
 	
@@ -558,7 +567,7 @@ TSharedRef<FMassEntityManager::FEntityCreationContext> FMassEntityManager::Batch
 
 void FMassEntityManager::DestroyEntity(FMassEntityHandle Entity)
 {
-	checkf(IsProcessing() == false, TEXT("Synchronous API function %s called during mass processing. Use asynchronous API instead."), ANSI_TO_TCHAR(__FUNCTION__));
+	checkf(IsProcessing() == false, TEXT("Synchronous API function %hs called during mass processing. Use asynchronous API instead."), __FUNCTION__);
 	
 	CheckIfEntityIsActive(Entity);
 
@@ -579,7 +588,7 @@ void FMassEntityManager::BatchDestroyEntities(TConstArrayView<FMassEntityHandle>
 {
 	TRACE_CPUPROFILER_EVENT_SCOPE(Mass_BatchDestroyEntities);
 
-	checkf(IsProcessing() == false, TEXT("Synchronous API function %s called during mass processing. Use asynchronous API instead."), ANSI_TO_TCHAR(__FUNCTION__));
+	checkf(IsProcessing() == false, TEXT("Synchronous API function %hs called during mass processing. Use asynchronous API instead."), __FUNCTION__);
 	
 	EntityFreeIndexList.Reserve(EntityFreeIndexList.Num() + InEntities.Num());
 
@@ -610,7 +619,7 @@ void FMassEntityManager::BatchDestroyEntityChunks(const FMassArchetypeEntityColl
 {
 	TRACE_CPUPROFILER_EVENT_SCOPE(Mass_BatchDestroyEntityChunks);
 
-	checkf(IsProcessing() == false, TEXT("Synchronous API function %s called during mass processing. Use asynchronous API instead."), ANSI_TO_TCHAR(__FUNCTION__));
+	checkf(IsProcessing() == false, TEXT("Synchronous API function %hs called during mass processing. Use asynchronous API instead."), __FUNCTION__);
 
 	TArray<FMassEntityHandle> EntitiesRemoved;
 	// note that it's important to place the context instance in the same scope as the loop below that updates 
@@ -650,8 +659,8 @@ void FMassEntityManager::BatchDestroyEntityChunks(const FMassArchetypeEntityColl
 
 void FMassEntityManager::AddFragmentToEntity(FMassEntityHandle Entity, const UScriptStruct* FragmentType)
 {
-	checkf(FragmentType, TEXT("Null fragment type passed in to %s"), ANSI_TO_TCHAR(__FUNCTION__));
-	checkf(IsProcessing() == false, TEXT("Synchronous API function %s called during mass processing. Use asynchronous API instead."), ANSI_TO_TCHAR(__FUNCTION__));
+	checkf(FragmentType, TEXT("Null fragment type passed in to %hs"), __FUNCTION__);
+	checkf(IsProcessing() == false, TEXT("Synchronous API function %hs called during mass processing. Use asynchronous API instead."), __FUNCTION__);
 
 	CheckIfEntityIsActive(Entity);
 
@@ -774,7 +783,7 @@ void FMassEntityManager::InternalAddFragmentListToEntityChecked(FMassEntityHandl
 
 void FMassEntityManager::InternalAddFragmentListToEntity(FMassEntityHandle Entity, const FMassFragmentBitSet& InFragments)
 {
-	checkf(InFragments.IsEmpty() == false, TEXT("%s is intended for internal calls with non empty NewFragments parameter"), ANSI_TO_TCHAR(__FUNCTION__));
+	checkf(InFragments.IsEmpty() == false, TEXT("%hs is intended for internal calls with non empty NewFragments parameter"), __FUNCTION__);
 	check(Entities.IsValidIndex(Entity.Index));
 	FEntityData& EntityData = Entities[Entity.Index];
 	check(EntityData.CurrentArchetype.IsValid());
@@ -795,7 +804,7 @@ void FMassEntityManager::InternalAddFragmentListToEntity(FMassEntityHandle Entit
 
 void FMassEntityManager::AddFragmentInstanceListToEntity(FMassEntityHandle Entity, TConstArrayView<FInstancedStruct> FragmentInstanceList)
 {
-	checkf(IsProcessing() == false, TEXT("Synchronous API function %s called during mass processing. Use asynchronous API instead."), ANSI_TO_TCHAR(__FUNCTION__));
+	checkf(IsProcessing() == false, TEXT("Synchronous API function %hs called during mass processing. Use asynchronous API instead."), __FUNCTION__);
 
 	CheckIfEntityIsActive(Entity);
 	checkf(FragmentInstanceList.Num() > 0, TEXT("Need to specify at least one fragment instances for this operation"));
@@ -813,7 +822,7 @@ void FMassEntityManager::RemoveFragmentFromEntity(FMassEntityHandle Entity, cons
 
 void FMassEntityManager::RemoveFragmentListFromEntity(FMassEntityHandle Entity, TConstArrayView<const UScriptStruct*> FragmentList)
 {
-	checkf(IsProcessing() == false, TEXT("Synchronous API function %s called during mass processing. Use asynchronous API instead."), ANSI_TO_TCHAR(__FUNCTION__));
+	checkf(IsProcessing() == false, TEXT("Synchronous API function %hs called during mass processing. Use asynchronous API instead."), __FUNCTION__);
 
 	CheckIfEntityIsActive(Entity);
 	
@@ -839,12 +848,12 @@ void FMassEntityManager::RemoveFragmentListFromEntity(FMassEntityHandle Entity,
 
 void FMassEntityManager::SwapTagsForEntity(FMassEntityHandle Entity, const UScriptStruct* OldTagType, const UScriptStruct* NewTagType)
 {
-	checkf(IsProcessing() == false, TEXT("Synchronous API function %s called during mass processing. Use asynchronous API instead."), ANSI_TO_TCHAR(__FUNCTION__));
+	checkf(IsProcessing() == false, TEXT("Synchronous API function %hs called during mass processing. Use asynchronous API instead."), __FUNCTION__);
 
 	CheckIfEntityIsActive(Entity);
 
-	checkf((OldTagType != nullptr) && OldTagType->IsChildOf(FMassTag::StaticStruct()), TEXT("%s works only with tags while '%s' is not one."), ANSI_TO_TCHAR(__FUNCTION__), *GetPathNameSafe(OldTagType));
-	checkf((NewTagType != nullptr) && NewTagType->IsChildOf(FMassTag::StaticStruct()), TEXT("%s works only with tags while '%s' is not one."), ANSI_TO_TCHAR(__FUNCTION__), *GetPathNameSafe(NewTagType));
+	checkf((OldTagType != nullptr) && OldTagType->IsChildOf(FMassTag::StaticStruct()), TEXT("%hs works only with tags while '%s' is not one."), __FUNCTION__, *GetPathNameSafe(OldTagType));
+	checkf((NewTagType != nullptr) && NewTagType->IsChildOf(FMassTag::StaticStruct()), TEXT("%hs works only with tags while '%s' is not one."), __FUNCTION__, *GetPathNameSafe(NewTagType));
 
 	FEntityData& EntityData = Entities[Entity.Index];
 	FMassArchetypeData* CurrentArchetype = EntityData.CurrentArchetype.Get();
@@ -867,7 +876,7 @@ void FMassEntityManager::SwapTagsForEntity(FMassEntityHandle Entity, const UScri
 
 void FMassEntityManager::AddTagToEntity(FMassEntityHandle Entity, const UScriptStruct* TagType)
 {
-	checkf((TagType != nullptr) && TagType->IsChildOf(FMassTag::StaticStruct()), TEXT("%s works only with tags while '%s' is not one."), ANSI_TO_TCHAR(__FUNCTION__), *GetPathNameSafe(TagType));
+	checkf((TagType != nullptr) && TagType->IsChildOf(FMassTag::StaticStruct()), TEXT("%hs works only with tags while '%s' is not one."), __FUNCTION__, *GetPathNameSafe(TagType));
 
 	CheckIfEntityIsActive(Entity);
 
@@ -891,7 +900,7 @@ void FMassEntityManager::AddTagToEntity(FMassEntityHandle Entity, const UScriptS
 	
 void FMassEntityManager::RemoveTagFromEntity(FMassEntityHandle Entity, const UScriptStruct* TagType)
 {
-	checkf((TagType != nullptr) && TagType->IsChildOf(FMassTag::StaticStruct()), TEXT("%s works only with tags while '%s' is not one."), ANSI_TO_TCHAR(__FUNCTION__), *GetPathNameSafe(TagType));
+	checkf((TagType != nullptr) && TagType->IsChildOf(FMassTag::StaticStruct()), TEXT("%hs works only with tags while '%s' is not one."), __FUNCTION__, *GetPathNameSafe(TagType));
 
 	CheckIfEntityIsActive(Entity);
 
