Added a function to check if a given  FMassArchetypeSharedFragmentValues instance has all the required data as indicated by given FMassSharedFragmentBitSet

Also, in the same file but unrelated, added trivial function/s for resetting TMultiArray's contents

diff --git a/Engine/Plugins/Runtime/MassEntity/Source/MassEntity/Public/MassEntityTypes.h b/Engine/Plugins/Runtime/MassEntity/Source/MassEntity/Public/MassEntityTypes.h
--- a/Engine/Plugins/Runtime/MassEntity/Source/MassEntity/Public/MassEntityTypes.h
+++ b/Engine/Plugins/Runtime/MassEntity/Source/MassEntity/Public/MassEntityTypes.h
@@ -251,6 +251,11 @@ struct MASSENTITY_API FMassArchetypeSharedFragmentValues
 		return SharedFragmentBitSet == InSharedFragmentBitSet;
 	}
 
+	FORCEINLINE bool HasAllRequiredFragmentTypes(const FMassSharedFragmentBitSet& InSharedFragmentBitSet) const
+	{
+		return SharedFragmentBitSet.HasAll(InSharedFragmentBitSet);
+	}
+
 	FORCEINLINE bool IsEquivalent(const FMassArchetypeSharedFragmentValues& OtherSharedFragmentValues) const
 	{
 		return GetTypeHash(*this) == GetTypeHash(OtherSharedFragmentValues);
@@ -503,6 +508,12 @@ namespace UE::Mass
 			OutBitSet += FMassFragmentBitSet::GetTypeBitSet<FType>();
 		}
 
+		void Reset()
+		{
+			Super::Reset();
+			FragmentInstances.Reset();
+		}
+
 		TArray<FType> FragmentInstances;
 	};
 
@@ -532,6 +543,11 @@ namespace UE::Mass
 			OutBitSet += FMassFragmentBitSet::GetTypeBitSet<FType>();
 		}
 
+		void Reset()
+		{
+			FragmentInstances.Reset();
+		}
+
 		TArray<FType> FragmentInstances;
 	};
 
