Fixed FMassCommandAddFragmentInstances not resetting properly which is an issue if it's the one owned by FMassEntityManager's command buffer

diff --git a/Engine/Plugins/Runtime/MassEntity/Source/MassEntity/Public/MassCommands.h b/Engine/Plugins/Runtime/MassEntity/Source/MassEntity/Public/MassCommands.h
--- a/Engine/Plugins/Runtime/MassEntity/Source/MassEntity/Public/MassCommands.h
+++ b/Engine/Plugins/Runtime/MassEntity/Source/MassEntity/Public/MassCommands.h
@@ -355,6 +355,13 @@ struct FMassCommandAddFragmentInstances : public FMassBatchedEntityCommand
 	}
 
 protected:
+	virtual void Reset() override
+	{
+		Fragments.Reset(); 
+		FragmentsAffected.Reset();
+		Super::Reset();
+	}
+
 	virtual SIZE_T GetAllocatedSize() const override
 	{
 		return Super::GetAllocatedSize() + Fragments.GetAllocatedSize() + FragmentsAffected.GetAllocatedSize();
