Fixed dynamic processors crashing the Mass Debugger and during shutdown.

Dynamic processors that would get marked as being garbage even when they get used again later, so the explicit marking was removed. This used to be needed to avoid duplicate entries in the Mass Debugger, but testing showed this is no longer needed.

diff --git a/Engine/Plugins/Runtime/MassEntity/Source/MassEntity/Private/MassProcessingTypes.cpp b/Engine/Plugins/Runtime/MassEntity/Source/MassEntity/Private/MassProcessingTypes.cpp
--- a/Engine/Plugins/Runtime/MassEntity/Source/MassEntity/Private/MassProcessingTypes.cpp
+++ b/Engine/Plugins/Runtime/MassEntity/Source/MassEntity/Private/MassProcessingTypes.cpp
@@ -45,13 +45,6 @@ FMassProcessingContext::~FMassProcessingContext()
 //----------------------------------------------------------------------//
 void FMassRuntimePipeline::Reset()
 {
-	for (UMassProcessor* Processor : Processors)
-	{
-		if (Processor)
-		{
-			Processor->MarkAsGarbage();
-		}
-	}
 	Processors.Reset();
 }
 
