Use dynamic string for FMassProcessingPhase::ExecuteTick TRACE event
* Before this change, all event names were "FMassProcessingPhase::ExecuteTick EMassProcessingPhase::PrePhysics" which leads to confusion when looking at mass processor dispatch in UE Insights

diff --git a/Engine/Plugins/Runtime/MassEntity/Source/MassEntity/Private/MassProcessingPhaseManager.cpp b/Engine/Plugins/Runtime/MassEntity/Source/MassEntity/Private/MassProcessingPhaseManager.cpp
--- a/Engine/Plugins/Runtime/MassEntity/Source/MassEntity/Private/MassProcessingPhaseManager.cpp
+++ b/Engine/Plugins/Runtime/MassEntity/Source/MassEntity/Private/MassProcessingPhaseManager.cpp
@@ -58,7 +58,7 @@ void FMassProcessingPhase::ExecuteTick(float DeltaTime, ELevelTick TickType, ENa
 
 	checkf(PhaseManager, TEXT("Manager is null which is not a supported case. Either this FMassProcessingPhase has not been initialized properly or it's been left dangling after the FMassProcessingPhase owner got destroyed."));
 
-	TRACE_CPUPROFILER_EVENT_SCOPE_STR(*FString::Printf(TEXT("FMassProcessingPhase::ExecuteTick %s"), *UEnum::GetValueAsString(Phase)));
+	TRACE_CPUPROFILER_EVENT_SCOPE_TEXT(*FString::Printf(TEXT("FMassProcessingPhase::ExecuteTick %s"), *UEnum::GetValueAsString(Phase)));
 
 	PhaseManager->OnPhaseStart(*this);
 	{
