Order TypedElementDatabaseSubsystem tick to wait on dispatched tasks after dispatching them.
* The current implementation would dispatch processors onto worker threads and potentially cause the jobs to execute across frame boundaries (and possibly more critically, during a UWorld::Tick()).  This _may_ cause race conditions.
* Many processors are tied to the GameThread though, so the current implementation doesn't actually run processors when CompleteAndDispatchWhenReady returns.  The tasks for these processors are queued onto the GameThread and can only run when the code yields within the TaskGraph system. The result is that processors tend to run with one frame delay.
* The fix here is to wait on the tasks after they are dispatched.  This will ensure that processors are run on the same frame number as the PreTickDelegate.

diff --git a/Engine/Plugins/Runtime/MassEntity/Source/MassEntityEditor/Private/MassEntityEditorSubsystem.cpp b/Engine/Plugins/Runtime/MassEntity/Source/MassEntityEditor/Private/MassEntityEditorSubsystem.cpp
--- a/Engine/Plugins/Runtime/MassEntity/Source/MassEntityEditor/Private/MassEntityEditorSubsystem.cpp
+++ b/Engine/Plugins/Runtime/MassEntity/Source/MassEntityEditor/Private/MassEntityEditorSubsystem.cpp
@@ -123,15 +123,12 @@ void UMassEntityEditorSubsystem::StopAndCleanUp()
 
 void UMassEntityEditorSubsystem::Tick(float DeltaTime)
 {
+	TRACE_CPUPROFILER_EVENT_SCOPE(UMassEntityEditorSubsystem::Tick)
 	IsProcessing = true;
 
-	if (CompletionEvent.IsValid())
-	{
-		CompletionEvent->Wait();
-	}
-
 	OnPreTickDelegate.Broadcast(DeltaTime);
 
+	FGraphEventRef CompletionEvent;
 	for (int PhaseIndex = 0; PhaseIndex < (int)EMassProcessingPhase::MAX; ++PhaseIndex)
 	{
 		const FGraphEventArray Prerequisites = { CompletionEvent };
@@ -139,6 +136,11 @@ void UMassEntityEditorSubsystem::Tick(float DeltaTime)
 			.ConstructAndDispatchWhenReady(PhaseManager, EMassProcessingPhase(PhaseIndex), DeltaTime);
 	}
 
+	if (CompletionEvent.IsValid())
+	{
+		CompletionEvent->Wait();
+	}
+
 	IsProcessing = false;
 }
 
diff --git a/Engine/Plugins/Runtime/MassEntity/Source/MassEntityEditor/Public/MassEntityEditorSubsystem.h b/Engine/Plugins/Runtime/MassEntity/Source/MassEntityEditor/Public/MassEntityEditorSubsystem.h
--- a/Engine/Plugins/Runtime/MassEntity/Source/MassEntityEditor/Public/MassEntityEditorSubsystem.h
+++ b/Engine/Plugins/Runtime/MassEntity/Source/MassEntityEditor/Public/MassEntityEditorSubsystem.h
@@ -49,7 +49,6 @@ protected:
 	TSharedRef<FMassEntityManager> EntityManager;
 
 	TSharedRef<FMassProcessingPhaseManager> PhaseManager;
-	FGraphEventRef CompletionEvent;
 
 	FMassProcessingPhaseConfig ProcessingPhasesConfig[(uint8)EMassProcessingPhase::MAX];
 
