Added a check to FMassEntityManager::BatchCreateEntities that ensures that all the required shared data is being passed in while creating entities

diff --git a/Engine/Plugins/Runtime/MassEntity/Source/MassEntity/Private/MassEntityManager.cpp b/Engine/Plugins/Runtime/MassEntity/Source/MassEntity/Private/MassEntityManager.cpp
--- a/Engine/Plugins/Runtime/MassEntity/Source/MassEntity/Private/MassEntityManager.cpp
+++ b/Engine/Plugins/Runtime/MassEntity/Source/MassEntity/Private/MassEntityManager.cpp
@@ -530,10 +530,16 @@ TSharedRef<FMassEntityManager::FEntityCreationContext> FMassEntityManager::Batch
 	TRACE_CPUPROFILER_EVENT_SCOPE(Mass_BatchCreateEntities);
 
 	checkf(IsProcessing() == false, TEXT("Synchronous API function %hs called during mass processing. Use asynchronous API instead."), __FUNCTION__);
-
-	FMassArchetypeData& ArchetypeData = FMassArchetypeHelper::ArchetypeDataFromHandleChecked(ArchetypeHandle);
 	check(Count > 0);
 	
+	// verify that SharedFragmentValues contains all the data needed for the archetype indicated by ArchetypeHandle
+	FMassArchetypeData& ArchetypeData = FMassArchetypeHelper::ArchetypeDataFromHandleChecked(ArchetypeHandle);
+	if (!ensureMsgf(SharedFragmentValues.HasAllRequiredFragmentTypes(ArchetypeData.GetSharedFragmentBitSet())
+		, TEXT("Trying to create entities with mismatching shared fragments collection. This would have lead to crashes, so we're bailing out. Make sure you pass in values for all the shared fragments declared in archetype's composition.")))
+	{
+		return MakeShareable(new FEntityCreationContext(0));
+	}
+
 	int32 Index = OutEntities.Num();
 	OutEntities.AddDefaulted(Count);
 
