Fixed UE::Mass::Utils::GetProcessorExecutionFlagsForWorld spelling

diff --git a/Engine/Plugins/Runtime/MassEntity/Source/MassEntity/Private/MassEntityUtils.cpp b/Engine/Plugins/Runtime/MassEntity/Source/MassEntity/Private/MassEntityUtils.cpp
--- a/Engine/Plugins/Runtime/MassEntity/Source/MassEntity/Private/MassEntityUtils.cpp
+++ b/Engine/Plugins/Runtime/MassEntity/Source/MassEntity/Private/MassEntityUtils.cpp
@@ -10,7 +10,7 @@
 
 namespace UE::Mass::Utils
 {
-EProcessorExecutionFlags GetProcessorExecutionFlagsForWold(const UWorld& World)
+EProcessorExecutionFlags GetProcessorExecutionFlagsForWorld(const UWorld& World)
 {
 	if (World.IsEditorWorld() && !World.IsGameWorld())
 	{
diff --git a/Engine/Plugins/Runtime/MassEntity/Source/MassEntity/Private/MassObserverManager.cpp b/Engine/Plugins/Runtime/MassEntity/Source/MassEntity/Private/MassObserverManager.cpp
--- a/Engine/Plugins/Runtime/MassEntity/Source/MassEntity/Private/MassObserverManager.cpp
+++ b/Engine/Plugins/Runtime/MassEntity/Source/MassEntity/Private/MassObserverManager.cpp
@@ -19,7 +19,7 @@ void SetUpObservers(FMassEntityManager& EntityManager, const TMap<TPointerType,
 	UObject* Owner = EntityManager.GetOwner();
 	check(Owner);
 	const UWorld* World = Owner->GetWorld();
-	const EProcessorExecutionFlags WorldExecutionFlags = World ? UE::Mass::Utils::GetProcessorExecutionFlagsForWold(*World) : EProcessorExecutionFlags::All;
+	const EProcessorExecutionFlags WorldExecutionFlags = World ? UE::Mass::Utils::GetProcessorExecutionFlagsForWorld(*World) : EProcessorExecutionFlags::All;
 
 	for (auto It : RegisteredObserverTypes)
 	{
diff --git a/Engine/Plugins/Runtime/MassEntity/Source/MassEntity/Private/MassProcessingTypes.cpp b/Engine/Plugins/Runtime/MassEntity/Source/MassEntity/Private/MassProcessingTypes.cpp
--- a/Engine/Plugins/Runtime/MassEntity/Source/MassEntity/Private/MassProcessingTypes.cpp
+++ b/Engine/Plugins/Runtime/MassEntity/Source/MassEntity/Private/MassProcessingTypes.cpp
@@ -96,7 +96,7 @@ void FMassRuntimePipeline::InitializeFromClassArray(TConstArrayView<TSubclassOf<
 	Reset();
 
 	const UWorld* World = InOwner.GetWorld();
-	const EProcessorExecutionFlags WorldExecutionFlags = World ? UE::Mass::Utils::GetProcessorExecutionFlagsForWold(*World) : EProcessorExecutionFlags::All;
+	const EProcessorExecutionFlags WorldExecutionFlags = World ? UE::Mass::Utils::GetProcessorExecutionFlagsForWorld(*World) : EProcessorExecutionFlags::All;
 
 	for (const TSubclassOf<UMassProcessor>& ProcessorClass : InProcessorClasses)
 	{
@@ -127,7 +127,7 @@ bool FMassRuntimePipeline::HasProcessorOfExactClass(TSubclassOf<UMassProcessor>
 void FMassRuntimePipeline::AppendUniqueRuntimeProcessorCopies(TConstArrayView<const UMassProcessor*> InProcessors, UObject& InOwner)
 {
 	const UWorld* World = InOwner.GetWorld();
-	const EProcessorExecutionFlags WorldExecutionFlags = World ? UE::Mass::Utils::GetProcessorExecutionFlagsForWold(*World) : EProcessorExecutionFlags::All;
+	const EProcessorExecutionFlags WorldExecutionFlags = World ? UE::Mass::Utils::GetProcessorExecutionFlagsForWorld(*World) : EProcessorExecutionFlags::All;
 	const int32 StartingCount = Processors.Num();
 		
 	for (const UMassProcessor* Proc : InProcessors)
@@ -166,7 +166,7 @@ void FMassRuntimePipeline::AppendUniqueRuntimeProcessorCopies(TConstArrayView<co
 void FMassRuntimePipeline::AppendOrOverrideRuntimeProcessorCopies(TConstArrayView<const UMassProcessor*> InProcessors, UObject& InOwner)
 {
 	const UWorld* World = InOwner.GetWorld();
-	const EProcessorExecutionFlags WorldExecutionFlags = World ? UE::Mass::Utils::GetProcessorExecutionFlagsForWold(*World) : EProcessorExecutionFlags::All;
+	const EProcessorExecutionFlags WorldExecutionFlags = World ? UE::Mass::Utils::GetProcessorExecutionFlagsForWorld(*World) : EProcessorExecutionFlags::All;
 	
 	for (const UMassProcessor* Proc : InProcessors)
 	{
diff --git a/Engine/Plugins/Runtime/MassEntity/Source/MassEntity/Private/MassProcessor.cpp b/Engine/Plugins/Runtime/MassEntity/Source/MassEntity/Private/MassProcessor.cpp
--- a/Engine/Plugins/Runtime/MassEntity/Source/MassEntity/Private/MassProcessor.cpp
+++ b/Engine/Plugins/Runtime/MassEntity/Source/MassEntity/Private/MassProcessor.cpp
@@ -425,7 +425,7 @@ void UMassCompositeProcessor::Populate(TConstArrayView<FMassProcessorOrderInfo>
 	ChildPipeline.Reset();
 
 	const UWorld* World = GetWorld();
-	const EProcessorExecutionFlags WorldExecutionFlags = World ? UE::Mass::Utils::GetProcessorExecutionFlagsForWold(*World) : EProcessorExecutionFlags::All;
+	const EProcessorExecutionFlags WorldExecutionFlags = World ? UE::Mass::Utils::GetProcessorExecutionFlagsForWorld(*World) : EProcessorExecutionFlags::All;
 	const FMassProcessingPhaseConfig& PhaseConfig = GET_MASS_CONFIG_VALUE(GetProcessingPhaseConfig(ProcessingPhase));
 
 	for (const FMassProcessorOrderInfo& ProcessorInfo : OrderedProcessors)
diff --git a/Engine/Plugins/Runtime/MassEntity/Source/MassEntity/Public/MassEntityUtils.h b/Engine/Plugins/Runtime/MassEntity/Source/MassEntity/Public/MassEntityUtils.h
--- a/Engine/Plugins/Runtime/MassEntity/Source/MassEntity/Public/MassEntityUtils.h
+++ b/Engine/Plugins/Runtime/MassEntity/Source/MassEntity/Public/MassEntityUtils.h
@@ -13,7 +13,13 @@ namespace UE::Mass::Utils
 {
 
 /** returns the current execution mode for the processors calculated from the world network mode */
-MASSENTITY_API extern EProcessorExecutionFlags GetProcessorExecutionFlagsForWold(const UWorld& World);
+MASSENTITY_API extern EProcessorExecutionFlags GetProcessorExecutionFlagsForWorld(const UWorld& World);
+
+UE_DEPRECATED(5.3, "This function is deprecated. Please use the properly spelled GetProcessorExecutionFlagsForWorld instead.")
+FORCEINLINE EProcessorExecutionFlags GetProcessorExecutionFlagsForWold(const UWorld& World)
+{
+	return GetProcessorExecutionFlagsForWorld(World);
+}
 
 /** 
  * Fills OutEntityCollections with per-archetype FMassArchetypeEntityCollection instances. 
