Fixed FMassArchetypeCompositionDescriptor::Reset not resetting shared fragment data

diff --git a/Engine/Plugins/Runtime/MassEntity/Source/MassEntity/Public/MassEntityTypes.h b/Engine/Plugins/Runtime/MassEntity/Source/MassEntity/Public/MassEntityTypes.h
--- a/Engine/Plugins/Runtime/MassEntity/Source/MassEntity/Public/MassEntityTypes.h
+++ b/Engine/Plugins/Runtime/MassEntity/Source/MassEntity/Public/MassEntityTypes.h
@@ -167,6 +167,7 @@ struct FMassArchetypeCompositionDescriptor
 		Fragments.Reset();
 		Tags.Reset();
 		ChunkFragments.Reset();
+		SharedFragments.Reset();
 	}
 
 	bool IsEquivalent(const FMassArchetypeCompositionDescriptor& OtherDescriptor) const
