Batch create reserved Mass entities.

This allows entities to be reserved first and then batch create the archetypes for them later.

diff --git a/Engine/Plugins/Runtime/MassEntity/Source/MassEntity/Private/MassEntityManager.cpp b/Engine/Plugins/Runtime/MassEntity/Source/MassEntity/Private/MassEntityManager.cpp
--- a/Engine/Plugins/Runtime/MassEntity/Source/MassEntity/Private/MassEntityManager.cpp
+++ b/Engine/Plugins/Runtime/MassEntity/Source/MassEntity/Private/MassEntityManager.cpp
@@ -477,7 +477,8 @@ void FMassEntityManager::BatchBuildEntities(const FMassArchetypeEntityCollection
 	BatchBuildEntities(EncodedEntitiesWithPayload, MoveTemp(Composition), SharedFragmentValues, ArchetypeDebugName);
 }
 
-void FMassEntityManager::BatchBuildEntities(const FMassArchetypeEntityCollectionWithPayload& EncodedEntitiesWithPayload, FMassArchetypeCompositionDescriptor&& Composition, const FMassArchetypeSharedFragmentValues& SharedFragmentValues, const FName ArchetypeDebugName)
+void FMassEntityManager::BatchBuildEntities(const FMassArchetypeEntityCollectionWithPayload& EncodedEntitiesWithPayload, FMassArchetypeCompositionDescriptor&& Composition
+	, const FMassArchetypeSharedFragmentValues& SharedFragmentValues, const FName ArchetypeDebugName)
 {
 	checkf(IsProcessing() == false, TEXT("Synchronous API function %hs called during mass processing. Use asynchronous API instead."), __FUNCTION__);
 
@@ -526,13 +527,14 @@ void FMassEntityManager::BatchBuildEntities(const FMassArchetypeEntityCollection
 	}
 }
 
-TSharedRef<FMassEntityManager::FEntityCreationContext> FMassEntityManager::BatchCreateEntities(const FMassArchetypeHandle& ArchetypeHandle, const FMassArchetypeSharedFragmentValues& SharedFragmentValues, const int32 Count, TArray<FMassEntityHandle>& OutEntities)
+TSharedRef<FMassEntityManager::FEntityCreationContext> FMassEntityManager::BatchCreateReservedEntities(const FMassArchetypeHandle& ArchetypeHandle
+	, const FMassArchetypeSharedFragmentValues& SharedFragmentValues, TConstArrayView<FMassEntityHandle> ReservedEntities)
 {
-	TRACE_CPUPROFILER_EVENT_SCOPE(Mass_BatchCreateEntities);
+	TRACE_CPUPROFILER_EVENT_SCOPE(Mass_BatchCreateReservedEntities);
 
 	checkf(IsProcessing() == false, TEXT("Synchronous API function %hs called during mass processing. Use asynchronous API instead."), __FUNCTION__);
-	check(Count > 0);
-	
+	checkf(!ReservedEntities.IsEmpty(), TEXT("No reserved entities given to batch create."));
+
 	// verify that SharedFragmentValues contains all the data needed for the archetype indicated by ArchetypeHandle
 	FMassArchetypeData& ArchetypeData = FMassArchetypeHelper::ArchetypeDataFromHandleChecked(ArchetypeHandle);
 	if (!ensureMsgf(SharedFragmentValues.HasAllRequiredFragmentTypes(ArchetypeData.GetSharedFragmentBitSet())
@@ -541,30 +543,57 @@ TSharedRef<FMassEntityManager::FEntityCreationContext> FMassEntityManager::Batch
 		return MakeShareable(new FEntityCreationContext(0));
 	}
 
+	return InternalBatchCreateReservedEntities(ArchetypeHandle, SharedFragmentValues, ReservedEntities);
+}
+
+TSharedRef<FMassEntityManager::FEntityCreationContext> FMassEntityManager::BatchCreateEntities(const FMassArchetypeHandle& ArchetypeHandle
+	, const FMassArchetypeSharedFragmentValues& SharedFragmentValues, const int32 Count, TArray<FMassEntityHandle>& OutEntities)
+{
+	TRACE_CPUPROFILER_EVENT_SCOPE(Mass_BatchCreateEntities);
+
+	FMassArchetypeData& ArchetypeData = FMassArchetypeHelper::ArchetypeDataFromHandleChecked(ArchetypeHandle);
+	if (!ensureMsgf(SharedFragmentValues.HasAllRequiredFragmentTypes(ArchetypeData.GetSharedFragmentBitSet())
+		, TEXT("Trying to create entities with mismatching shared fragments collection. This would have lead to crashes, so we're bailing out. Make sure you pass in values for all the shared fragments declared in archetype's composition.")))
+	{
+		return MakeShareable(new FEntityCreationContext(0));
+	}
+
 	int32 Index = OutEntities.Num();
-	OutEntities.AddDefaulted(Count);
+	OutEntities.Reserve(Index + Count);
+	for (int32 Counter = 0; Counter < Count; ++Counter)
+	{
+		OutEntities.Add(ReserveEntity());
+	}
+	
+	return InternalBatchCreateReservedEntities(ArchetypeHandle, SharedFragmentValues, MakeArrayView(OutEntities.GetData() + Index, Count));
+}
+
+TSharedRef<FMassEntityManager::FEntityCreationContext> FMassEntityManager::InternalBatchCreateReservedEntities(const FMassArchetypeHandle& ArchetypeHandle
+	, const FMassArchetypeSharedFragmentValues& SharedFragmentValues, TConstArrayView<FMassEntityHandle> ReservedEntities)
+{
+	// Functions calling into this one are required to verify that the archetype handle is valid
+	FMassArchetypeData* ArchetypeData = FMassArchetypeHelper::ArchetypeDataFromHandle(ArchetypeHandle);
 
 	// @todo optimize
-	for (; Index < OutEntities.Num(); ++Index)
+	for (FMassEntityHandle Entity : ReservedEntities)
 	{
-		FMassEntityHandle& Result = OutEntities[Index];
-		Result.Index = (EntityFreeIndexList.Num() > 0) ? EntityFreeIndexList.Pop(/*bAllowShrinking=*/ false) : Entities.Add();
-		Result.SerialNumber = SerialNumberGenerator++;
+		checkf(!IsEntityBuilt(Entity), TEXT("Batch creating reserved entities can only use entities that have not been constructed yet."));
 
-		FEntityData& EntityData = Entities[Result.Index];
+		FEntityData& EntityData = Entities[Entity.Index];
 		EntityData.CurrentArchetype = ArchetypeHandle.DataPtr;
-		EntityData.SerialNumber = Result.SerialNumber;
-		
-		ArchetypeData.AddEntity(Result, SharedFragmentValues);
+		EntityData.SerialNumber = Entity.SerialNumber;
+
+		ArchetypeData->AddEntity(Entity, SharedFragmentValues);
 	}
-		
-	FEntityCreationContext* CreationContext = new FEntityCreationContext(Count);
+
+	FEntityCreationContext* CreationContext = new FEntityCreationContext(ReservedEntities.Num());
 	// @todo this could probably be optimized since one would assume we're adding elements to OutEntities in order.
 	// Then again, if that's the case, the sorting will be almost instant
-	new (&CreationContext->EntityCollection) FMassArchetypeEntityCollection(ArchetypeHandle, MakeArrayView(&OutEntities[OutEntities.Num() - Count], Count), FMassArchetypeEntityCollection::NoDuplicates);
-	if (ObserverManager.HasObserversForBitSet(ArchetypeData.GetCompositionDescriptor().Fragments, EMassObservedOperation::Add))
+	new (&CreationContext->EntityCollection) FMassArchetypeEntityCollection(ArchetypeHandle, ReservedEntities, FMassArchetypeEntityCollection::NoDuplicates);
+	if (ObserverManager.HasObserversForBitSet(ArchetypeData->GetCompositionDescriptor().Fragments, EMassObservedOperation::Add)
+		|| ObserverManager.HasObserversForBitSet(ArchetypeData->GetCompositionDescriptor().Tags, EMassObservedOperation::Add))
 	{
-		CreationContext->OnSpawningFinished = [this](FEntityCreationContext& Context){
+		CreationContext->OnSpawningFinished = [this](FEntityCreationContext& Context) {
 			ObserverManager.OnPostEntitiesCreated(Context.EntityCollection);
 		};
 	}
diff --git a/Engine/Plugins/Runtime/MassEntity/Source/MassEntity/Public/MassEntityManager.h b/Engine/Plugins/Runtime/MassEntity/Source/MassEntity/Public/MassEntityManager.h
--- a/Engine/Plugins/Runtime/MassEntity/Source/MassEntity/Public/MassEntityManager.h
+++ b/Engine/Plugins/Runtime/MassEntity/Source/MassEntity/Public/MassEntityManager.h
@@ -185,7 +185,20 @@ public:
 	};
 
 	/**
-	 * A version of CreateEntity that's creating a number of entities (Count) at one go
+	 * A version of CreateEntity that's creating a number of entities (Count) in one go
+	 * @param Archetype you want this entity to be
+	 * @param SharedFragmentValues to be associated with the entities
+	 * @param ReservedEntities a list of reserved entities that have not yet been assigned to an archetype.
+	 * @return a creation context that will notify all the interested observers about newly created fragments once the context is released */
+	TSharedRef<FEntityCreationContext> BatchCreateReservedEntities(const FMassArchetypeHandle& ArchetypeHandle,
+		const FMassArchetypeSharedFragmentValues& SharedFragmentValues, TConstArrayView<FMassEntityHandle> ReservedEntities);
+	FORCEINLINE TSharedRef<FEntityCreationContext> BatchCreateReservedEntities(const FMassArchetypeHandle& ArchetypeHandle,
+		TConstArrayView<FMassEntityHandle> OutEntities)
+	{
+		return BatchCreateReservedEntities(ArchetypeHandle, FMassArchetypeSharedFragmentValues(), OutEntities);
+	}
+	/**
+	 * A version of CreateEntity that's creating a number of entities (Count) in one go
 	 * @param Archetype you want this entity to be
 	 * @param SharedFragmentValues to be associated with the entities
 	 * @param Count number of entities to create
@@ -451,6 +464,9 @@ private:
 	void InternalAddFragmentListToEntity(FMassEntityHandle Entity, const FMassFragmentBitSet& InFragments);
 	void* InternalGetFragmentDataChecked(FMassEntityHandle Entity, const UScriptStruct* FragmentType) const;
 	void* InternalGetFragmentDataPtr(FMassEntityHandle Entity, const UScriptStruct* FragmentType) const;
+
+	TSharedRef<FEntityCreationContext> InternalBatchCreateReservedEntities(const FMassArchetypeHandle& ArchetypeHandle,
+		const FMassArchetypeSharedFragmentValues& SharedFragmentValues, TConstArrayView<FMassEntityHandle> ReservedEntities);
 private:
 	TChunkedArray<FEntityData> Entities;
 	TArray<int32> EntityFreeIndexList;
