Made synchronously destroying entities with FMassEntityManager::DestroyEntity result in triggering observers (part of "fixing observers" effort).

diff --git a/Engine/Plugins/Runtime/MassEntity/Source/MassEntity/Private/MassEntityManager.cpp b/Engine/Plugins/Runtime/MassEntity/Source/MassEntity/Private/MassEntityManager.cpp
--- a/Engine/Plugins/Runtime/MassEntity/Source/MassEntity/Private/MassEntityManager.cpp
+++ b/Engine/Plugins/Runtime/MassEntity/Source/MassEntity/Private/MassEntityManager.cpp
@@ -609,15 +609,14 @@ void FMassEntityManager::DestroyEntity(FMassEntityHandle Entity)
 
 	const FEntityData& EntityData = Entities[Entity.Index];
 	FMassArchetypeData* Archetype = EntityData.CurrentArchetype.Get();
-	check(Archetype);
-	Archetype->RemoveEntity(Entity);
-
-	InternalReleaseEntity(Entity);
 
-	if (Archetype->GetNumEntities() == 0)
+	if (Archetype)
 	{
-		//@TODO: Drop our reference to the archetype?
+		ObserverManager.OnPreEntityDestroyed(Archetype->GetCompositionDescriptor(), Entity);
+		Archetype->RemoveEntity(Entity);
 	}
+
+	InternalReleaseEntity(Entity);
 }
 
 void FMassEntityManager::BatchDestroyEntities(TConstArrayView<FMassEntityHandle> InEntities)
diff --git a/Engine/Plugins/Runtime/MassEntity/Source/MassEntity/Private/MassObserverManager.cpp b/Engine/Plugins/Runtime/MassEntity/Source/MassEntity/Private/MassObserverManager.cpp
--- a/Engine/Plugins/Runtime/MassEntity/Source/MassEntity/Private/MassObserverManager.cpp
+++ b/Engine/Plugins/Runtime/MassEntity/Source/MassEntity/Private/MassObserverManager.cpp
@@ -133,6 +133,12 @@ bool FMassObserverManager::OnPreEntitiesDestroyed(FMassProcessingContext& Proces
 	return OnCompositionChanged(EntityCollection, ArchetypeComposition, EMassObservedOperation::Remove, &ProcessingContext);
 }
 
+bool FMassObserverManager::OnPreEntityDestroyed(const FMassArchetypeCompositionDescriptor& ArchetypeComposition, const FMassEntityHandle Entity)
+{
+	TRACE_CPUPROFILER_EVENT_SCOPE_STR("OnPreEntityDestroyed")
+	return OnCompositionChanged(Entity, ArchetypeComposition, EMassObservedOperation::Remove);
+}
+
 bool FMassObserverManager::OnCompositionChanged(const FMassArchetypeEntityCollection& EntityCollection, const FMassArchetypeCompositionDescriptor& CompositionDelta, const EMassObservedOperation Operation, FMassProcessingContext* InProcessingContext)
 {
 	const FMassFragmentBitSet FragmentOverlap = ObservedFragments[(uint8)Operation].GetOverlap(CompositionDelta.Fragments);
diff --git a/Engine/Plugins/Runtime/MassEntity/Source/MassEntity/Public/MassObserverManager.h b/Engine/Plugins/Runtime/MassEntity/Source/MassEntity/Public/MassObserverManager.h
--- a/Engine/Plugins/Runtime/MassEntity/Source/MassEntity/Public/MassObserverManager.h
+++ b/Engine/Plugins/Runtime/MassEntity/Source/MassEntity/Public/MassObserverManager.h
@@ -75,6 +75,7 @@ public:
 
 	bool OnPreEntitiesDestroyed(const FMassArchetypeEntityCollection& EntityCollection);
 	bool OnPreEntitiesDestroyed(FMassProcessingContext& ProcessingContext, const FMassArchetypeEntityCollection& EntityCollection);
+	bool OnPreEntityDestroyed(const FMassArchetypeCompositionDescriptor& ArchetypeComposition, const FMassEntityHandle Entity);
 
 	bool OnCompositionChanged(const FMassEntityHandle Entity, const FMassArchetypeCompositionDescriptor& CompositionDelta, const EMassObservedOperation Operation);
 	bool OnPostCompositionAdded(const FMassEntityHandle Entity, const FMassArchetypeCompositionDescriptor& Composition)
