Fixed an issue in Mass's processor list creation resulting in game-plugin-loaded processors not being included. The issue boils down to MassEntitySettings being used before the engine finished initializing which results in locking the list before all the types get loaded.

diff --git a/Engine/Plugins/Runtime/MassEntity/Source/MassEntity/Private/MassEntitySettings.cpp b/Engine/Plugins/Runtime/MassEntity/Source/MassEntity/Private/MassEntitySettings.cpp
--- a/Engine/Plugins/Runtime/MassEntity/Source/MassEntity/Private/MassEntitySettings.cpp
+++ b/Engine/Plugins/Runtime/MassEntity/Source/MassEntity/Private/MassEntitySettings.cpp
@@ -24,7 +24,7 @@ UMassEntitySettings::UMassEntitySettings(const FObjectInitializer& ObjectInitial
 		ProcessingPhasesConfig[i].PhaseName = *UEnum::GetDisplayValueAsText(EMassProcessingPhase(i)).ToString();
 	}
 
-	FCoreDelegates::OnPostEngineInit.AddUObject(this, &UMassEntitySettings::BuildProcessorListAndPhases);
+	FCoreDelegates::OnPostEngineInit.AddUObject(this, &UMassEntitySettings::OnPostEngineInit);
 }
 
 void UMassEntitySettings::BeginDestroy()
@@ -33,9 +33,15 @@ void UMassEntitySettings::BeginDestroy()
 	Super::BeginDestroy();
 }
 
+void UMassEntitySettings::OnPostEngineInit()
+{
+	bEngineInitialized = true;
+	BuildProcessorListAndPhases();
+}
+
 void UMassEntitySettings::BuildProcessorListAndPhases()
 {
-	if (bInitialized)
+	if (bInitialized == true || bEngineInitialized == false)
 	{
 		return;
 	}
diff --git a/Engine/Plugins/Runtime/MassEntity/Source/MassEntity/Public/MassEntitySettings.h b/Engine/Plugins/Runtime/MassEntity/Source/MassEntity/Public/MassEntitySettings.h
--- a/Engine/Plugins/Runtime/MassEntity/Source/MassEntity/Public/MassEntitySettings.h
+++ b/Engine/Plugins/Runtime/MassEntity/Source/MassEntity/Public/MassEntitySettings.h
@@ -44,6 +44,7 @@ protected:
 
 	virtual void BeginDestroy() override;
 
+	void OnPostEngineInit();
 	void BuildPhases();
 	void BuildProcessorList();
 
@@ -69,4 +70,5 @@ protected:
 	FOnSettingsChange OnSettingsChange;
 #endif // WITH_EDITORONLY_DATA
 	bool bInitialized = false;
+	bool bEngineInitialized = false;
 };
