Fixed a bug in FMassEntityManager::BatchDestroyEntities where the ObserverManager wasn't being informed of the entity deletions, causing fragment destructors not to run.

diff --git a/Engine/Plugins/Runtime/MassEntity/Source/MassEntity/Private/MassEntityManager.cpp b/Engine/Plugins/Runtime/MassEntity/Source/MassEntity/Private/MassEntityManager.cpp
--- a/Engine/Plugins/Runtime/MassEntity/Source/MassEntity/Private/MassEntityManager.cpp
+++ b/Engine/Plugins/Runtime/MassEntity/Source/MassEntity/Private/MassEntityManager.cpp
@@ -643,6 +643,7 @@ void FMassEntityManager::BatchDestroyEntities(TConstArrayView<FMassEntityHandle>
 
 		FMassArchetypeData* Archetype = EntityData.CurrentArchetype.Get();
 		check(Archetype);
+		ObserverManager.OnPreEntityDestroyed(Archetype->GetCompositionDescriptor(), Entity);
 		Archetype->RemoveEntity(Entity);
 
 		EntityData.Reset();
