Fix false positive error log for reaching the loop maximum when flushing Mass command buffers.

diff --git a/Engine/Plugins/Runtime/MassEntity/Source/MassEntity/Private/MassEntityManager.cpp b/Engine/Plugins/Runtime/MassEntity/Source/MassEntity/Private/MassEntityManager.cpp
--- a/Engine/Plugins/Runtime/MassEntity/Source/MassEntity/Private/MassEntityManager.cpp
+++ b/Engine/Plugins/Runtime/MassEntity/Source/MassEntity/Private/MassEntityManager.cpp
@@ -1410,11 +1410,17 @@ void FMassEntityManager::FlushCommands(const TSharedPtr<FMassCommandBuffer>& InC
 
 	if (InCommandBuffer)
 	{
-		FlushedCommandBufferQueue.Enqueue(InCommandBuffer);
+		if(InCommandBuffer->HasPendingCommands())
+		{
+			FlushedCommandBufferQueue.Enqueue(InCommandBuffer);
+		}
 	}
 	else
 	{
-		FlushedCommandBufferQueue.Enqueue(DeferredCommandBuffer);
+		if(DeferredCommandBuffer->HasPendingCommands())
+		{
+			FlushedCommandBufferQueue.Enqueue(DeferredCommandBuffer);
+		}
 	}
 
 	if (bCommandBufferFlushingInProgress == false && IsProcessing() == false)
@@ -1423,8 +1429,9 @@ void FMassEntityManager::FlushCommands(const TSharedPtr<FMassCommandBuffer>& InC
 		
 		int IterationsCounter = 0;
 		TOptional<TSharedPtr<FMassCommandBuffer>> CurrentCommandBuffer = FlushedCommandBufferQueue.Dequeue();
-		while (IterationsCounter++ < MaxIterations && CurrentCommandBuffer.IsSet())
+		while (IterationsCounter < MaxIterations && CurrentCommandBuffer.IsSet())
 		{
+			IterationsCounter++;
 			(*CurrentCommandBuffer)->Flush(*this);
 			CurrentCommandBuffer = FlushedCommandBufferQueue.Dequeue();
 		}
