Made Mass's chunk size configurable via MassSettings

diff --git a/Engine/Plugins/Runtime/MassEntity/Source/MassEntity/Private/MassArchetypeData.cpp b/Engine/Plugins/Runtime/MassEntity/Source/MassEntity/Private/MassArchetypeData.cpp
--- a/Engine/Plugins/Runtime/MassEntity/Source/MassEntity/Private/MassArchetypeData.cpp
+++ b/Engine/Plugins/Runtime/MassEntity/Source/MassEntity/Private/MassArchetypeData.cpp
@@ -3,10 +3,25 @@
 #include "MassArchetypeData.h"
 #include "MassEntityTypes.h"
 #include "MassExecutionContext.h"
+#include "MassEntitySettings.h"
 #include "Misc/StringBuilder.h"
 
+
+namespace UE::Mass::Private
+{
+	constexpr int32 UninitializedInt32 = -1;
+}
+
 //////////////////////////////////////////////////////////////////////
 // FMassArchetypeData
+FMassArchetypeData::FMassArchetypeData()
+	: NumEntitiesPerChunk(UE::Mass::Private::UninitializedInt32)
+	, TotalBytesPerEntity(UE::Mass::Private::UninitializedInt32)
+	, EntityListOffsetWithinChunk(UE::Mass::Private::UninitializedInt32)
+	, ChunkMemorySize(GET_MASS_CONFIG_VALUE(ChunkMemorySize))
+{
+	
+}
 
 void FMassArchetypeData::ForEachFragmentType(TFunction< void(const UScriptStruct* /*Fragment*/)> Function) const
 {
diff --git a/Engine/Plugins/Runtime/MassEntity/Source/MassEntity/Private/MassArchetypeData.h b/Engine/Plugins/Runtime/MassEntity/Source/MassEntity/Private/MassArchetypeData.h
--- a/Engine/Plugins/Runtime/MassEntity/Source/MassEntity/Private/MassArchetypeData.h
+++ b/Engine/Plugins/Runtime/MassEntity/Source/MassEntity/Private/MassArchetypeData.h
@@ -10,11 +10,6 @@ struct FMassExecutionContext;
 class FOutputDevice;
 struct FMassArchetypeEntityCollection;
 
-namespace UE::Mass
-{
-	constexpr int32 ChunkSize = 128*1024;
-}
-
 // This is one chunk within an archetype
 struct FMassArchetypeChunk
 {
@@ -27,7 +22,7 @@ private:
 	FMassArchetypeSharedFragmentValues SharedFragmentValues;
 
 public:
-	explicit FMassArchetypeChunk(int32 InAllocSize, TConstArrayView<FInstancedStruct> InChunkFragmentTemplates, FMassArchetypeSharedFragmentValues InSharedFragmentValues)
+	explicit FMassArchetypeChunk(const int32 InAllocSize, TConstArrayView<FInstancedStruct> InChunkFragmentTemplates, FMassArchetypeSharedFragmentValues InSharedFragmentValues)
 		: AllocSize(InAllocSize)
 		, ChunkFragmentData(InChunkFragmentTemplates)
 		, SharedFragmentValues(InSharedFragmentValues)
@@ -177,12 +172,17 @@ private:
 
 	// Arrays of names the archetype is referred as.
 	TArray<FName> DebugNames;
+
+	// Defaults to UMassEntitySettings.ChunkMemorySize. In near future will support being set via constructor.
+	const int32 ChunkMemorySize;
 	
 	friend FMassEntityQuery;
 	friend FMassArchetypeEntityCollection;
 	friend FMassDebugger;
 
 public:
+	FMassArchetypeData();
+
 	TConstArrayView<FMassArchetypeFragmentConfig> GetFragmentConfigs() const { return FragmentConfigs; }
 	const FMassFragmentBitSet& GetFragmentBitSet() const { return CompositionDescriptor.Fragments; }
 	const FMassTagBitSet& GetTagBitSet() const { return CompositionDescriptor.Tags; }
@@ -229,10 +229,11 @@ public:
 
 	FORCEINLINE int32 GetInternalIndexForEntity(const int32 EntityIndex) const { return EntityMap.FindChecked(EntityIndex); }
 	int32 GetNumEntitiesPerChunk() const { return NumEntitiesPerChunk; }
+	int32 GetBytesPerEntity() const { return TotalBytesPerEntity; }
 
 	int32 GetNumEntities() const { return EntityMap.Num(); }
 
-	int32 GetChunkAllocSize() const { return UE::Mass::ChunkSize; }
+	int32 GetChunkAllocSize() const { return ChunkMemorySize; }
 
 	int32 GetChunkCount() const { return Chunks.Num(); }
 
diff --git a/Engine/Plugins/Runtime/MassEntity/Source/MassEntity/Public/MassEntitySettings.h b/Engine/Plugins/Runtime/MassEntity/Source/MassEntity/Public/MassEntitySettings.h
--- a/Engine/Plugins/Runtime/MassEntity/Source/MassEntity/Public/MassEntitySettings.h
+++ b/Engine/Plugins/Runtime/MassEntity/Source/MassEntity/Public/MassEntitySettings.h
@@ -49,7 +49,10 @@ protected:
 	void BuildProcessorList();
 
 public:
-	/** 
+	UPROPERTY(EditDefaultsOnly, Category = Mass, config, meta = (ClampMin = "8", UIMin = "8"), AdvancedDisplay)
+	int32 ChunkMemorySize = 128 * 1024;
+
+	/**
 	 * The name of the file to dump the processor dependency graph. T
 	 * The dot file will be put in the project log folder.
 	 * To generate a svg out of that file, simply run dot executable with following parameters: -Tsvg -O filename.dot 
