Made chunk memory configurable per archetype

diff --git a/Engine/Plugins/Runtime/MassEntity/Source/MassEntity/Private/MassArchetypeData.cpp b/Engine/Plugins/Runtime/MassEntity/Source/MassEntity/Private/MassArchetypeData.cpp
--- a/Engine/Plugins/Runtime/MassEntity/Source/MassEntity/Private/MassArchetypeData.cpp
+++ b/Engine/Plugins/Runtime/MassEntity/Source/MassEntity/Private/MassArchetypeData.cpp
@@ -7,20 +7,35 @@
 #include "Misc/StringBuilder.h"
 
 
-namespace UE::Mass::Private
+namespace UE::Mass
 {
-	constexpr int32 UninitializedInt32 = -1;
+	namespace Private
+	{
+		constexpr int32 UninitializedInt32 = -1;
+		constexpr int32 MinChunkMemorySize = 1024;
+		constexpr int32 MaxChunkMemorySize = 512 * 1024;
+	}
+
+	int32 SanitizeChunkMemorySize(const int32 InChunkMemorySize, const bool bLogMismatch)
+	{
+		const int32 SanitizedSize = FMath::Clamp(InChunkMemorySize, Private::MinChunkMemorySize, Private::MaxChunkMemorySize);
+		UE_CLOG(bLogMismatch && SanitizedSize != InChunkMemorySize, LogMass, Warning
+			, TEXT("ChunkMemorySize sanitization resulted in changing value. Old: %d, modified: %d")
+			, InChunkMemorySize, SanitizedSize);
+		return SanitizedSize;
+	}
 }
 
 //////////////////////////////////////////////////////////////////////
 // FMassArchetypeData
-FMassArchetypeData::FMassArchetypeData()
+
+FMassArchetypeData::FMassArchetypeData(const FMassArchetypeCreationParams& CreationParams)
 	: NumEntitiesPerChunk(UE::Mass::Private::UninitializedInt32)
 	, TotalBytesPerEntity(UE::Mass::Private::UninitializedInt32)
 	, EntityListOffsetWithinChunk(UE::Mass::Private::UninitializedInt32)
-	, ChunkMemorySize(GET_MASS_CONFIG_VALUE(ChunkMemorySize))
+	, ChunkMemorySize(UE::Mass::SanitizeChunkMemorySize(CreationParams.ChunkMemorySize ? CreationParams.ChunkMemorySize : GET_MASS_CONFIG_VALUE(ChunkMemorySize)))
 {
-	
+	DebugNames.Add(CreationParams.DebugName);
 }
 
 void FMassArchetypeData::ForEachFragmentType(TFunction< void(const UScriptStruct* /*Fragment*/)> Function) const
diff --git a/Engine/Plugins/Runtime/MassEntity/Source/MassEntity/Private/MassArchetypeData.h b/Engine/Plugins/Runtime/MassEntity/Source/MassEntity/Private/MassArchetypeData.h
--- a/Engine/Plugins/Runtime/MassEntity/Source/MassEntity/Private/MassArchetypeData.h
+++ b/Engine/Plugins/Runtime/MassEntity/Source/MassEntity/Private/MassArchetypeData.h
@@ -5,11 +5,17 @@
 #include "MassEntityManager.h"
 #include "MassArchetypeTypes.h"
 
+
 struct FMassEntityQuery;
 struct FMassExecutionContext;
 class FOutputDevice;
 struct FMassArchetypeEntityCollection;
 
+namespace UE::Mass
+{
+	int32 SanitizeChunkMemorySize(const int32 InChunkMemorySize, const bool bLogMismatch = true);
+}
+
 // This is one chunk within an archetype
 struct FMassArchetypeChunk
 {
@@ -181,7 +187,7 @@ private:
 	friend FMassDebugger;
 
 public:
-	FMassArchetypeData();
+	explicit FMassArchetypeData(const FMassArchetypeCreationParams& CreationParams = FMassArchetypeCreationParams());
 
 	TConstArrayView<FMassArchetypeFragmentConfig> GetFragmentConfigs() const { return FragmentConfigs; }
 	const FMassFragmentBitSet& GetFragmentBitSet() const { return CompositionDescriptor.Fragments; }
diff --git a/Engine/Plugins/Runtime/MassEntity/Source/MassEntity/Private/MassEntityManager.cpp b/Engine/Plugins/Runtime/MassEntity/Source/MassEntity/Private/MassEntityManager.cpp
--- a/Engine/Plugins/Runtime/MassEntity/Source/MassEntity/Private/MassEntityManager.cpp
+++ b/Engine/Plugins/Runtime/MassEntity/Source/MassEntity/Private/MassEntityManager.cpp
@@ -180,32 +180,32 @@ void FMassEntityManager::Deinitialize()
 	}
 }
 
-FMassArchetypeHandle FMassEntityManager::CreateArchetype(TConstArrayView<const UScriptStruct*> FragmentsAndTagsList, const FName ArchetypeDebugName)
+FMassArchetypeHandle FMassEntityManager::CreateArchetype(TConstArrayView<const UScriptStruct*> FragmentsAndTagsList, const FMassArchetypeCreationParams& CreationParams)
 {
 	FMassArchetypeCompositionDescriptor Composition;
 	InternalAppendFragmentsAndTagsToArchetypeCompositionDescriptor(Composition, FragmentsAndTagsList);
-	return CreateArchetype(Composition, ArchetypeDebugName);
+	return CreateArchetype(Composition, CreationParams);
 }
 
 FMassArchetypeHandle FMassEntityManager::CreateArchetype(FMassArchetypeHandle SourceArchetype,
-	TConstArrayView<const UScriptStruct*> FragmentsAndTagsList, const FName ArchetypeDebugName)
+	TConstArrayView<const UScriptStruct*> FragmentsAndTagsList, const FMassArchetypeCreationParams& CreationParams)
 {
 	const FMassArchetypeData& ArchetypeData = FMassArchetypeHelper::ArchetypeDataFromHandleChecked(SourceArchetype);
 	FMassArchetypeCompositionDescriptor Composition = ArchetypeData.GetCompositionDescriptor();
 	InternalAppendFragmentsAndTagsToArchetypeCompositionDescriptor(Composition, FragmentsAndTagsList);
-	return CreateArchetype(Composition, ArchetypeDebugName);
+	return CreateArchetype(Composition, CreationParams);
 }
 
-FMassArchetypeHandle FMassEntityManager::CreateArchetype(const TSharedPtr<FMassArchetypeData>& SourceArchetype, const FMassFragmentBitSet& AddedFragments, const FName ArchetypeDebugName)
+FMassArchetypeHandle FMassEntityManager::CreateArchetype(const TSharedPtr<FMassArchetypeData>& SourceArchetype, const FMassFragmentBitSet& AddedFragments, const FMassArchetypeCreationParams& CreationParams)
 {
 	check(SourceArchetype.IsValid());
 	checkf(AddedFragments.IsEmpty() == false, TEXT("%hs Adding an empty fragment list to an archetype is not supported."), __FUNCTION__);
 
 	const FMassArchetypeCompositionDescriptor Composition(AddedFragments + SourceArchetype->GetFragmentBitSet(), SourceArchetype->GetTagBitSet(), SourceArchetype->GetChunkFragmentBitSet(), SourceArchetype->GetSharedFragmentBitSet());
-	return CreateArchetype(Composition, ArchetypeDebugName);
+	return CreateArchetype(Composition, CreationParams);
 }
 
-FMassArchetypeHandle FMassEntityManager::CreateArchetype(const FMassArchetypeCompositionDescriptor& Composition, const FName ArchetypeDebugName)
+FMassArchetypeHandle FMassEntityManager::CreateArchetype(const FMassArchetypeCompositionDescriptor& Composition, const FMassArchetypeCreationParams& CreationParams)
 {
 	const uint32 TypeHash = Composition.CalculateHash();
 
@@ -217,9 +217,14 @@ FMassArchetypeHandle FMassEntityManager::CreateArchetype(const FMassArchetypeCom
 		if (Ptr->IsEquivalent(Composition))
 		{
 			// Keep track of all names for this archetype.
-			if (!ArchetypeDebugName.IsNone())
+			if (!CreationParams.DebugName.IsNone())
 			{
-				Ptr->AddUniqueDebugName(ArchetypeDebugName);
+				Ptr->AddUniqueDebugName(CreationParams.DebugName);
+			}
+			if (CreationParams.ChunkMemorySize > 0 && CreationParams.ChunkMemorySize != Ptr->GetChunkAllocSize())
+			{
+				UE_LOG(LogMass, Warning, TEXT("Reusing existing Archetype, but the requested ChunkMemorySize is different. Requested %d, existing: %d")
+					, CreationParams.ChunkMemorySize, Ptr->GetChunkAllocSize());
 			}
 			ArchetypeDataPtr = Ptr;
 			break;
@@ -232,12 +237,8 @@ FMassArchetypeHandle FMassEntityManager::CreateArchetype(const FMassArchetypeCom
 		++ArchetypeDataVersion;
 
 		// Create a new archetype
-		FMassArchetypeData* NewArchetype = new FMassArchetypeData();
+		FMassArchetypeData* NewArchetype = new FMassArchetypeData(CreationParams);
 		NewArchetype->Initialize(Composition, ArchetypeDataVersion);
-		if (!ArchetypeDebugName.IsNone())
-		{
-			NewArchetype->AddUniqueDebugName(ArchetypeDebugName);
-		}
 		ArchetypeDataPtr = HashRow.Add_GetRef(MakeShareable(NewArchetype));
 
 		for (const FMassArchetypeFragmentConfig& FragmentConfig : NewArchetype->GetFragmentConfigs())
@@ -388,13 +389,13 @@ FMassEntityHandle FMassEntityManager::CreateEntity(const FMassArchetypeHandle& A
 	return Entity;
 }
 
-FMassEntityHandle FMassEntityManager::CreateEntity(TConstArrayView<FInstancedStruct> FragmentInstanceList, const FMassArchetypeSharedFragmentValues& SharedFragmentValues, const FName ArchetypeDebugName)
+FMassEntityHandle FMassEntityManager::CreateEntity(TConstArrayView<FInstancedStruct> FragmentInstanceList, const FMassArchetypeSharedFragmentValues& SharedFragmentValues, const FMassArchetypeCreationParams& CreationParams)
 {
 	checkf(IsProcessing() == false, TEXT("Synchronous API function %hs called during mass processing. Use asynchronous API instead."), __FUNCTION__);
 	check(FragmentInstanceList.Num() > 0);
 
 	const FMassArchetypeHandle& ArchetypeHandle = CreateArchetype(FMassArchetypeCompositionDescriptor(FragmentInstanceList,
-		FMassTagBitSet(), FMassChunkFragmentBitSet(), FMassSharedFragmentBitSet()), ArchetypeDebugName);
+		FMassTagBitSet(), FMassChunkFragmentBitSet(), FMassSharedFragmentBitSet()), CreationParams);
 	check(ArchetypeHandle.IsValid());
 
 	const FMassEntityHandle Entity = ReserveEntity();
@@ -459,7 +460,8 @@ void FMassEntityManager::BuildEntity(FMassEntityHandle Entity, TConstArrayView<F
 	EntityData.CurrentArchetype->SetFragmentsData(Entity, FragmentInstanceList);
 }
 
-void FMassEntityManager::BatchBuildEntities(const FMassArchetypeEntityCollectionWithPayload& EncodedEntitiesWithPayload, const FMassFragmentBitSet& FragmentsAffected, const FMassArchetypeSharedFragmentValues& SharedFragmentValues, const FName ArchetypeDebugName)
+void FMassEntityManager::BatchBuildEntities(const FMassArchetypeEntityCollectionWithPayload& EncodedEntitiesWithPayload
+	, const FMassFragmentBitSet& FragmentsAffected, const FMassArchetypeSharedFragmentValues& SharedFragmentValues, const FMassArchetypeCreationParams& CreationParams)
 {
 	checkf(IsProcessing() == false, TEXT("Synchronous API function %hs called during mass processing. Use asynchronous API instead."), __FUNCTION__);
 	check(SharedFragmentValues.IsSorted());
@@ -474,11 +476,11 @@ void FMassEntityManager::BatchBuildEntities(const FMassArchetypeEntityCollection
 		Composition.SharedFragments.Add(*SharedFragment.GetScriptStruct());
 	}
 
-	BatchBuildEntities(EncodedEntitiesWithPayload, MoveTemp(Composition), SharedFragmentValues, ArchetypeDebugName);
+	BatchBuildEntities(EncodedEntitiesWithPayload, MoveTemp(Composition), SharedFragmentValues, CreationParams);
 }
 
 void FMassEntityManager::BatchBuildEntities(const FMassArchetypeEntityCollectionWithPayload& EncodedEntitiesWithPayload, FMassArchetypeCompositionDescriptor&& Composition
-	, const FMassArchetypeSharedFragmentValues& SharedFragmentValues, const FName ArchetypeDebugName)
+	, const FMassArchetypeSharedFragmentValues& SharedFragmentValues, const FMassArchetypeCreationParams& CreationParams)
 {
 	checkf(IsProcessing() == false, TEXT("Synchronous API function %hs called during mass processing. Use asynchronous API instead."), __FUNCTION__);
 
@@ -487,7 +489,7 @@ void FMassEntityManager::BatchBuildEntities(const FMassArchetypeEntityCollection
 	FMassArchetypeEntityCollection::FEntityRangeArray TargetArchetypeEntityRanges;
 
 	// "built" entities case, this is verified during FMassArchetypeEntityCollectionWithPayload construction
-	FMassArchetypeHandle TargetArchetypeHandle = CreateArchetype(Composition, ArchetypeDebugName);
+	FMassArchetypeHandle TargetArchetypeHandle = CreateArchetype(Composition, CreationParams);
 	check(TargetArchetypeHandle.IsValid());
 
 	// there are some extra steps in creating EncodedEntities from the original given entity handles and then back
@@ -1410,7 +1412,7 @@ void FMassEntityManager::FlushCommands(const TSharedPtr<FMassCommandBuffer>& InC
 
 	if (InCommandBuffer)
 	{
-		if(InCommandBuffer->HasPendingCommands())
+		if (InCommandBuffer->HasPendingCommands())
 		{
 			FlushedCommandBufferQueue.Enqueue(InCommandBuffer);
 		}
@@ -1554,3 +1556,58 @@ void FMassEntityManager::DebugGetArchetypeStrings(const FMassArchetypeHandle& Ar
 }
 
 #endif // WITH_MASSENTITY_DEBUG
+
+
+//-----------------------------------------------------------------------------
+// DEPRECATED
+//-----------------------------------------------------------------------------
+FMassArchetypeHandle FMassEntityManager::CreateArchetype(TConstArrayView<const UScriptStruct*> FragmentsAndTagsList, const FName ArchetypeDebugName)
+{
+	FMassArchetypeCreationParams Params;
+	Params.DebugName = ArchetypeDebugName;
+	return CreateArchetype(FragmentsAndTagsList, Params);
+}
+
+FMassArchetypeHandle FMassEntityManager::CreateArchetype(FMassArchetypeHandle SourceArchetype, TConstArrayView<const UScriptStruct*> FragmentsAndTagsList, const FName ArchetypeDebugName)
+{
+	FMassArchetypeCreationParams Params;
+	Params.DebugName = ArchetypeDebugName;
+	return CreateArchetype(SourceArchetype, FragmentsAndTagsList, Params);
+}
+
+FMassArchetypeHandle FMassEntityManager::CreateArchetype(const FMassArchetypeCompositionDescriptor& Composition, const FName ArchetypeDebugName)
+{
+	FMassArchetypeCreationParams Params;
+	Params.DebugName = ArchetypeDebugName;
+	return CreateArchetype(Composition, Params);
+}
+
+FMassArchetypeHandle FMassEntityManager::CreateArchetype(const TSharedPtr<FMassArchetypeData>& SourceArchetype, const FMassFragmentBitSet& InFragments, const FName ArchetypeDebugName)
+{
+	FMassArchetypeCreationParams Params;
+	Params.DebugName = ArchetypeDebugName;
+	return CreateArchetype(SourceArchetype, InFragments, Params);
+}
+
+FMassEntityHandle FMassEntityManager::CreateEntity(TConstArrayView<FInstancedStruct> FragmentInstanceList, const FMassArchetypeSharedFragmentValues& SharedFragmentValues, const FName ArchetypeDebugName)
+{
+	FMassArchetypeCreationParams Params;
+	Params.DebugName = ArchetypeDebugName;
+	return CreateEntity(FragmentInstanceList, SharedFragmentValues, Params);
+}
+
+void FMassEntityManager::BatchBuildEntities(const FMassArchetypeEntityCollectionWithPayload& EncodedEntitiesWithPayload, const FMassFragmentBitSet& FragmentsAffected
+	, const FMassArchetypeSharedFragmentValues& SharedFragmentValues, const FName ArchetypeDebugName)
+{
+	FMassArchetypeCreationParams Params;
+	Params.DebugName = ArchetypeDebugName;
+	BatchBuildEntities(EncodedEntitiesWithPayload, FragmentsAffected, SharedFragmentValues, Params);
+}
+
+void FMassEntityManager::BatchBuildEntities(const FMassArchetypeEntityCollectionWithPayload& EncodedEntitiesWithPayload, FMassArchetypeCompositionDescriptor&& Composition
+	, const FMassArchetypeSharedFragmentValues& SharedFragmentValues, const FName ArchetypeDebugName)
+{
+	FMassArchetypeCreationParams Params;
+	Params.DebugName = ArchetypeDebugName;
+	BatchBuildEntities(EncodedEntitiesWithPayload, MoveTemp(Composition), SharedFragmentValues, Params);
+}
diff --git a/Engine/Plugins/Runtime/MassEntity/Source/MassEntity/Private/MassEntitySettings.cpp b/Engine/Plugins/Runtime/MassEntity/Source/MassEntity/Private/MassEntitySettings.cpp
--- a/Engine/Plugins/Runtime/MassEntity/Source/MassEntity/Private/MassEntitySettings.cpp
+++ b/Engine/Plugins/Runtime/MassEntity/Source/MassEntity/Private/MassEntitySettings.cpp
@@ -5,6 +5,7 @@
 #include "VisualLogger/VisualLogger.h"
 #include "UObject/UObjectHash.h"
 #include "Misc/CoreDelegates.h"
+#include "MassArchetypeData.h"
 
 #include UE_INLINE_GENERATED_CPP_BY_NAME(MassEntitySettings)
 
@@ -27,6 +28,12 @@ UMassEntitySettings::UMassEntitySettings(const FObjectInitializer& ObjectInitial
 	FCoreDelegates::OnPostEngineInit.AddUObject(this, &UMassEntitySettings::OnPostEngineInit);
 }
 
+void UMassEntitySettings::PostInitProperties()
+{
+	Super::PostInitProperties();
+	ChunkMemorySize = UE::Mass::SanitizeChunkMemorySize(ChunkMemorySize);
+}
+
 void UMassEntitySettings::BeginDestroy()
 {
 	FCoreDelegates::OnPostEngineInit.RemoveAll(this);
@@ -165,6 +172,7 @@ TConstArrayView<FMassProcessingPhaseConfig> UMassEntitySettings::GetProcessingPh
 void UMassEntitySettings::PostEditChangeProperty(FPropertyChangedEvent& PropertyChangedEvent)
 {
 	static const FName ProcessorCDOsName = GET_MEMBER_NAME_CHECKED(UMassEntitySettings, ProcessorCDOs);
+	static const FName ChunkMemorySizeName = GET_MEMBER_NAME_CHECKED(UMassEntitySettings, ChunkMemorySize);
 
 	Super::PostEditChangeProperty(PropertyChangedEvent);
 
@@ -181,6 +189,10 @@ void UMassEntitySettings::PostEditChangeProperty(FPropertyChangedEvent& Property
 		{
 			BuildProcessorList();
 		}
+		else if (PropName == ChunkMemorySizeName)
+		{
+			ChunkMemorySize = UE::Mass::SanitizeChunkMemorySize(ChunkMemorySize);
+		}
 
 		BuildPhases();
 		OnSettingsChange.Broadcast(PropertyChangedEvent);
diff --git a/Engine/Plugins/Runtime/MassEntity/Source/MassEntity/Public/MassEntityManager.h b/Engine/Plugins/Runtime/MassEntity/Source/MassEntity/Public/MassEntityManager.h
--- a/Engine/Plugins/Runtime/MassEntity/Source/MassEntity/Public/MassEntityManager.h
+++ b/Engine/Plugins/Runtime/MassEntity/Source/MassEntity/Public/MassEntityManager.h
@@ -102,7 +102,7 @@ public:
 	 * A special, relaxed but slower version of CreateArchetype functions that allows FragmentAngTagsList to contain 
 	 * both fragments and tags. 
 	 */
-	FMassArchetypeHandle CreateArchetype(TConstArrayView<const UScriptStruct*> FragmentsAndTagsList, const FName ArchetypeDebugName = FName());
+	FMassArchetypeHandle CreateArchetype(TConstArrayView<const UScriptStruct*> FragmentsAndTagsList, const FMassArchetypeCreationParams& CreationParams = FMassArchetypeCreationParams());
 
 	/**
 	 * A special, relaxed but slower version of CreateArchetype functions that allows FragmentAngTagsList to contain
@@ -110,29 +110,29 @@ public:
 	 * provided list if they're not already in the original archetype.
 	 */
 	FMassArchetypeHandle CreateArchetype(FMassArchetypeHandle SourceArchetype, 
-		TConstArrayView<const UScriptStruct*> FragmentsAndTagsList, const FName ArchetypeDebugName = FName());
+		TConstArrayView<const UScriptStruct*> FragmentsAndTagsList, const FMassArchetypeCreationParams& CreationParams = FMassArchetypeCreationParams());
 
 	/**
 	 * CreateArchetype from a composition descriptor and initial values
 	 *
 	 * @param Composition of fragment, tag and chunk fragment types
-	 * @param ArchetypeDebugName Name to identify the archetype while debugging
+	 * @param CreationParams Parameters used during archetype construction
 	 * @return a handle of a new archetype 
 	 */
-	FMassArchetypeHandle CreateArchetype(const FMassArchetypeCompositionDescriptor& Composition, const FName ArchetypeDebugName = FName());
+	FMassArchetypeHandle CreateArchetype(const FMassArchetypeCompositionDescriptor& Composition, const FMassArchetypeCreationParams& CreationParams = FMassArchetypeCreationParams());
 
 	/** 
 	 *  Creates an archetype like SourceArchetype + InFragments. 
 	 *  @param SourceArchetype the archetype used to initially populate the list of fragments of the archetype being created. 
 	 *  @param InFragments list of unique fragments to add to fragments fetched from SourceArchetype. Note that 
 	 *   adding an empty list is not supported and doing so will result in failing a `check`
-	 *  @param ArchetypeDebugName Name to identify the archetype while debugging
+	 *  @param CreationParams Parameters used during archetype construction
 	 *  @return a handle of a new archetype
 	 *  @note it's caller's responsibility to ensure that NewFragmentList is not empty and contains only fragment
 	 *   types that SourceArchetype doesn't already have. If the caller cannot guarantee it use of AddFragment functions
 	 *   family is recommended.
 	 */
-	FMassArchetypeHandle CreateArchetype(const TSharedPtr<FMassArchetypeData>& SourceArchetype, const FMassFragmentBitSet& InFragments, const FName ArchetypeDebugName = FName());
+	FMassArchetypeHandle CreateArchetype(const TSharedPtr<FMassArchetypeData>& SourceArchetype, const FMassFragmentBitSet& InFragments, const FMassArchetypeCreationParams& CreationParams = FMassArchetypeCreationParams());
 
 	/** Fetches the archetype for a given Entity. If Entity is not valid it will still return a handle, just with an invalid archetype */
 	FMassArchetypeHandle GetArchetypeForEntity(FMassEntityHandle Entity) const;
@@ -164,7 +164,7 @@ public:
 	 * @param SharedFragmentValues to be associated with the entity
 	 * @param ArchetypeDebugName Name to identify the archetype while debugging
 	 * @return FMassEntityHandle id of the newly created entity */
-	FMassEntityHandle CreateEntity(TConstArrayView<FInstancedStruct> FragmentInstanceList, const FMassArchetypeSharedFragmentValues& SharedFragmentValues = {}, const FName ArchetypeDebugName = FName());
+	FMassEntityHandle CreateEntity(TConstArrayView<FInstancedStruct> FragmentInstanceList, const FMassArchetypeSharedFragmentValues& SharedFragmentValues = {}, const FMassArchetypeCreationParams& CreationParams = FMassArchetypeCreationParams());
 
 	/**
 	 * A dedicated structure for ensuring the "on entities creation" observers get notified only once all other 
@@ -265,8 +265,10 @@ public:
 	void RemoveTagFromEntity(FMassEntityHandle Entity, const UScriptStruct* TagType);
 	void SwapTagsForEntity(FMassEntityHandle Entity, const UScriptStruct* FromFragmentType, const UScriptStruct* ToFragmentType);
 
-	void BatchBuildEntities(const FMassArchetypeEntityCollectionWithPayload& EncodedEntitiesWithPayload, const FMassFragmentBitSet& FragmentsAffected, const FMassArchetypeSharedFragmentValues& SharedFragmentValues = {}, const FName ArchetypeDebugName = FName());
-	void BatchBuildEntities(const FMassArchetypeEntityCollectionWithPayload& EncodedEntitiesWithPayload, FMassArchetypeCompositionDescriptor&& Composition, const FMassArchetypeSharedFragmentValues& SharedFragmentValues = {}, const FName ArchetypeDebugName = FName());
+	void BatchBuildEntities(const FMassArchetypeEntityCollectionWithPayload& EncodedEntitiesWithPayload, const FMassFragmentBitSet& FragmentsAffected
+		, const FMassArchetypeSharedFragmentValues& SharedFragmentValues = {}, const FMassArchetypeCreationParams& CreationParams = FMassArchetypeCreationParams());
+	void BatchBuildEntities(const FMassArchetypeEntityCollectionWithPayload& EncodedEntitiesWithPayload, FMassArchetypeCompositionDescriptor&& Composition
+		, const FMassArchetypeSharedFragmentValues& SharedFragmentValues = {}, const FMassArchetypeCreationParams& CreationParams = FMassArchetypeCreationParams());
 	void BatchChangeTagsForEntities(TConstArrayView<FMassArchetypeEntityCollection> EntityCollections, const FMassTagBitSet& TagsToAdd, const FMassTagBitSet& TagsToRemove);
 	void BatchChangeFragmentCompositionForEntities(TConstArrayView<FMassArchetypeEntityCollection> EntityCollections, const FMassFragmentBitSet& FragmentsToAdd, const FMassFragmentBitSet& FragmentsToRemove);
 	void BatchAddFragmentInstancesForEntities(TConstArrayView<FMassArchetypeEntityCollectionWithPayload> EntityCollections, const FMassFragmentBitSet& FragmentsAffected);
@@ -508,6 +510,28 @@ private:
 	FOnNewArchetypeDelegate OnNewArchetypeEvent;
 
 	bool bInitialized = false;
+
+
+	//-----------------------------------------------------------------------------
+	// DEPRECATED
+	//-----------------------------------------------------------------------------
+public:
+	UE_DEPRECATED(5.3, "This Flavor of CreateArchetype is deprecated. Use the one with FMassArchetypeCreationParams parameter instead.")
+	FMassArchetypeHandle CreateArchetype(TConstArrayView<const UScriptStruct*> FragmentsAndTagsList, const FName ArchetypeDebugName);
+	UE_DEPRECATED(5.3, "This Flavor of CreateArchetype is deprecated. Use the one with FMassArchetypeCreationParams parameter instead.")
+	FMassArchetypeHandle CreateArchetype(FMassArchetypeHandle SourceArchetype, TConstArrayView<const UScriptStruct*> FragmentsAndTagsList, const FName ArchetypeDebugName);
+	UE_DEPRECATED(5.3, "This Flavor of CreateArchetype is deprecated. Use the one with FMassArchetypeCreationParams parameter instead.")
+	FMassArchetypeHandle CreateArchetype(const FMassArchetypeCompositionDescriptor& Composition, const FName ArchetypeDebugName);
+	UE_DEPRECATED(5.3, "This Flavor of CreateArchetype is deprecated. Use the one with FMassArchetypeCreationParams parameter instead.")
+	FMassArchetypeHandle CreateArchetype(const TSharedPtr<FMassArchetypeData>& SourceArchetype, const FMassFragmentBitSet& InFragments, const FName ArchetypeDebugName);
+	UE_DEPRECATED(5.3, "This Flavor of CreateEntity is deprecated. Use the one with FMassArchetypeCreationParams parameter instead.")
+	FMassEntityHandle CreateEntity(TConstArrayView<FInstancedStruct> FragmentInstanceList, const FMassArchetypeSharedFragmentValues& SharedFragmentValues, const FName ArchetypeDebugName);
+	UE_DEPRECATED(5.3, "This Flavor of BatchBuildEntities is deprecated. Use the one with FMassArchetypeCreationParams parameter instead.")
+	void BatchBuildEntities(const FMassArchetypeEntityCollectionWithPayload& EncodedEntitiesWithPayload, const FMassFragmentBitSet& FragmentsAffected
+		, const FMassArchetypeSharedFragmentValues& SharedFragmentValues, const FName ArchetypeDebugName);
+	UE_DEPRECATED(5.3, "This Flavor of BatchBuildEntities is deprecated. Use the one with FMassArchetypeCreationParams parameter instead.")
+	void BatchBuildEntities(const FMassArchetypeEntityCollectionWithPayload& EncodedEntitiesWithPayload, FMassArchetypeCompositionDescriptor&& Composition
+		, const FMassArchetypeSharedFragmentValues& SharedFragmentValues, const FName ArchetypeDebugName);
 };
 
 #if UE_ENABLE_INCLUDE_ORDER_DEPRECATED_IN_5_2
diff --git a/Engine/Plugins/Runtime/MassEntity/Source/MassEntity/Public/MassEntitySettings.h b/Engine/Plugins/Runtime/MassEntity/Source/MassEntity/Public/MassEntitySettings.h
--- a/Engine/Plugins/Runtime/MassEntity/Source/MassEntity/Public/MassEntitySettings.h
+++ b/Engine/Plugins/Runtime/MassEntity/Source/MassEntity/Public/MassEntitySettings.h
@@ -42,6 +42,7 @@ protected:
 	virtual void PostEditChangeChainProperty(struct FPropertyChangedChainEvent& PropertyChangedEvent) override;
 #endif // WITH_EDITOR
 
+	virtual void PostInitProperties() override;
 	virtual void BeginDestroy() override;
 
 	void OnPostEngineInit();
@@ -49,7 +50,7 @@ protected:
 	void BuildProcessorList();
 
 public:
-	UPROPERTY(EditDefaultsOnly, Category = Mass, config, meta = (ClampMin = "8", UIMin = "8"), AdvancedDisplay)
+	UPROPERTY(EditDefaultsOnly, Category = Mass, config, AdvancedDisplay)
 	int32 ChunkMemorySize = 128 * 1024;
 
 	/**
diff --git a/Engine/Plugins/Runtime/MassEntity/Source/MassEntity/Public/MassEntityTypes.h b/Engine/Plugins/Runtime/MassEntity/Source/MassEntity/Public/MassEntityTypes.h
--- a/Engine/Plugins/Runtime/MassEntity/Source/MassEntity/Public/MassEntityTypes.h
+++ b/Engine/Plugins/Runtime/MassEntity/Source/MassEntity/Public/MassEntityTypes.h
@@ -553,3 +553,13 @@ namespace UE::Mass
 	};
 
 } // UE::Mass
+
+
+struct FMassArchetypeCreationParams
+{
+	/** Created archetype will have chunks of this size. 0 denotes "use default" (see UE::Mass::ChunkSize) */
+	int32 ChunkMemorySize = 0;
+
+	/** Name to identify the archetype while debugging*/
+	FName DebugName;
+};
diff --git a/Engine/Plugins/Runtime/MassGameplay/Source/MassSpawner/Private/MassEntityTemplate.cpp b/Engine/Plugins/Runtime/MassGameplay/Source/MassSpawner/Private/MassEntityTemplate.cpp
--- a/Engine/Plugins/Runtime/MassGameplay/Source/MassSpawner/Private/MassEntityTemplate.cpp
+++ b/Engine/Plugins/Runtime/MassGameplay/Source/MassSpawner/Private/MassEntityTemplate.cpp
@@ -102,7 +102,8 @@ FMassEntityTemplate::FMassEntityTemplate(const FMassEntityTemplateData& InData,
 	// Sort anything there is to sort for later comparison purposes
 	TemplateData.Sort();
 
-	const FMassArchetypeHandle ArchetypeHandle = EntityManager.CreateArchetype(GetCompositionDescriptor(), FName(GetTemplateName()));
+	TemplateData.GetArchetypeCreationParams().DebugName = FName(GetTemplateName());
+	const FMassArchetypeHandle ArchetypeHandle = EntityManager.CreateArchetype(GetCompositionDescriptor(), TemplateData.GetArchetypeCreationParams());
 	SetArchetype(ArchetypeHandle);
 }
 
@@ -113,7 +114,8 @@ FMassEntityTemplate::FMassEntityTemplate(FMassEntityTemplateData&& InData, FMass
 	// Sort anything there is to sort for later comparison purposes
 	TemplateData.Sort();
 
-	const FMassArchetypeHandle ArchetypeHandle = EntityManager.CreateArchetype(GetCompositionDescriptor(), FName(GetTemplateName()));
+	TemplateData.GetArchetypeCreationParams().DebugName = FName(GetTemplateName());
+	const FMassArchetypeHandle ArchetypeHandle = EntityManager.CreateArchetype(GetCompositionDescriptor(), TemplateData.GetArchetypeCreationParams());
 	SetArchetype(ArchetypeHandle);
 }
 
diff --git a/Engine/Plugins/Runtime/MassGameplay/Source/MassSpawner/Public/MassEntityTemplate.h b/Engine/Plugins/Runtime/MassGameplay/Source/MassSpawner/Public/MassEntityTemplate.h
--- a/Engine/Plugins/Runtime/MassGameplay/Source/MassSpawner/Public/MassEntityTemplate.h
+++ b/Engine/Plugins/Runtime/MassGameplay/Source/MassSpawner/Public/MassEntityTemplate.h
@@ -277,6 +277,8 @@ struct MASSSPAWNER_API FMassEntityTemplateData
 	 *  @Note that the function can be slow, depending on how elaborate the template is. This function is meant for debugging purposes. */
 	bool SlowIsEquivalent(const FMassEntityTemplateData& Other) const;
 
+	FMassArchetypeCreationParams& GetArchetypeCreationParams() { return CreationParams; }
+	
 protected:
 	FMassArchetypeCompositionDescriptor Composition;
 	FMassArchetypeSharedFragmentValues SharedFragmentValues;
@@ -287,6 +289,8 @@ protected:
 	// These functions will be called to initialize entity's UObject-based fragments
 	TArray<FObjectFragmentInitializerFunction> ObjectInitializers;
 
+	FMassArchetypeCreationParams CreationParams;
+
 	FString TemplateName;
 };
 
