Changed how UMassCompositeProcessor::Populate works to ensure the existing processors do not get overridden. The function got renamed to UMassCompositeProcessor::UpdateProcessorsCollection, the old function deprecated.

diff --git a/Engine/Plugins/Runtime/MassEntity/Source/MassEntity/Private/MassProcessingPhaseManager.cpp b/Engine/Plugins/Runtime/MassEntity/Source/MassEntity/Private/MassProcessingPhaseManager.cpp
--- a/Engine/Plugins/Runtime/MassEntity/Source/MassEntity/Private/MassProcessingPhaseManager.cpp
+++ b/Engine/Plugins/Runtime/MassEntity/Source/MassEntity/Private/MassProcessingPhaseManager.cpp
@@ -158,7 +158,7 @@ void FMassPhaseProcessorConfigurationHelper::Configure(TArrayView<UMassProcessor
 
 	Solver.ResolveDependencies(SortedProcessors, EntityManager, OutOptionalResult);
 
-	PhaseProcessor.Populate(SortedProcessors);
+	PhaseProcessor.UpdateProcessorsCollection(SortedProcessors);
 
 #if WITH_MASSENTITY_DEBUG
 	for (const FMassProcessorOrderInfo& ProcessorOrderInfo : SortedProcessors)
diff --git a/Engine/Plugins/Runtime/MassEntity/Source/MassEntity/Private/MassProcessingTypes.cpp b/Engine/Plugins/Runtime/MassEntity/Source/MassEntity/Private/MassProcessingTypes.cpp
--- a/Engine/Plugins/Runtime/MassEntity/Source/MassEntity/Private/MassProcessingTypes.cpp
+++ b/Engine/Plugins/Runtime/MassEntity/Source/MassEntity/Private/MassProcessingTypes.cpp
@@ -58,8 +58,11 @@ void FMassRuntimePipeline::Initialize(UObject& Owner)
 	{
 		if (Proc)
 		{
-			REDIRECT_OBJECT_TO_VLOG(Proc, &Owner);
-			Proc->Initialize(Owner);
+			if (Proc->IsInitialized() == false)
+			{
+				REDIRECT_OBJECT_TO_VLOG(Proc, &Owner);
+				Proc->Initialize(Owner);
+			}
 		}
 		else
 		{
diff --git a/Engine/Plugins/Runtime/MassEntity/Source/MassEntity/Private/MassProcessor.cpp b/Engine/Plugins/Runtime/MassEntity/Source/MassEntity/Private/MassProcessor.cpp
--- a/Engine/Plugins/Runtime/MassEntity/Source/MassEntity/Private/MassProcessor.cpp
+++ b/Engine/Plugins/Runtime/MassEntity/Source/MassEntity/Private/MassProcessor.cpp
@@ -357,7 +357,7 @@ void UMassCompositeProcessor::SetProcessors(TArrayView<UMassProcessor*> InProces
 	TArray<FMassProcessorOrderInfo> SortedProcessors;
 	Solver.ResolveDependencies(SortedProcessors, EntityManager);
 
-	Populate(SortedProcessors);
+	UpdateProcessorsCollection(SortedProcessors);
 
 	if (Solver.IsSolvingForSingleThread() == false)
 	{
@@ -420,21 +420,41 @@ void UMassCompositeProcessor::BuildFlatProcessingGraph(TConstArrayView<FMassProc
 #endif // WITH_MASSENTITY_DEBUG
 }
 
-void UMassCompositeProcessor::Populate(TConstArrayView<FMassProcessorOrderInfo> OrderedProcessors)
+void UMassCompositeProcessor::UpdateProcessorsCollection(TArrayView<FMassProcessorOrderInfo> InOutOrderedProcessors)
 {
+	TArray<TObjectPtr<UMassProcessor>> ExistingProcessors(ChildPipeline.GetMutableProcessors());
 	ChildPipeline.Reset();
 
 	const UWorld* World = GetWorld();
 	const EProcessorExecutionFlags WorldExecutionFlags = World ? UE::Mass::Utils::GetProcessorExecutionFlagsForWorld(*World) : EProcessorExecutionFlags::All;
 	const FMassProcessingPhaseConfig& PhaseConfig = GET_MASS_CONFIG_VALUE(GetProcessingPhaseConfig(ProcessingPhase));
 
-	for (const FMassProcessorOrderInfo& ProcessorInfo : OrderedProcessors)
+	for (FMassProcessorOrderInfo& ProcessorInfo : InOutOrderedProcessors)
 	{
 		if (ensureMsgf(ProcessorInfo.NodeType == FMassProcessorOrderInfo::EDependencyNodeType::Processor, TEXT("Encountered unexpected FMassProcessorOrderInfo::EDependencyNodeType while populating %s"), *GetGroupName().ToString()))
 		{
 			checkSlow(ProcessorInfo.Processor);
 			if (ProcessorInfo.Processor->ShouldExecute(WorldExecutionFlags))
 			{
+				// we want to reuse existing processors to maintain state. It's recommended to keep processors state-less
+				// but we already have processors that do have some state, like signaling processors.
+				// the following search only makes sense for "single instance" processors
+				if (ProcessorInfo.Processor->ShouldAllowMultipleInstances() == false)
+				{
+					TObjectPtr<UMassProcessor>* FoundProcessor = ExistingProcessors.FindByPredicate([ProcessorClass = ProcessorInfo.Processor->GetClass()](TObjectPtr<UMassProcessor>& Element)
+						{
+							return Element && (Element->GetClass() == ProcessorClass);
+						});
+
+					if (FoundProcessor)
+					{
+						// overriding the stored value since the InOutOrderedProcessors can get used after the call and it 
+						// needs to reflect the actual work performed
+						ProcessorInfo.Processor = FoundProcessor->Get();
+					}
+				}
+
+				CA_ASSUME(ProcessorInfo.Processor);
 				ChildPipeline.AppendProcessor(*ProcessorInfo.Processor);
 			}
 		}
@@ -515,3 +535,11 @@ UMassCompositeProcessor* UMassCompositeProcessor::FindOrAddGroupProcessor(FName
 	return GroupProcessor;
 }
 
+//-----------------------------------------------------------------------------
+// DEPRECATED
+//-----------------------------------------------------------------------------
+void UMassCompositeProcessor::Populate(TConstArrayView<FMassProcessorOrderInfo> OrderedProcessors)
+{
+	TArray<FMassProcessorOrderInfo> OrderedProcessorsCopy(OrderedProcessors);
+	UpdateProcessorsCollection(OrderedProcessorsCopy);
+}
diff --git a/Engine/Plugins/Runtime/MassEntity/Source/MassEntity/Public/MassProcessor.h b/Engine/Plugins/Runtime/MassEntity/Source/MassEntity/Public/MassProcessor.h
--- a/Engine/Plugins/Runtime/MassEntity/Source/MassEntity/Public/MassProcessor.h
+++ b/Engine/Plugins/Runtime/MassEntity/Source/MassEntity/Public/MassProcessor.h
@@ -47,7 +47,8 @@ public:
 	UMassProcessor();
 	UMassProcessor(const FObjectInitializer& ObjectInitializer);
 
-	virtual void Initialize(UObject& Owner) {}
+	bool IsInitialized() const { return bInitialized; }
+	virtual void Initialize(UObject& Owner) { bInitialized = true; }
 	virtual FGraphEventRef DispatchProcessorTasks(const TSharedPtr<FMassEntityManager>& EntityManager, FMassExecutionContext& ExecutionContext, const FGraphEventArray& Prerequisites = FGraphEventArray());
 
 	EProcessorExecutionFlags GetExecutionFlags() const { return (EProcessorExecutionFlags)ExecutionFlags; }
@@ -166,6 +167,9 @@ private:
 	 *  @note that it's safe to store pointers here since RegisterQuery does verify that a given registered query is 
 	 *  a member variable of a given processor */
 	TArray<FMassEntityQuery*> OwnedQueries;
+
+	/** Used to track whether Initialized has been called. */
+	bool bInitialized = false;
 };
 
 
@@ -206,9 +210,13 @@ public:
 	virtual void BuildFlatProcessingGraph(TConstArrayView<FMassProcessorOrderInfo> SortedProcessors);
 
 	/**
-	 *  Adds processors in OrderedProcessors to ChildPipeline
+	 * Adds processors in InOutOrderedProcessors to ChildPipeline. 
+	 * Note that this operation is non-destructive for the existing processors - the ones of classes found in InOutOrderedProcessors 
+	 * will be retained and used instead of the instances provided via InOutOrderedProcessors. Respective entries in InOutOrderedProcessors
+	 * will be updated to reflect the reuse.
+	 * The described behavior however is available only for processors with bAllowMultipleInstances == false.
 	 */
-	void Populate(TConstArrayView<FMassProcessorOrderInfo> OrderedProcessors);
+	void UpdateProcessorsCollection(TArrayView<FMassProcessorOrderInfo> InOutOrderedProcessors);
 
 	/** adds SubProcessor to an appropriately named group. If RequestedGroupName == None then SubProcessor
 	 *  will be added directly to ChildPipeline. If not then the indicated group will be searched for in ChildPipeline 
@@ -262,4 +270,11 @@ protected:
 		}
 	};
 	TArray<FProcessorCompletion> CompletionStatus;
+
+	//-----------------------------------------------------------------------------
+	// DEPRECATED
+	//-----------------------------------------------------------------------------
+public:
+	UE_DEPRECATED(5.3, "Populate is deprecated. Please use UpdateProcessorsCollection instead.")
+	void Populate(TConstArrayView<FMassProcessorOrderInfo> OrderedProcessors);
 };
