Massaging Mass debugging console commands - renamed all `ai.debug.mass.*` commands to `mass.debug.*`
Added a "mass.debug.SelectEntity" debug command that works with FMassDebugger::OnEntitySelectedDelegate (and all of the subsequent synchronization)

diff --git a/Engine/Plugins/AI/MassAI/Source/MassAIBehavior/Private/MassLookAtProcessors.cpp b/Engine/Plugins/AI/MassAI/Source/MassAIBehavior/Private/MassLookAtProcessors.cpp
--- a/Engine/Plugins/AI/MassAI/Source/MassAIBehavior/Private/MassLookAtProcessors.cpp
+++ b/Engine/Plugins/AI/MassAI/Source/MassAIBehavior/Private/MassLookAtProcessors.cpp
@@ -20,6 +20,8 @@
 #include "Engine/World.h"
 #include "MassLODFragments.h"
 #include "MassGameplayExternalTraits.h"
+#include "MassDebugger.h"
+
 
 #define UNSAFE_FOR_MT 1
 
diff --git a/Engine/Plugins/AI/MassAI/Source/MassAIBehavior/Private/MassStateTreeProcessors.cpp b/Engine/Plugins/AI/MassAI/Source/MassAIBehavior/Private/MassStateTreeProcessors.cpp
--- a/Engine/Plugins/AI/MassAI/Source/MassAIBehavior/Private/MassStateTreeProcessors.cpp
+++ b/Engine/Plugins/AI/MassAI/Source/MassAIBehavior/Private/MassStateTreeProcessors.cpp
@@ -16,6 +16,8 @@
 #include "Engine/World.h"
 #include "MassBehaviorSettings.h"
 #include "VisualLogger/VisualLogger.h"
+#include "MassDebugger.h"
+
 
 CSV_DEFINE_CATEGORY(StateTreeProcessor, true);
 
diff --git a/Engine/Plugins/AI/MassAI/Source/MassAIBehavior/Private/MassZoneGraphAnnotationFragments.cpp b/Engine/Plugins/AI/MassAI/Source/MassAIBehavior/Private/MassZoneGraphAnnotationFragments.cpp
--- a/Engine/Plugins/AI/MassAI/Source/MassAIBehavior/Private/MassZoneGraphAnnotationFragments.cpp
+++ b/Engine/Plugins/AI/MassAI/Source/MassAIBehavior/Private/MassZoneGraphAnnotationFragments.cpp
@@ -4,6 +4,8 @@
 #include "MassEntityManager.h"
 #include "MassExecutionContext.h"
 #include "MassSimulationLOD.h"
+#include "MassDebugger.h"
+
 
 namespace UE::Mass::ZoneGraphAnnotations
 {
diff --git a/Engine/Plugins/AI/MassAI/Source/MassAIBehavior/Private/Tasks/MassZoneGraphFindEscapeTarget.cpp b/Engine/Plugins/AI/MassAI/Source/MassAIBehavior/Private/Tasks/MassZoneGraphFindEscapeTarget.cpp
--- a/Engine/Plugins/AI/MassAI/Source/MassAIBehavior/Private/Tasks/MassZoneGraphFindEscapeTarget.cpp
+++ b/Engine/Plugins/AI/MassAI/Source/MassAIBehavior/Private/Tasks/MassZoneGraphFindEscapeTarget.cpp
@@ -11,6 +11,8 @@
 #include "MassStateTreeExecutionContext.h"
 #include "ZoneGraphSettings.h"
 #include "StateTreeLinker.h"
+#include "MassDebugger.h"
+
 
 FMassZoneGraphFindEscapeTarget::FMassZoneGraphFindEscapeTarget()
 {
diff --git a/Engine/Plugins/AI/MassAI/Source/MassAIBehavior/Private/Tasks/MassZoneGraphPathFollowTask.cpp b/Engine/Plugins/AI/MassAI/Source/MassAIBehavior/Private/Tasks/MassZoneGraphPathFollowTask.cpp
--- a/Engine/Plugins/AI/MassAI/Source/MassAIBehavior/Private/Tasks/MassZoneGraphPathFollowTask.cpp
+++ b/Engine/Plugins/AI/MassAI/Source/MassAIBehavior/Private/Tasks/MassZoneGraphPathFollowTask.cpp
@@ -11,6 +11,8 @@
 #include "MassZoneGraphNavigationUtils.h"
 #include "ZoneGraphSubsystem.h"
 #include "StateTreeLinker.h"
+#include "MassDebugger.h"
+
 
 bool FMassZoneGraphPathFollowTask::Link(FStateTreeLinker& Linker)
 {
diff --git a/Engine/Plugins/AI/MassAI/Source/MassAIDebug/Private/MassDebugStateTreeProcessor.cpp b/Engine/Plugins/AI/MassAI/Source/MassAIDebug/Private/MassDebugStateTreeProcessor.cpp
--- a/Engine/Plugins/AI/MassAI/Source/MassAIDebug/Private/MassDebugStateTreeProcessor.cpp
+++ b/Engine/Plugins/AI/MassAI/Source/MassAIDebug/Private/MassDebugStateTreeProcessor.cpp
@@ -6,6 +6,7 @@
 #include "MassStateTreeExecutionContext.h"
 #include "MassStateTreeFragments.h"
 #include "MassCommonFragments.h"
+#include "MassDebugger.h"
 #include "MassExecutionContext.h"
 #include "VisualLogger/VisualLogger.h"
 #include "Engine/World.h"
diff --git a/Engine/Plugins/AI/MassAI/Source/MassNavigation/Private/MassAvoidanceProcessors.cpp b/Engine/Plugins/AI/MassAI/Source/MassNavigation/Private/MassAvoidanceProcessors.cpp
--- a/Engine/Plugins/AI/MassAI/Source/MassNavigation/Private/MassAvoidanceProcessors.cpp
+++ b/Engine/Plugins/AI/MassAI/Source/MassNavigation/Private/MassAvoidanceProcessors.cpp
@@ -15,6 +15,8 @@
 #include "MassNavigationFragments.h"
 #include "MassNavigationUtils.h"
 #include "Engine/World.h"
+#include "MassDebugger.h"
+
 
 #define UNSAFE_FOR_MT 1
 
diff --git a/Engine/Plugins/AI/MassAI/Source/MassZoneGraphNavigation/Private/MassZoneGraphNavigationProcessors.cpp b/Engine/Plugins/AI/MassAI/Source/MassZoneGraphNavigation/Private/MassZoneGraphNavigationProcessors.cpp
--- a/Engine/Plugins/AI/MassAI/Source/MassZoneGraphNavigation/Private/MassZoneGraphNavigationProcessors.cpp
+++ b/Engine/Plugins/AI/MassAI/Source/MassZoneGraphNavigation/Private/MassZoneGraphNavigationProcessors.cpp
@@ -15,6 +15,7 @@
 #include "VisualLogger/VisualLogger.h"
 #include "MassSimulationLOD.h"
 #include "Engine/World.h"
+#include "MassDebugger.h"
 
 #define UNSAFE_FOR_MT 1
 
diff --git a/Engine/Plugins/AI/MassAI/Source/MassZoneGraphNavigation/Private/MassZoneGraphNavigationUtils.cpp b/Engine/Plugins/AI/MassAI/Source/MassZoneGraphNavigation/Private/MassZoneGraphNavigationUtils.cpp
--- a/Engine/Plugins/AI/MassAI/Source/MassZoneGraphNavigation/Private/MassZoneGraphNavigationUtils.cpp
+++ b/Engine/Plugins/AI/MassAI/Source/MassZoneGraphNavigation/Private/MassZoneGraphNavigationUtils.cpp
@@ -6,6 +6,7 @@
 #include "MassZoneGraphNavigationFragments.h"
 #include "ZoneGraphSubsystem.h"
 #include "VisualLogger/VisualLogger.h"
+#include "MassDebugger.h"
 
 namespace UE::MassNavigation
 {
diff --git a/Engine/Plugins/Runtime/MassEntity/Source/MassEntity/Private/MassDebugger.cpp b/Engine/Plugins/Runtime/MassEntity/Source/MassEntity/Private/MassDebugger.cpp
--- a/Engine/Plugins/Runtime/MassEntity/Source/MassEntity/Private/MassDebugger.cpp
+++ b/Engine/Plugins/Runtime/MassEntity/Source/MassEntity/Private/MassDebugger.cpp
@@ -53,6 +53,92 @@ namespace UE::Mass::Debug
 		Ar.SetAutoEmitLineTerminator(bAutoLineEnd);
 	}
 
+	// First Id of a range of lightweight entity for which we want to activate debug information
+	int32 DebugEntityBegin = INDEX_NONE;
+
+	// Last Id of a range of lightweight entity for which we want to activate debug information
+	int32 DebugEntityEnd = INDEX_NONE;
+
+	void SetDebugEntityRange(const int32 InDebugEntityBegin, const int32 InDebugEntityEnd)
+	{
+		DebugEntityBegin = InDebugEntityBegin;
+		DebugEntityEnd = InDebugEntityBegin;
+	}
+
+	static FAutoConsoleCommand SetDebugEntityRangeCommand(
+		TEXT("mass.debug.SetDebugEntityRange"),
+		TEXT("Range of lightweight entity IDs that we want to debug.")
+		TEXT("Usage: \"mass.debug.SetDebugEntityRange <FirstEntity> <LastEntity>\""),
+		FConsoleCommandWithArgsDelegate::CreateLambda([](const TArray<FString>& Args)
+			{
+				if (Args.Num() != 2)
+				{
+					UE_LOG(LogConsoleResponse, Display, TEXT("Error: Expecting 2 parameters"));
+					return;
+				}
+
+				int32 FirstID = INDEX_NONE;
+				int32 LastID = INDEX_NONE;
+				if (!LexTryParseString<int32>(FirstID, *Args[0]))
+				{
+					UE_LOG(LogConsoleResponse, Display, TEXT("Error: first parameter must be an integer"));
+					return;
+				}
+			
+				if (!LexTryParseString<int32>(LastID, *Args[1]))
+				{
+					UE_LOG(LogConsoleResponse, Display, TEXT("Error: second parameter must be an integer"));
+					return;
+				}
+
+				SetDebugEntityRange(FirstID, LastID);
+			}));
+
+	static FAutoConsoleCommand ResetDebugEntity(
+		TEXT("mass.debug.ResetDebugEntity"),
+		TEXT("Disables lightweight entities debugging.")
+		TEXT("Usage: \"mass.debug.ResetDebugEntity\""),
+		FConsoleCommandWithArgsDelegate::CreateLambda([](const TArray<FString>& Args)
+			{
+				SetDebugEntityRange(INDEX_NONE, INDEX_NONE);
+			}));
+
+	bool HasDebugEntities()
+	{
+		return DebugEntityBegin != INDEX_NONE && DebugEntityEnd != INDEX_NONE;
+	}
+
+	bool IsDebuggingSingleEntity()
+	{
+		return DebugEntityBegin != INDEX_NONE && DebugEntityBegin == DebugEntityEnd;
+	}
+
+	bool GetDebugEntitiesRange(int32& OutBegin, int32& OutEnd)
+	{
+		OutBegin = DebugEntityBegin;
+		OutEnd = DebugEntityEnd;
+		return DebugEntityBegin != INDEX_NONE && DebugEntityEnd != INDEX_NONE && DebugEntityBegin <= DebugEntityEnd;
+	}
+	
+	bool IsDebuggingEntity(FMassEntityHandle Entity, FColor* OutEntityColor)
+	{
+		const int32 EntityIdx = Entity.Index;
+		const bool bIsDebuggingEntity = (DebugEntityBegin != INDEX_NONE && DebugEntityEnd != INDEX_NONE && DebugEntityBegin <= EntityIdx && EntityIdx <= DebugEntityEnd);
+	
+		if (bIsDebuggingEntity && OutEntityColor != nullptr)
+		{
+			*OutEntityColor = GetEntityDebugColor(Entity);
+		}
+
+		return bIsDebuggingEntity;
+	}
+
+	FColor GetEntityDebugColor(FMassEntityHandle Entity)
+	{
+		const int32 EntityIdx = Entity.Index;
+		return EntityIdx != INDEX_NONE ? GColorList.GetFColorByIndex(EntityIdx % GColorList.GetColorsNum()) : FColor::Black;
+	}
+
 	FAutoConsoleCommandWithWorldArgsAndOutputDevice PrintEntityFragmentsCmd(
 		TEXT("mass.PrintEntityFragments"),
 		TEXT("Prints all fragment types and values (uproperties) for the specified Entity index"),
@@ -192,9 +278,9 @@ namespace UE::Mass::Debug
 			}));
 
 	static FAutoConsoleCommandWithWorldAndArgs DestroyEntity(
-		TEXT("ai.debug.mass.DestroyEntity"),
-		TEXT("ID of a lightweight entity that we want to destroy.")
-		TEXT("Usage: \"ai.debug.mass.DestoryEntity <Entity>\""),
+		TEXT("mass.debug.DestroyEntity"),
+		TEXT("ID of a Mass entity that we want to destroy.")
+		TEXT("Usage: \"mass.debug.DestoryEntity <Entity>\""),
 		FConsoleCommandWithWorldAndArgsDelegate::CreateLambda([](const TArray<FString>& Args, UWorld* World)
 	{
 		if (Args.Num() != 1)
@@ -227,6 +313,45 @@ namespace UE::Mass::Debug
 		EntityManager.Defer().DestroyEntity(EntityToDestroy);
 	}));
 
+	static FAutoConsoleCommandWithWorldAndArgs SetDebugEntity(
+		TEXT("mass.debug.DebugEntity"),
+		TEXT("ID of a Mass entity that we want to debug.")
+		TEXT("Usage: \"mass.debug.DebugEntity <Entity>\""),
+		FConsoleCommandWithWorldAndArgsDelegate::CreateLambda([](const TArray<FString>& Args, UWorld* World)
+		{
+			if (Args.Num() != 1)
+			{
+				UE_LOG(LogConsoleResponse, Display, TEXT("Error: Expecting 1 parameter"));
+				return;
+			}
+
+			int32 ID = INDEX_NONE;
+			if (!LexTryParseString<int32>(ID, *Args[0]))
+			{
+				UE_LOG(LogConsoleResponse, Display, TEXT("Error: parameter must be an integer"));
+				return;
+			}
+
+			if (!World)
+			{
+				UE_LOG(LogConsoleResponse, Display, TEXT("Error: invalid world"));
+				return;
+			}
+
+			SetDebugEntityRange(ID, ID);
+
+			FMassEntityManager& EntityManager = UE::Mass::Utils::GetEntityManagerChecked(*World);
+			FMassEntityHandle EntityToDebug = EntityManager.DebugGetEntityIndexHandle(ID);
+			if (!EntityToDebug.IsSet())
+			{
+				UE_LOG(LogConsoleResponse, Display, TEXT("Error: cannot find entity for this index"));
+				return;
+			}
+
+			FMassDebugger::SelectEntity(EntityManager, EntityToDebug);
+		}
+	));
+
 } // namespace UE::Mass::Debug
 
 //----------------------------------------------------------------------//
@@ -299,6 +424,7 @@ void FMassDebugger::GetArchetypeEntityStats(const FMassArchetypeHandle& Archetyp
 	OutStats.EntitiesCountPerChunk = ArchetypeData.GetNumEntitiesPerChunk();
 	OutStats.ChunksCount = ArchetypeData.GetChunkCount();
 	OutStats.AllocatedSize = ArchetypeData.GetAllocatedSize();
+	OutStats.BytesPerEntity = ArchetypeData.GetBytesPerEntity();
 }
 
 const TConstArrayView<FName> FMassDebugger::GetArchetypeDebugNames(const FMassArchetypeHandle& ArchetypeHandle)
@@ -462,6 +588,7 @@ void FMassDebugger::OutputEntityDescription(FOutputDevice& Ar, const FMassEntity
 
 void FMassDebugger::SelectEntity(const FMassEntityManager& EntityManager, const FMassEntityHandle EntityHandle)
 {
+	UE::Mass::Debug::SetDebugEntityRange(EntityHandle.Index, EntityHandle.Index);
 	OnEntitySelectedDelegate.Broadcast(EntityManager, EntityHandle);
 }
 
diff --git a/Engine/Plugins/Runtime/MassEntity/Source/MassEntity/Public/MassDebugger.h b/Engine/Plugins/Runtime/MassEntity/Source/MassEntity/Public/MassDebugger.h
--- a/Engine/Plugins/Runtime/MassEntity/Source/MassEntity/Public/MassDebugger.h
+++ b/Engine/Plugins/Runtime/MassEntity/Source/MassEntity/Public/MassDebugger.h
@@ -8,7 +8,7 @@
 #include "MassEntityQuery.h"
 #include "MassProcessor.h"
 #include "Misc/SpinLock.h"
-#endif // WITH_MASSENTITY_DEBUG
+
 
 class FOutputDevice;
 class UMassProcessor;
@@ -22,38 +22,49 @@ enum class EMassFragmentPresence : uint8;
 
 namespace UE::Mass::Debug
 {
-#if WITH_MASSENTITY_DEBUG
-struct MASSENTITY_API FQueryRequirementsView
-{
-	TConstArrayView<FMassFragmentRequirementDescription> FragmentRequirements;
-	TConstArrayView<FMassFragmentRequirementDescription> ChunkRequirements;
-	TConstArrayView<FMassFragmentRequirementDescription> ConstSharedRequirements;
-	TConstArrayView<FMassFragmentRequirementDescription> SharedRequirements;
-	const FMassTagBitSet& RequiredAllTags;
-	const FMassTagBitSet& RequiredAnyTags;
-	const FMassTagBitSet& RequiredNoneTags;
-	const FMassExternalSubsystemBitSet& RequiredConstSubsystems;
-	const FMassExternalSubsystemBitSet& RequiredMutableSubsystems;
-};
-
-FString DebugGetFragmentAccessString(EMassFragmentAccess Access);
-MASSENTITY_API extern void DebugOutputDescription(TConstArrayView<UMassProcessor*> Processors, FOutputDevice& Ar);
-#endif // WITH_MASSENTITY_DEBUG
-
-struct FArchetypeStats
-{
-	/** Number of active entities of the archetype. */
-	int32 EntitiesCount = 0;
-	/** Number of entities that fit per chunk. */
-	int32 EntitiesCountPerChunk = 0;
-	/** Number of allocated chunks. */
-	int32 ChunksCount = 0;
-	/** Total amount of memory taken by this archetype */
-	SIZE_T AllocatedSize = 0;
-};
+	struct MASSENTITY_API FQueryRequirementsView
+	{
+		TConstArrayView<FMassFragmentRequirementDescription> FragmentRequirements;
+		TConstArrayView<FMassFragmentRequirementDescription> ChunkRequirements;
+		TConstArrayView<FMassFragmentRequirementDescription> ConstSharedRequirements;
+		TConstArrayView<FMassFragmentRequirementDescription> SharedRequirements;
+		const FMassTagBitSet& RequiredAllTags;
+		const FMassTagBitSet& RequiredAnyTags;
+		const FMassTagBitSet& RequiredNoneTags;
+		const FMassExternalSubsystemBitSet& RequiredConstSubsystems;
+		const FMassExternalSubsystemBitSet& RequiredMutableSubsystems;
+	};
+
+	FString DebugGetFragmentAccessString(EMassFragmentAccess Access);
+	MASSENTITY_API extern void DebugOutputDescription(TConstArrayView<UMassProcessor*> Processors, FOutputDevice& Ar);
+
+	struct FArchetypeStats
+	{
+		/** Number of active entities of the archetype. */
+		int32 EntitiesCount = 0;
+		/** Number of entities that fit per chunk. */
+		int32 EntitiesCountPerChunk = 0;
+		/** Number of allocated chunks. */
+		int32 ChunksCount = 0;
+		/** Total amount of memory taken by this archetype */
+		SIZE_T AllocatedSize = 0;
+		/** Total amount of memory needed by a single entity */
+		int32 BytesPerEntity = 0;
+	};
+
+	MASSENTITY_API extern bool HasDebugEntities();
+	MASSENTITY_API extern bool IsDebuggingSingleEntity();
+
+	/**
+	 * Populates OutBegin and OutEnd with entity index ranges as set by mass.debug.SetDebugEntityRange or
+	 * mass.debug.DebugEntity console commands.
+	 * @return whether any range has been configured.
+	 */
+	MASSENTITY_API extern bool GetDebugEntitiesRange(int32& OutBegin, int32& OutEnd);
+	MASSENTITY_API extern bool IsDebuggingEntity(FMassEntityHandle Entity, FColor* OutEntityColor = nullptr);
+	MASSENTITY_API extern FColor GetEntityDebugColor(FMassEntityHandle Entity);
 } // namespace UE::Mass::Debug
 
-#if WITH_MASSENTITY_DEBUG
 
 struct MASSENTITY_API FMassDebugger
 {
@@ -104,6 +115,10 @@ private:
 
 #else
 
+struct FMassArchetypeHandle;
+struct FMassFragmentRequirements;
+struct FMassFragmentRequirementDescription;
+
 struct MASSENTITY_API FMassDebugger
 {
 	static FString GetSingleRequirementDescription(const FMassFragmentRequirementDescription& Requirement) { return TEXT("[no debug information]"); }
diff --git a/Engine/Plugins/Runtime/MassGameplay/Source/MassCommon/Private/MassCommonTypes.cpp b/Engine/Plugins/Runtime/MassGameplay/Source/MassCommon/Private/MassCommonTypes.cpp
--- a/Engine/Plugins/Runtime/MassGameplay/Source/MassCommon/Private/MassCommonTypes.cpp
+++ b/Engine/Plugins/Runtime/MassGameplay/Source/MassCommon/Private/MassCommonTypes.cpp
@@ -1,114 +1,3 @@
 // Copyright Epic Games, Inc. All Rights Reserved.
 
 #include "MassCommonTypes.h"
-#include "MassEntityManager.h"
-#include "Math/ColorList.h"
-#include "HAL/IConsoleManager.h"
-
-#if WITH_MASSGAMEPLAY_DEBUG
-
-namespace UE::Mass::Debug
-{
-
-// First Id of a range of lightweight entity for which we want to activate debug information
-int32 DebugEntityBegin = INDEX_NONE;
-
-// Last Id of a range of lightweight entity for which we want to activate debug information
-int32 DebugEntityEnd = INDEX_NONE;
-
-static FAutoConsoleCommand SetDebugEntityRange(
-	TEXT("ai.debug.mass.SetDebugEntityRange"),
-	TEXT("Range of lightweight entity IDs that we want to debug.")
-	TEXT("Usage: \"ai.debug.mass.SetDebugEntityRange <FirstEntity> <LastEntity>\""),
-	FConsoleCommandWithArgsDelegate::CreateLambda([](const TArray<FString>& Args)
-		{
-			if (Args.Num() != 2)
-			{
-				UE_LOG(LogConsoleResponse, Display, TEXT("Error: Expecting 2 parameters"));
-				return;
-			}
-
-			int32 FirstID = INDEX_NONE;
-			int32 LastID = INDEX_NONE;
-			if (!LexTryParseString<int32>(FirstID, *Args[0]))
-			{
-				UE_LOG(LogConsoleResponse, Display, TEXT("Error: first parameter must be an integer"));
-				return;
-			}
-			
-			if (!LexTryParseString<int32>(LastID, *Args[1]))
-			{
-				UE_LOG(LogConsoleResponse, Display, TEXT("Error: second parameter must be an integer"));
-				return;
-			}
-
-			DebugEntityBegin = FirstID;
-			DebugEntityEnd = LastID;
-		}));
-
-static FAutoConsoleCommand SetDebugEntity(
-	TEXT("ai.debug.mass.DebugEntity"),
-	TEXT("ID of a lightweight entity that we want to debug.")
-	TEXT("Usage: \"ai.debug.mass.DebugEntity <Entity>\""),
-	FConsoleCommandWithArgsDelegate::CreateLambda([](const TArray<FString>& Args)
-		{
-			if (Args.Num() != 1)
-			{
-				UE_LOG(LogConsoleResponse, Display, TEXT("Error: Expecting 1 parameter"));
-				return;
-			}
-
-			int32 ID = INDEX_NONE;
-			if (!LexTryParseString<int32>(ID, *Args[0]))
-			{
-				UE_LOG(LogConsoleResponse, Display, TEXT("Error: parameter must be an integer"));
-				return;
-			}
-
-			DebugEntityBegin = ID;
-			DebugEntityEnd = ID;
-		}));
-
-static FAutoConsoleCommand ResetDebugEntity(
-	TEXT("ai.debug.mass.ResetDebugEntity"),
-	TEXT("Disables lightweight entities debugging.")
-	TEXT("Usage: \"ai.debug.mass.ResetDebugEntity\""),
-	FConsoleCommandWithArgsDelegate::CreateLambda([](const TArray<FString>& Args)
-		{
-			DebugEntityBegin = INDEX_NONE;
-			DebugEntityEnd = INDEX_NONE;
-		}));
-
-bool HasDebugEntities()
-{
-	return DebugEntityBegin != INDEX_NONE && DebugEntityEnd != INDEX_NONE;
-}
-
-bool GetDebugEntitiesRange(int32& OutBegin, int32& OutEnd)
-{
-	OutBegin = DebugEntityBegin;
-	OutEnd = DebugEntityEnd;
-	return DebugEntityBegin != INDEX_NONE && DebugEntityEnd != INDEX_NONE && DebugEntityBegin <= DebugEntityEnd;
-}
-	
-bool IsDebuggingEntity(FMassEntityHandle Entity, FColor* OutEntityColor)
-{
-	const int32 EntityIdx = Entity.Index;
-	const bool bIsDebuggingEntity = (DebugEntityBegin != INDEX_NONE && DebugEntityEnd != INDEX_NONE && DebugEntityBegin <= EntityIdx && EntityIdx <= DebugEntityEnd);
-	
-	if (bIsDebuggingEntity && OutEntityColor != nullptr)
-	{
-		*OutEntityColor = GetEntityDebugColor(Entity);
-	}
-
-	return bIsDebuggingEntity;
-}
-
-FColor GetEntityDebugColor(FMassEntityHandle Entity)
-{
-	const int32 EntityIdx = Entity.Index;
-	return EntityIdx != INDEX_NONE ? GColorList.GetFColorByIndex(EntityIdx % GColorList.GetColorsNum()) : FColor::Black;
-}
-
-} // namespace UE::Mass::Debug
-#endif // WITH_MASSGAMEPLAY_DEBUG
\ No newline at end of file
diff --git a/Engine/Plugins/Runtime/MassGameplay/Source/MassCommon/Public/MassCommonTypes.h b/Engine/Plugins/Runtime/MassGameplay/Source/MassCommon/Public/MassCommonTypes.h
--- a/Engine/Plugins/Runtime/MassGameplay/Source/MassCommon/Public/MassCommonTypes.h
+++ b/Engine/Plugins/Runtime/MassGameplay/Source/MassCommon/Public/MassCommonTypes.h
@@ -13,21 +13,6 @@
 class UStaticMesh;
 class UMaterialInterface;
 
-#if WITH_MASSGAMEPLAY_DEBUG
-namespace UE::Mass::Debug
-{
-	MASSCOMMON_API extern bool HasDebugEntities();
-	/**
-	 * Populates OutBegin and OutEnd with entity index ranges as set by ai.debug.mass.SetDebugEntityRange or
-	 * ai.debug.mass.DebugEntity console commands.
-	 * @return whether any range has been configured.
-	 */
-	MASSCOMMON_API extern bool GetDebugEntitiesRange(int32& OutBegin, int32& OutEnd);
-	MASSCOMMON_API extern bool IsDebuggingEntity(FMassEntityHandle Entity, FColor* OutEntityColor = nullptr);
-	MASSCOMMON_API extern FColor GetEntityDebugColor(FMassEntityHandle Entity);
-} // namespace UE::Mass::Debug
-#endif // WITH_MASSGAMEPLAY_DEBUG
-
 namespace UE::Mass::ProcessorGroupNames
 {
 	const FName UpdateWorldFromMass = FName(TEXT("UpdateWorldFromMass"));
diff --git a/Engine/Plugins/Runtime/MassGameplay/Source/MassGameplayDebug/Private/MassGameplayDebugTypes.cpp b/Engine/Plugins/Runtime/MassGameplay/Source/MassGameplayDebug/Private/MassGameplayDebugTypes.cpp
--- a/Engine/Plugins/Runtime/MassGameplay/Source/MassGameplayDebug/Private/MassGameplayDebugTypes.cpp
+++ b/Engine/Plugins/Runtime/MassGameplay/Source/MassGameplayDebug/Private/MassGameplayDebugTypes.cpp
@@ -4,9 +4,9 @@
 #include "MassCommonTypes.h"
 #include "MassCommonFragments.h"
 #include "MassEntityManager.h"
-
-
 #if WITH_MASSGAMEPLAY_DEBUG
+#include "MassDebugger.h"
+
 
 namespace UE::Mass::Debug
 {
diff --git a/Engine/Plugins/Runtime/MassGameplay/Source/MassGameplayDebug/Public/MassGameplayDebugTypes.h b/Engine/Plugins/Runtime/MassGameplay/Source/MassGameplayDebug/Public/MassGameplayDebugTypes.h
--- a/Engine/Plugins/Runtime/MassGameplay/Source/MassGameplayDebug/Public/MassGameplayDebugTypes.h
+++ b/Engine/Plugins/Runtime/MassGameplay/Source/MassGameplayDebug/Public/MassGameplayDebugTypes.h
@@ -24,7 +24,7 @@ namespace UE::Mass::Debug
 {
 	/**
 	 * Fetches entity handles and their locations for entities indicated by index range as set by
-	 * ai.debug.mass.SetDebugEntityRange or ai.debug.mass.DebugEntity console commands.
+	 * mass.debug.SetDebugEntityRange or mass.debug.DebugEntity console commands.
 	 */
 	MASSGAMEPLAYDEBUG_API extern void GetDebugEntitiesAndLocations(const FMassEntityManager& EntitySubsystem, TArray<FMassEntityHandle>& OutEntities, TArray<FVector>& OutLocations);
 	MASSGAMEPLAYDEBUG_API extern FMassEntityHandle ConvertEntityIndexToHandle(const FMassEntityManager& EntitySubsystem, const int32 EntityIndex);
diff --git a/Engine/Plugins/Runtime/MassGameplay/Source/MassMovement/Private/MassMovementTypes.cpp b/Engine/Plugins/Runtime/MassGameplay/Source/MassMovement/Private/MassMovementTypes.cpp
--- a/Engine/Plugins/Runtime/MassGameplay/Source/MassMovement/Private/MassMovementTypes.cpp
+++ b/Engine/Plugins/Runtime/MassGameplay/Source/MassMovement/Private/MassMovementTypes.cpp
@@ -1,4 +1,4 @@
-﻿// Copyright Epic Games, Inc. All Rights Reserved.
+// Copyright Epic Games, Inc. All Rights Reserved.
 
 #include "MassMovementTypes.h"
 #include "HAL/IConsoleManager.h"
@@ -6,6 +6,6 @@
 namespace UE::MassMovement
 {
 	int32 bFreezeMovement = 0;
-	FAutoConsoleVariableRef CVarFreezeMovement(TEXT("ai.debug.mass.FreezeMovement"), bFreezeMovement, TEXT("Freeze any movement by common movement processors."));
+	FAutoConsoleVariableRef CVarFreezeMovement(TEXT("mass.debug.FreezeMovement"), bFreezeMovement, TEXT("Freeze any movement by common movement processors."));
 
 } // UE::MassMovement
diff --git a/Engine/Plugins/Runtime/MassGameplay/Source/MassSmartObjects/Private/MassSmartObjectHandler.cpp b/Engine/Plugins/Runtime/MassGameplay/Source/MassSmartObjects/Private/MassSmartObjectHandler.cpp
--- a/Engine/Plugins/Runtime/MassGameplay/Source/MassSmartObjects/Private/MassSmartObjectHandler.cpp
+++ b/Engine/Plugins/Runtime/MassGameplay/Source/MassSmartObjects/Private/MassSmartObjectHandler.cpp
@@ -11,6 +11,8 @@
 #include "MassSignalSubsystem.h"
 #include "MassSmartObjectFragments.h"
 #include "Engine/World.h"
+#include "MassDebugger.h"
+
 
 namespace UE::Mass::SmartObject
 {
diff --git a/Engine/Plugins/Runtime/MassGameplay/Source/MassSmartObjects/Private/MassSmartObjectProcessor.cpp b/Engine/Plugins/Runtime/MassGameplay/Source/MassSmartObjects/Private/MassSmartObjectProcessor.cpp
--- a/Engine/Plugins/Runtime/MassGameplay/Source/MassSmartObjects/Private/MassSmartObjectProcessor.cpp
+++ b/Engine/Plugins/Runtime/MassGameplay/Source/MassSmartObjects/Private/MassSmartObjectProcessor.cpp
@@ -19,6 +19,8 @@
 #include "MassGameplayExternalTraits.h"
 #include "ZoneGraphSubsystem.h"
 #include "MassGameplayExternalTraits.h"
+#include "MassDebugger.h"
+
 
 //----------------------------------------------------------------------//
 // UMassSmartObjectCandidatesFinderProcessor
