Added BytesPerEntity information to archetype description in MassDebugger

diff --git a/Engine/Plugins/Runtime/MassEntity/Source/MassEntityDebugger/Private/SMassArchetype.cpp b/Engine/Plugins/Runtime/MassEntity/Source/MassEntityDebugger/Private/SMassArchetype.cpp
--- a/Engine/Plugins/Runtime/MassEntity/Source/MassEntityDebugger/Private/SMassArchetype.cpp
+++ b/Engine/Plugins/Runtime/MassEntity/Source/MassEntityDebugger/Private/SMassArchetype.cpp
@@ -55,8 +55,9 @@ void SMassArchetype::Construct(const FArguments& InArgs, TSharedPtr<FMassDebugge
 		.Font(FCoreStyle::GetDefaultFontStyle("Mono", 9))
 	];
 	
-	FText ArchetypeDescription = FText::Format(LOCTEXT("ArchetypeDescrption", "EntitiesCount: {0}\nEntitiesCountPerChunk: {1}\nChunksCount: {2}\nAllocated memory: {3}")
+	FText ArchetypeDescription = FText::Format(LOCTEXT("ArchetypeDescrption", "EntitiesCount: {0}\nBytesPerEntity: {1}\nEntitiesCountPerChunk: {2}\nChunksCount: {3}\nAllocated memory: {4}")
 		, FText::AsNumber(ArchetypeDebugData.ArchetypeStats.EntitiesCount)
+		, FText::AsMemory(ArchetypeDebugData.ArchetypeStats.BytesPerEntity)
 		, FText::AsNumber(ArchetypeDebugData.ArchetypeStats.EntitiesCountPerChunk)
 		, FText::AsNumber(ArchetypeDebugData.ArchetypeStats.ChunksCount)
 		, FText::AsMemory(ArchetypeDebugData.ArchetypeStats.AllocatedSize));
