Fix CIS Issue 332632

diff --git a/Engine/Plugins/Runtime/MassEntity/Source/MassEntity/Public/MassDebugger.h b/Engine/Plugins/Runtime/MassEntity/Source/MassEntity/Public/MassDebugger.h
--- a/Engine/Plugins/Runtime/MassEntity/Source/MassEntity/Public/MassDebugger.h
+++ b/Engine/Plugins/Runtime/MassEntity/Source/MassEntity/Public/MassDebugger.h
@@ -38,20 +38,6 @@ namespace UE::Mass::Debug
 	FString DebugGetFragmentAccessString(EMassFragmentAccess Access);
 	MASSENTITY_API extern void DebugOutputDescription(TConstArrayView<UMassProcessor*> Processors, FOutputDevice& Ar);
 
-	struct FArchetypeStats
-	{
-		/** Number of active entities of the archetype. */
-		int32 EntitiesCount = 0;
-		/** Number of entities that fit per chunk. */
-		int32 EntitiesCountPerChunk = 0;
-		/** Number of allocated chunks. */
-		int32 ChunksCount = 0;
-		/** Total amount of memory taken by this archetype */
-		SIZE_T AllocatedSize = 0;
-		/** Total amount of memory needed by a single entity */
-		int32 BytesPerEntity = 0;
-	};
-
 	MASSENTITY_API extern bool HasDebugEntities();
 	MASSENTITY_API extern bool IsDebuggingSingleEntity();
 
@@ -127,3 +113,20 @@ struct MASSENTITY_API FMassDebugger
 };
 
 #endif // WITH_MASSENTITY_DEBUG
+
+namespace UE::Mass::Debug
+{
+	struct FArchetypeStats
+	{
+		/** Number of active entities of the archetype. */
+		int32 EntitiesCount = 0;
+		/** Number of entities that fit per chunk. */
+		int32 EntitiesCountPerChunk = 0;
+		/** Number of allocated chunks. */
+		int32 ChunksCount = 0;
+		/** Total amount of memory taken by this archetype */
+		SIZE_T AllocatedSize = 0;
+		/** Total amount of memory needed by a single entity */
+		int32 BytesPerEntity = 0;
+	};
+} // namespace UE::Mass::Debug
