Added FMassEntityView::TryMakeView that handles entity not being valid.
Also Renamed MassEntityView.EntityHandle to EntityDataHandle to indicate it's not a regular entity handle.
Also added "IsValid" that calls "IsSet", for convenience and consistency with other code.

diff --git a/Engine/Plugins/Runtime/MassEntity/Source/MassEntity/Private/MassEntityView.cpp b/Engine/Plugins/Runtime/MassEntity/Source/MassEntity/Private/MassEntityView.cpp
--- a/Engine/Plugins/Runtime/MassEntity/Source/MassEntity/Private/MassEntityView.cpp
+++ b/Engine/Plugins/Runtime/MassEntity/Source/MassEntity/Private/MassEntityView.cpp
@@ -13,7 +13,7 @@ FMassEntityView::FMassEntityView(const FMassArchetypeHandle& ArchetypeHandle, FM
 {
 	Entity = InEntity;
 	Archetype = &FMassArchetypeHelper::ArchetypeDataFromHandleChecked(ArchetypeHandle);
-	EntityHandle = Archetype->MakeEntityHandle(Entity);
+	EntityDataHandle = Archetype->MakeEntityHandle(Entity);
 }
 
 FMassEntityView::FMassEntityView(const FMassEntityManager& EntityManager, FMassEntityHandle InEntity)
@@ -21,25 +21,31 @@ FMassEntityView::FMassEntityView(const FMassEntityManager& EntityManager, FMassE
 	Entity = InEntity;
 	const FMassArchetypeHandle ArchetypeHandle = EntityManager.GetArchetypeForEntity(Entity);
 	Archetype = &FMassArchetypeHelper::ArchetypeDataFromHandleChecked(ArchetypeHandle);
-	EntityHandle = Archetype->MakeEntityHandle(Entity);
+	EntityDataHandle = Archetype->MakeEntityHandle(Entity);
+}
+
+FMassEntityView FMassEntityView::TryMakeView(const FMassEntityManager& EntityManager, FMassEntityHandle InEntity)
+{
+	const FMassArchetypeHandle ArchetypeHandle = EntityManager.GetArchetypeForEntity(InEntity);
+	return ArchetypeHandle.IsValid() ? FMassEntityView(ArchetypeHandle, InEntity) : FMassEntityView();
 }
 
 void* FMassEntityView::GetFragmentPtr(const UScriptStruct& FragmentType) const
 {
-	checkSlow(Archetype && EntityHandle.IsValid());
+	checkSlow(Archetype && EntityDataHandle.IsValid());
 	if (const int32* FragmentIndex = Archetype->GetFragmentIndex(&FragmentType))
 	{
 		// failing the below Find means given entity's archetype is missing given FragmentType
-		return Archetype->GetFragmentData(*FragmentIndex, EntityHandle);
+		return Archetype->GetFragmentData(*FragmentIndex, EntityDataHandle);
 	}
 	return nullptr;
 }
 
 void* FMassEntityView::GetFragmentPtrChecked(const UScriptStruct& FragmentType) const
 {
-	checkSlow(Archetype && EntityHandle.IsValid());
+	checkSlow(Archetype && EntityDataHandle.IsValid());
 	const int32 FragmentIndex = Archetype->GetFragmentIndexChecked(&FragmentType);
-	return Archetype->GetFragmentData(FragmentIndex, EntityHandle);
+	return Archetype->GetFragmentData(FragmentIndex, EntityDataHandle);
 }
 
 const void* FMassEntityView::GetConstSharedFragmentPtr(const UScriptStruct& FragmentType) const
@@ -72,6 +78,6 @@ void* FMassEntityView::GetSharedFragmentPtrChecked(const UScriptStruct& Fragment
 
 bool FMassEntityView::HasTag(const UScriptStruct& TagType) const
 {
-	checkSlow(Archetype && EntityHandle.IsValid());
+	checkSlow(Archetype && EntityDataHandle.IsValid());
 	return Archetype->HasTagType(&TagType);
 }
diff --git a/Engine/Plugins/Runtime/MassEntity/Source/MassEntity/Public/MassEntityView.h b/Engine/Plugins/Runtime/MassEntity/Source/MassEntity/Public/MassEntityView.h
--- a/Engine/Plugins/Runtime/MassEntity/Source/MassEntity/Public/MassEntityView.h
+++ b/Engine/Plugins/Runtime/MassEntity/Source/MassEntity/Public/MassEntityView.h
@@ -36,6 +36,12 @@ struct MASSENTITY_API FMassEntityView
 	 */
 	FMassEntityView(const FMassEntityManager& EntityManager, FMassEntityHandle Entity);
 
+	/** 
+	 * If the given handle represents a valid entity the function will create a FMassEntityView just like a constructor 
+	 * would. If the entity is not valid the produced view will be "unset".
+	 */
+	static FMassEntityView TryMakeView(const FMassEntityManager& EntityManager, FMassEntityHandle Entity);
+
 	FMassEntityHandle GetEntity() const	{ return Entity; }
 
 	/** will fail a check if the viewed entity doesn't have the given fragment */	
@@ -121,8 +127,9 @@ struct MASSENTITY_API FMassEntityView
 		return HasTag(*T::StaticStruct());
 	}
 
-	bool IsSet() const { return Archetype != nullptr && EntityHandle.IsValid(); }
-	bool operator==(const FMassEntityView& Other) const { return Archetype == Other.Archetype && EntityHandle == Other.EntityHandle; }
+	bool IsSet() const { return Archetype != nullptr && EntityDataHandle.IsValid(); }
+	bool IsValid() const { return IsSet(); }
+	bool operator==(const FMassEntityView& Other) const { return Archetype == Other.Archetype && EntityDataHandle == Other.EntityDataHandle; }
 
 protected:
 	void* GetFragmentPtr(const UScriptStruct& FragmentType) const;
@@ -135,6 +142,6 @@ protected:
 
 private:
 	FMassEntityHandle Entity;
-	FMassRawEntityInChunkData EntityHandle;
+	FMassRawEntityInChunkData EntityDataHandle;
 	FMassArchetypeData* Archetype = nullptr;
 };
