Added checking whether MassProcessor::Initialize got called by its child classes overrides.

diff --git a/Engine/Plugins/Runtime/MassEntity/Source/MassEntity/Private/MassProcessingTypes.cpp b/Engine/Plugins/Runtime/MassEntity/Source/MassEntity/Private/MassProcessingTypes.cpp
--- a/Engine/Plugins/Runtime/MassEntity/Source/MassEntity/Private/MassProcessingTypes.cpp
+++ b/Engine/Plugins/Runtime/MassEntity/Source/MassEntity/Private/MassProcessingTypes.cpp
@@ -62,6 +62,7 @@ void FMassRuntimePipeline::Initialize(UObject& Owner)
 			{
 				REDIRECT_OBJECT_TO_VLOG(Proc, &Owner);
 				Proc->Initialize(Owner);
+				ensureMsgf(Proc->IsInitialized(), TEXT("Missing Super::Initialize call for %s"), *Proc->GetFullName());
 			}
 		}
 		else
@@ -161,8 +162,13 @@ void FMassRuntimePipeline::AppendUniqueRuntimeProcessorCopies(TConstArrayView<co
 	{
 		UMassProcessor* Proc = Processors[NewProcIndex];
 		check(Proc);
-		REDIRECT_OBJECT_TO_VLOG(Proc, &InOwner);
-		Proc->Initialize(InOwner);
+		
+		if (Proc->IsInitialized() == false)
+		{
+			REDIRECT_OBJECT_TO_VLOG(Proc, &InOwner);
+			Proc->Initialize(InOwner);
+			ensureMsgf(Proc->IsInitialized(), TEXT("Missing Super::Initialize call for %s"), *Proc->GetFullName());
+		}
 	}
 }
 
diff --git a/Engine/Plugins/Runtime/MassEntity/Source/MassEntity/Private/MassProcessor.cpp b/Engine/Plugins/Runtime/MassEntity/Source/MassEntity/Private/MassProcessor.cpp
--- a/Engine/Plugins/Runtime/MassEntity/Source/MassEntity/Private/MassProcessor.cpp
+++ b/Engine/Plugins/Runtime/MassEntity/Source/MassEntity/Private/MassProcessor.cpp
@@ -348,6 +348,7 @@ void UMassCompositeProcessor::Execute(FMassEntityManager& EntityManager, FMassEx
 void UMassCompositeProcessor::Initialize(UObject& Owner)
 {
 	ChildPipeline.Initialize(Owner);
+	Super::Initialize(Owner);
 }
 
 void UMassCompositeProcessor::SetProcessors(TArrayView<UMassProcessor*> InProcessorInstances, const TSharedPtr<FMassEntityManager>& EntityManager)
