Added FMassEntityQuery::ForEachEntityChunkInCollections as a convenience function to execute code on an array or entity collections

diff --git a/Engine/Plugins/Runtime/MassEntity/Source/MassEntity/Private/MassEntityQuery.cpp b/Engine/Plugins/Runtime/MassEntity/Source/MassEntity/Private/MassEntityQuery.cpp
--- a/Engine/Plugins/Runtime/MassEntity/Source/MassEntity/Private/MassEntityQuery.cpp
+++ b/Engine/Plugins/Runtime/MassEntity/Source/MassEntity/Private/MassEntityQuery.cpp
@@ -138,6 +138,14 @@ void FMassEntityQuery::CacheArchetypes(const FMassEntityManager& InEntityManager
 	}
 }
 
+void FMassEntityQuery::ForEachEntityChunkInCollections(TConstArrayView<FMassArchetypeEntityCollection> Collections, FMassEntityManager& EntityManager, FMassExecutionContext& ExecutionContext, const FMassExecuteFunction& ExecuteFunction)
+{
+	for (const FMassArchetypeEntityCollection& Collection : Collections)
+	{
+		ForEachEntityChunk(Collection, EntityManager, ExecutionContext, ExecuteFunction);
+	}
+}
+
 void FMassEntityQuery::ForEachEntityChunk(const FMassArchetypeEntityCollection& Collection, FMassEntityManager& EntityManager, FMassExecutionContext& ExecutionContext, const FMassExecuteFunction& ExecuteFunction)
 {
 	// mz@todo I don't like that we're copying data here.
diff --git a/Engine/Plugins/Runtime/MassEntity/Source/MassEntity/Public/MassEntityQuery.h b/Engine/Plugins/Runtime/MassEntity/Source/MassEntity/Public/MassEntityQuery.h
--- a/Engine/Plugins/Runtime/MassEntity/Source/MassEntity/Public/MassEntityQuery.h
+++ b/Engine/Plugins/Runtime/MassEntity/Source/MassEntity/Public/MassEntityQuery.h
@@ -41,6 +41,8 @@ public:
 	/** Will first verify that the archetype given with Collection matches the query's requirements, and if so will run the other, more generic ForEachEntityChunk implementation */
 	void ForEachEntityChunk(const FMassArchetypeEntityCollection& Collection, FMassEntityManager& EntitySubsystem, FMassExecutionContext& ExecutionContext, const FMassExecuteFunction& ExecuteFunction);
 
+	void ForEachEntityChunkInCollections(TConstArrayView<FMassArchetypeEntityCollection> Collections, FMassEntityManager& EntityManager, FMassExecutionContext& ExecutionContext, const FMassExecuteFunction& ExecuteFunction);
+
 	/** Will gather all archetypes from InEntityManager matching this->Requirements.
 	 *  Note that no work will be done if the cached data is up to date (as tracked by EntitySubsystemHash and 
 	 *	ArchetypeDataVersion properties). */
