Massaged Mass command flushing iteration limit logic to allow for limitless first flush

diff --git a/Engine/Plugins/Runtime/MassEntity/Source/MassEntity/Private/MassEntityManager.cpp b/Engine/Plugins/Runtime/MassEntity/Source/MassEntity/Private/MassEntityManager.cpp
--- a/Engine/Plugins/Runtime/MassEntity/Source/MassEntity/Private/MassEntityManager.cpp
+++ b/Engine/Plugins/Runtime/MassEntity/Source/MassEntity/Private/MassEntityManager.cpp
@@ -150,6 +150,7 @@ void FMassEntityManager::Initialize()
 #endif // WITH_MASSENTITY_DEBUG
 
 	bInitialized = true;
+	bFirstCommandFlush = true;
 }
 
 void FMassEntityManager::PostInitialize()
@@ -1408,7 +1409,7 @@ FMassExecutionContext FMassEntityManager::CreateExecutionContext(const float Del
 
 void FMassEntityManager::FlushCommands(const TSharedPtr<FMassCommandBuffer>& InCommandBuffer)
 {
-	constexpr int MaxIterations = 3;
+	constexpr int32 MaxIterations = 5;
 
 	if (InCommandBuffer)
 	{
@@ -1419,7 +1420,7 @@ void FMassEntityManager::FlushCommands(const TSharedPtr<FMassCommandBuffer>& InC
 	}
 	else
 	{
-		if(DeferredCommandBuffer->HasPendingCommands())
+		if (DeferredCommandBuffer->HasPendingCommands())
 		{
 			FlushedCommandBufferQueue.Enqueue(DeferredCommandBuffer);
 		}
@@ -1429,8 +1430,10 @@ void FMassEntityManager::FlushCommands(const TSharedPtr<FMassCommandBuffer>& InC
 	{
 		bCommandBufferFlushingInProgress = true;
 		
-		int IterationsCounter = 0;
+		const int32 IterationsLimit = bFirstCommandFlush ? MAX_int32 : MaxIterations;
+		int32 IterationsCounter = 0;
 		TOptional<TSharedPtr<FMassCommandBuffer>> CurrentCommandBuffer = FlushedCommandBufferQueue.Dequeue();
+
 		while (IterationsCounter < MaxIterations && CurrentCommandBuffer.IsSet())
 		{
 			IterationsCounter++;
diff --git a/Engine/Plugins/Runtime/MassEntity/Source/MassEntity/Public/MassEntityManager.h b/Engine/Plugins/Runtime/MassEntity/Source/MassEntity/Public/MassEntityManager.h
--- a/Engine/Plugins/Runtime/MassEntity/Source/MassEntity/Public/MassEntityManager.h
+++ b/Engine/Plugins/Runtime/MassEntity/Source/MassEntity/Public/MassEntityManager.h
@@ -510,6 +510,7 @@ private:
 	FOnNewArchetypeDelegate OnNewArchetypeEvent;
 
 	bool bInitialized = false;
+	bool bFirstCommandFlush = true;
 
 
 	//-----------------------------------------------------------------------------
