Massaged tag-related Mass commands to support addition and removal of multiple tags at one go.

diff --git a/Engine/Plugins/Runtime/MassEntity/Source/MassEntity/Public/MassCommandBuffer.h b/Engine/Plugins/Runtime/MassEntity/Source/MassEntity/Public/MassCommandBuffer.h
--- a/Engine/Plugins/Runtime/MassEntity/Source/MassEntity/Public/MassCommandBuffer.h
+++ b/Engine/Plugins/Runtime/MassEntity/Source/MassEntity/Public/MassCommandBuffer.h
@@ -82,14 +82,14 @@ public:
 	void AddTag(FMassEntityHandle Entity)
 	{
 		static_assert(TIsDerivedFrom<T, FMassTag>::IsDerived, "Given struct type is not a valid tag type.");
-		PushCommand<FMassCommandAddTagInternal<EMassCommandCheckTime::CompileTimeCheck, T>>(Entity);
+		PushCommand<FMassCommandAddTagsInternal<EMassCommandCheckTime::CompileTimeCheck, T>>(Entity);
 	}
 
 	template<typename T>
 	void AddTag_RuntimeCheck(FMassEntityHandle Entity)
 	{
 		checkf(T::StaticStruct()->IsChildOf(FMassTag::StaticStruct()), TEXT("Given struct type is not a valid tag type."));
-		PushCommand<FMassCommandAddTagInternal<EMassCommandCheckTime::RuntimeCheck, T>>(Entity);
+		PushCommand<FMassCommandAddTagsInternal<EMassCommandCheckTime::RuntimeCheck, T>>(Entity);
 	}
 
 	/** the convenience function equivalent to calling PushCommand<FMassCommandRemoveTag<T>>(Entity) */
@@ -97,14 +97,14 @@ public:
 	void RemoveTag(FMassEntityHandle Entity)
 	{
 		static_assert(TIsDerivedFrom<T, FMassTag>::IsDerived, "Given struct type is not a valid tag type.");
-		PushCommand<FMassCommandRemoveTagInternal<EMassCommandCheckTime::CompileTimeCheck, T>>(Entity);
+		PushCommand<FMassCommandRemoveTagsInternal<EMassCommandCheckTime::CompileTimeCheck, T>>(Entity);
 	}
 
 	template<typename T>
 	void RemoveTag_RuntimeCheck(FMassEntityHandle Entity)
 	{
 		checkf(T::StaticStruct()->IsChildOf(FMassTag::StaticStruct()), TEXT("Given struct type is not a valid tag type."));
-		PushCommand<FMassCommandRemoveTagInternal<EMassCommandCheckTime::RuntimeCheck, T>>(Entity);
+		PushCommand<FMassCommandRemoveTagsInternal<EMassCommandCheckTime::RuntimeCheck, T>>(Entity);
 	}
 
 	/** the convenience function equivalent to calling PushCommand<FMassCommandSwapTags<TOld, TNew>>(Entity)  */
diff --git a/Engine/Plugins/Runtime/MassEntity/Source/MassEntity/Public/MassCommands.h b/Engine/Plugins/Runtime/MassEntity/Source/MassEntity/Public/MassCommands.h
--- a/Engine/Plugins/Runtime/MassEntity/Source/MassEntity/Public/MassCommands.h
+++ b/Engine/Plugins/Runtime/MassEntity/Source/MassEntity/Public/MassCommands.h
@@ -287,37 +287,43 @@ protected:
 	FMassTagBitSet TagsToRemove;
 };
 
-template<EMassCommandCheckTime CheckTime, typename T>
-struct FMassCommandAddTagInternal : public FMassCommandChangeTags
+template<EMassCommandCheckTime CheckTime, typename... TTypes>
+struct FMassCommandAddTagsInternal : public FMassCommandChangeTags
 {
 	using Super = FMassCommandChangeTags;
-	FMassCommandAddTagInternal()
+	FMassCommandAddTagsInternal()
 		: Super(
 			EMassCommandOperationType::Add, 
-			UE::Mass::Utils::ConstructTagBitSet<CheckTime, T>(),
+			UE::Mass::Utils::ConstructTagBitSet<CheckTime, TTypes...>(),
 			{} 
 			DEBUG_NAME("CommandAddTag"))
 	{}
 };
 
 template<typename T>
-using FMassCommandAddTag = FMassCommandAddTagInternal<EMassCommandCheckTime::CompileTimeCheck, T>;
+using FMassCommandAddTag = FMassCommandAddTagsInternal<EMassCommandCheckTime::CompileTimeCheck, T>;
+
+template<typename... TTypes>
+using FMassCommandAddTags = FMassCommandAddTagsInternal<EMassCommandCheckTime::CompileTimeCheck, TTypes...>;
 
-template<EMassCommandCheckTime CheckTime, typename T>
-struct FMassCommandRemoveTagInternal : public FMassCommandChangeTags
+template<EMassCommandCheckTime CheckTime, typename... TTypes>
+struct FMassCommandRemoveTagsInternal : public FMassCommandChangeTags
 {
 	using Super = FMassCommandChangeTags;
-	FMassCommandRemoveTagInternal()
+	FMassCommandRemoveTagsInternal()
 		: Super(
 			EMassCommandOperationType::Remove, 
 			{}, 
-			UE::Mass::Utils::ConstructTagBitSet<CheckTime, T>()
+			UE::Mass::Utils::ConstructTagBitSet<CheckTime, TTypes...>()
 			DEBUG_NAME("CommandRemoveTag"))
 	{}
 };
 
 template<typename T>
-using FMassCommandAddTag = FMassCommandAddTagInternal<EMassCommandCheckTime::CompileTimeCheck, T>;
+using FMassCommandRemoveTag = FMassCommandRemoveTagsInternal<EMassCommandCheckTime::CompileTimeCheck, T>;
+
+template<typename... TTypes>
+using FMassCommandRemoveTags = FMassCommandRemoveTagsInternal<EMassCommandCheckTime::CompileTimeCheck, TTypes...>;
 
 template<EMassCommandCheckTime CheckTime, typename TOld, typename TNew>
 struct FMassCommandSwapTagsInternal : public FMassCommandChangeTags
