Removed access mode restrictions on FMassFragmentRequirementDescription to support requiring a fragment to be present without needing to bind it.

diff --git a/Engine/Plugins/Runtime/MassEntity/Source/MassEntity/Public/MassExecutionContext.h b/Engine/Plugins/Runtime/MassEntity/Source/MassEntity/Public/MassExecutionContext.h
--- a/Engine/Plugins/Runtime/MassEntity/Source/MassEntity/Public/MassExecutionContext.h
+++ b/Engine/Plugins/Runtime/MassEntity/Source/MassEntity/Public/MassExecutionContext.h
@@ -225,8 +225,8 @@ public:
 	{
 		const UScriptStruct* FragmentType = TFragment::StaticStruct();
 		const FFragmentView* View = FragmentViews.FindByPredicate([FragmentType](const FFragmentView& Element) { return Element.Requirement.StructType == FragmentType; });
-		//checkfSlow(View != nullptr, TEXT("Requested fragment type not bound"));
-		//checkfSlow(View->Requirement.AccessMode == EMassFragmentAccess::ReadWrite, TEXT("Requested fragment has not been bound for writing"));
+		checkfSlow(View != nullptr, TEXT("Requested fragment type not bound"));
+		checkfSlow(View->Requirement.AccessMode == EMassFragmentAccess::ReadWrite, TEXT("Requested fragment has not been bound for writing"));
 		return MakeArrayView<TFragment>((TFragment*)View->FragmentView.GetData(), View->FragmentView.Num());
 	}
 
@@ -235,7 +235,7 @@ public:
 	{
 		const UScriptStruct* FragmentType = TFragment::StaticStruct();
 		const FFragmentView* View = FragmentViews.FindByPredicate([FragmentType](const FFragmentView& Element) { return Element.Requirement.StructType == FragmentType; });
-		//checkfSlow(View != nullptr, TEXT("Requested fragment type not bound"));
+		checkfSlow(View != nullptr, TEXT("Requested fragment type not bound. Make sure the Frament was requested in ReadOnly or ReadWrite mode."));
 		return TConstArrayView<TFragment>((const TFragment*)View->FragmentView.GetData(), View->FragmentView.Num());
 	}
 
diff --git a/Engine/Plugins/Runtime/MassEntity/Source/MassEntity/Public/MassRequirements.h b/Engine/Plugins/Runtime/MassEntity/Source/MassEntity/Public/MassRequirements.h
--- a/Engine/Plugins/Runtime/MassEntity/Source/MassEntity/Public/MassRequirements.h
+++ b/Engine/Plugins/Runtime/MassEntity/Source/MassEntity/Public/MassRequirements.h
@@ -58,8 +58,6 @@ public:
 		, Presence(InPresence)
 	{
 		check(InStruct);
-		checkf((Presence != EMassFragmentPresence::Any && Presence != EMassFragmentPresence::Optional)
-			|| AccessMode == EMassFragmentAccess::ReadOnly || AccessMode == EMassFragmentAccess::ReadWrite, TEXT("Only ReadOnly and ReadWrite modes are suppored for optional requirements"));
 	}
 
 	bool RequiresBinding() const { return (AccessMode != EMassFragmentAccess::None); }
diff --git a/Engine/Plugins/Runtime/MassEntity/Source/MassEntityTestSuite/Private/MassQueryTest.cpp b/Engine/Plugins/Runtime/MassEntity/Source/MassEntityTestSuite/Private/MassQueryTest.cpp
--- a/Engine/Plugins/Runtime/MassEntity/Source/MassEntityTestSuite/Private/MassQueryTest.cpp
+++ b/Engine/Plugins/Runtime/MassEntity/Source/MassEntityTestSuite/Private/MassQueryTest.cpp
@@ -296,13 +296,28 @@ struct FQueryTest_FragmentPresent : FEntityTestBase
 		CA_ASSUME(EntityManager);
 
 		FMassEntityQuery Query;
-		Query.AddRequirement<FTestFragment_Int>(EMassFragmentAccess::ReadOnly);
+		// using EMassFragmentAccess::None to indicate we're interested only in the archetype having the fragment, no binding is required
+		Query.AddRequirement<FTestFragment_Int>(EMassFragmentAccess::None, EMassFragmentPresence::Any);
 		Query.CacheArchetypes(*EntityManager);
 
 		AITEST_EQUAL("There are exactly two archetypes matching the requirements", Query.GetArchetypes().Num(), 2);
 		AITEST_TRUE("FloatsArchetype is not amongst matching archetypes"
 			, !(FloatsArchetype == Query.GetArchetypes()[0] || FloatsArchetype == Query.GetArchetypes()[1]));
 
+		constexpr int32 NumberOfEntitiesToAddA = 5;
+		constexpr int32 NumberOfEntitiesToAddB = 7;
+		TArray<FMassEntityHandle> MatchingEntities;
+		EntityManager->BatchCreateEntities(IntsArchetype, NumberOfEntitiesToAddA, MatchingEntities);
+		EntityManager->BatchCreateEntities(FloatsIntsArchetype, NumberOfEntitiesToAddB, MatchingEntities);
+		ensure(MatchingEntities.Num() == NumberOfEntitiesToAddA + NumberOfEntitiesToAddB);
+
+		int TotalProcessed = 0;
+		FMassExecutionContext ExecContext(*EntityManager.Get());
+		Query.ForEachEntityChunk(*EntityManager, ExecContext, [&TotalProcessed](FMassExecutionContext& Context) {
+			TotalProcessed += Context.GetNumEntities();
+		});
+		AITEST_EQUAL("We expect the number of entities processed to match number added to matching archetypes", MatchingEntities.Num(), TotalProcessed);
+
 		return true;
 	}
 };
