Refactored FMassEntityManager::GetMatchingArchetypes and renamed to GetMatchingArchetypes. The change extends what's possible to query - now a query declaring some tag requirements and only optional fragment requirements will still fetch some results (assuming there are archetypes matching). The change also made the new GetMatchingArchetypes use FMassFragmentRequirements as a parameter, rather than a FMassEntityQuery.

As part of the change FMassEntityManager::AllArchetypes array has been added that collects all archetypes created by the manager. This lets us find archetypes based on tags alone a bit faster. This addition is also a prerequisite for a future change to FMassArchetypeHandle to not store a shared pointer but rather an index to the FMassEntityManager owned array.

diff --git a/Engine/Plugins/Runtime/MassEntity/Source/MassEntity/Private/MassEntityManager.cpp b/Engine/Plugins/Runtime/MassEntity/Source/MassEntity/Private/MassEntityManager.cpp
--- a/Engine/Plugins/Runtime/MassEntity/Source/MassEntity/Private/MassEntityManager.cpp
+++ b/Engine/Plugins/Runtime/MassEntity/Source/MassEntity/Private/MassEntityManager.cpp
@@ -241,6 +241,8 @@ FMassArchetypeHandle FMassEntityManager::CreateArchetype(const FMassArchetypeCom
 		FMassArchetypeData* NewArchetype = new FMassArchetypeData(CreationParams);
 		NewArchetype->Initialize(Composition, ArchetypeDataVersion);
 		ArchetypeDataPtr = HashRow.Add_GetRef(MakeShareable(NewArchetype));
+		AllArchetypes.Add(ArchetypeDataPtr);
+		ensure(AllArchetypes.Num() == ArchetypeDataVersion);
 
 		for (const FMassArchetypeFragmentConfig& FragmentConfig : NewArchetype->GetFragmentConfigs())
 		{
@@ -297,6 +299,8 @@ FMassArchetypeHandle FMassEntityManager::InternalCreateSimilarArchetype(const FM
 		NewArchetype->CopyDebugNamesFrom(SourceArchetypeRef);
 
 		ArchetypeDataPtr = HashRow.Add_GetRef(MakeShareable(NewArchetype));
+		AllArchetypes.Add(ArchetypeDataPtr);
+		ensure(AllArchetypes.Num() == ArchetypeDataVersion);
 
 		for (const FMassArchetypeFragmentConfig& FragmentConfig : NewArchetype->GetFragmentConfigs())
 		{
@@ -1219,13 +1223,21 @@ void FMassEntityManager::CheckIfEntityIsActive(FMassEntityHandle Entity) const
 	checkf(IsEntityBuilt(Entity), TEXT("Entity not yet created(ID: %d, SN:%d)"));
 }
 
-void FMassEntityManager::GetValidArchetypes(const FMassEntityQuery& Query, TArray<FMassArchetypeHandle>& OutValidArchetypes, const uint32 FromArchetypeDataVersion) const
+void FMassEntityManager::GetMatchingArchetypes(const FMassFragmentRequirements& Requirements, TArray<FMassArchetypeHandle>& OutValidArchetypes, const uint32 FromArchetypeDataVersion) const
 {
 	//@TODO: Not optimized yet, but we call this rarely now, so not a big deal.
 
 	// First get set of all archetypes that contain *any* fragment
 	TSet<TSharedPtr<FMassArchetypeData>> AnyArchetypes;
-	TConstArrayView<FMassFragmentRequirementDescription> FragmentRequirements = Query.GetFragmentRequirements();
+	TConstArrayView<FMassFragmentRequirementDescription> FragmentRequirements = Requirements.GetFragmentRequirements();
+
+	// We need to find out if the query does indeed strictly require any fragments to be there - it's for example 
+	// possible to build a query consisting of only optional fragment requirements. In that case we need to test for 
+	// tags as well, independently of fragments, because none of the fragments might be there and we still need to 
+	// respect the tags.
+	// A practical example would be a query asking for all entities that have a given fragment OR a given tag.
+	bool bHasHardFragmentdRequirements = false;
+
 	if (!FragmentRequirements.IsEmpty())
 	{
 		for (const FMassFragmentRequirementDescription& Requirement : FragmentRequirements)
@@ -1236,22 +1248,31 @@ void FMassEntityManager::GetValidArchetypes(const FMassEntityQuery& Query, TArra
 				if (const TArray<TSharedPtr<FMassArchetypeData>>* pData = FragmentTypeToArchetypeMap.Find(Requirement.StructType))
 				{
 					AnyArchetypes.Append(*pData);
+					// only Any and All count as a "hard requirement" since other modes don't strictly require a fragment to be there
+					bHasHardFragmentdRequirements = bHasHardFragmentdRequirements
+						|| Requirement.Presence == EMassFragmentPresence::All
+						|| Requirement.Presence == EMassFragmentPresence::Any;
 				}
 			}
 		}
 	}
-	else
+	
+	// If the query does not have any hard requirements for fragments, but has tag requirements.
+	// We need to search all archetypes, because we don't have a Tag->Archetype Map (and we don't want to create one 
+	// just for this one edge case) so we need to go through all archetypes and match tags. On the bright side,
+	// we can ignore all archetypes with data version < FromArchetypeDataVersion
+	// Note that we collect all archetypes even remotely matching to let the regular process below to filter out
+	// the failing archetypes while justifying the reasons. 
+	if (bHasHardFragmentdRequirements == false
+		&& (Requirements.GetRequiredAllTags().IsEmpty() == false || Requirements.GetRequiredAnyTags().IsEmpty() == false))
 	{
-		// If there are no fragment requirements, assume this is a tag-only query.
-		const FMassTagBitSet& AllTags = Query.GetRequiredAllTags();
-		const FMassTagBitSet& AnyTags = Query.GetRequiredAnyTags();
-
-		for (const TPair<uint32, TArray<TSharedPtr<FMassArchetypeData>>>& KeyValue : FragmentHashToArchetypeMap)
+		const FMassTagBitSet RoughTagFilter = Requirements.GetRequiredAllTags() | Requirements.GetRequiredAnyTags();
+		for (int32 ArchetypeIndex = FromArchetypeDataVersion; ArchetypeIndex < AllArchetypes.Num(); ++ArchetypeIndex)
 		{
-			for (const TSharedPtr<FMassArchetypeData>& ArchetypeData : KeyValue.Value)
+			const TSharedPtr<FMassArchetypeData>& ArchetypeData = AllArchetypes[ArchetypeIndex];
+			if (ArchetypeData.IsValid())
 			{
-				const FMassTagBitSet ArchetypeTags = ArchetypeData->GetTagBitSet();
-				if (ArchetypeTags.HasAll(AllTags) || ArchetypeTags.HasAny(AnyTags))
+				if (ArchetypeData->GetTagBitSet().HasAny(RoughTagFilter))
 				{
 					AnyArchetypes.Add(ArchetypeData);
 				}
@@ -1270,11 +1291,11 @@ void FMassEntityManager::GetValidArchetypes(const FMassEntityQuery& Query, TArra
 			continue;
 		}
 
-		if (Archetype.GetTagBitSet().HasAll(Query.GetRequiredAllTags()) == false)
+		if (Archetype.GetTagBitSet().HasAll(Requirements.GetRequiredAllTags()) == false)
 		{
 			// missing some required tags, skip.
 #if WITH_MASSENTITY_DEBUG
-			const FMassTagBitSet UnsatisfiedTags = Query.GetRequiredAllTags() - Archetype.GetTagBitSet();
+			const FMassTagBitSet UnsatisfiedTags = Requirements.GetRequiredAllTags() - Archetype.GetTagBitSet();
 			FStringOutputDevice Description;
 			UnsatisfiedTags.DebugGetStringDesc(Description);
 			UE_LOG(LogMass, VeryVerbose, TEXT("Archetype did not match due to missing tags: %s")
@@ -1283,11 +1304,11 @@ void FMassEntityManager::GetValidArchetypes(const FMassEntityQuery& Query, TArra
 			continue;
 		}
 
-		if (Archetype.GetTagBitSet().HasNone(Query.GetRequiredNoneTags()) == false)
+		if (Archetype.GetTagBitSet().HasNone(Requirements.GetRequiredNoneTags()) == false)
 		{
 			// has some tags required to be absent
 #if WITH_MASSENTITY_DEBUG
-			const FMassTagBitSet UnwantedTags = Query.GetRequiredAllTags().GetOverlap(Archetype.GetTagBitSet());
+			const FMassTagBitSet UnwantedTags = Requirements.GetRequiredAllTags().GetOverlap(Archetype.GetTagBitSet());
 			FStringOutputDevice Description;
 			UnwantedTags.DebugGetStringDesc(Description);
 			UE_LOG(LogMass, VeryVerbose, TEXT("Archetype has tags required absent: %s")
@@ -1296,23 +1317,23 @@ void FMassEntityManager::GetValidArchetypes(const FMassEntityQuery& Query, TArra
 			continue;
 		}
 
-		if (Query.GetRequiredAnyTags().IsEmpty() == false 
-			&& Archetype.GetTagBitSet().HasAny(Query.GetRequiredAnyTags()) == false)
+		if (Requirements.GetRequiredAnyTags().IsEmpty() == false 
+			&& Archetype.GetTagBitSet().HasAny(Requirements.GetRequiredAnyTags()) == false)
 		{
 #if WITH_MASSENTITY_DEBUG
 			FStringOutputDevice Description;
-			Query.GetRequiredAnyTags().DebugGetStringDesc(Description);
+			Requirements.GetRequiredAnyTags().DebugGetStringDesc(Description);
 			UE_LOG(LogMass, VeryVerbose, TEXT("Archetype did not match due to missing \'any\' tags: %s")
 				, *Description);
 #endif // WITH_MASSENTITY_DEBUG
 			continue;
 		}
 		
-		if (Archetype.GetFragmentBitSet().HasAll(Query.GetRequiredAllFragments()) == false)
+		if (Archetype.GetFragmentBitSet().HasAll(Requirements.GetRequiredAllFragments()) == false)
 		{
 			// missing some required fragments, skip.
 #if WITH_MASSENTITY_DEBUG
-			const FMassFragmentBitSet UnsatisfiedFragments = Query.GetRequiredAllFragments() - Archetype.GetFragmentBitSet();
+			const FMassFragmentBitSet UnsatisfiedFragments = Requirements.GetRequiredAllFragments() - Archetype.GetFragmentBitSet();
 			FStringOutputDevice Description;
 			UnsatisfiedFragments.DebugGetStringDesc(Description);
 			UE_LOG(LogMass, VeryVerbose, TEXT("Archetype did not match due to missing Fragments: %s")
@@ -1321,11 +1342,11 @@ void FMassEntityManager::GetValidArchetypes(const FMassEntityQuery& Query, TArra
 			continue;
 		}
 
-		if (Archetype.GetFragmentBitSet().HasNone(Query.GetRequiredNoneFragments()) == false)
+		if (Archetype.GetFragmentBitSet().HasNone(Requirements.GetRequiredNoneFragments()) == false)
 		{
 			// has some Fragments required to be absent
 #if WITH_MASSENTITY_DEBUG
-			const FMassFragmentBitSet UnwantedFragments = Query.GetRequiredAllFragments().GetOverlap(Archetype.GetFragmentBitSet());
+			const FMassFragmentBitSet UnwantedFragments = Requirements.GetRequiredAllFragments().GetOverlap(Archetype.GetFragmentBitSet());
 			FStringOutputDevice Description;
 			UnwantedFragments.DebugGetStringDesc(Description);
 			UE_LOG(LogMass, VeryVerbose, TEXT("Archetype has Fragments required absent: %s")
@@ -1334,23 +1355,23 @@ void FMassEntityManager::GetValidArchetypes(const FMassEntityQuery& Query, TArra
 			continue;
 		}
 
-		if (Query.GetRequiredAnyFragments().IsEmpty() == false 
-			&& Archetype.GetFragmentBitSet().HasAny(Query.GetRequiredAnyFragments()) == false)
+		if (Requirements.GetRequiredAnyFragments().IsEmpty() == false 
+			&& Archetype.GetFragmentBitSet().HasAny(Requirements.GetRequiredAnyFragments()) == false)
 		{
 #if WITH_MASSENTITY_DEBUG
 			FStringOutputDevice Description;
-			Query.GetRequiredAnyFragments().DebugGetStringDesc(Description);
+			Requirements.GetRequiredAnyFragments().DebugGetStringDesc(Description);
 			UE_LOG(LogMass, VeryVerbose, TEXT("Archetype did not match due to missing \'any\' fragments: %s")
 				, *Description);
 #endif // WITH_MASSENTITY_DEBUG
 			continue;
 		}
 
-		if (Archetype.GetChunkFragmentBitSet().HasAll(Query.GetRequiredAllChunkFragments()) == false)
+		if (Archetype.GetChunkFragmentBitSet().HasAll(Requirements.GetRequiredAllChunkFragments()) == false)
 		{
 			// missing some required fragments, skip.
 #if WITH_MASSENTITY_DEBUG
-			const FMassChunkFragmentBitSet UnsatisfiedFragments = Query.GetRequiredAllChunkFragments() - Archetype.GetChunkFragmentBitSet();
+			const FMassChunkFragmentBitSet UnsatisfiedFragments = Requirements.GetRequiredAllChunkFragments() - Archetype.GetChunkFragmentBitSet();
 			FStringOutputDevice Description;
 			UnsatisfiedFragments.DebugGetStringDesc(Description);
 			UE_LOG(LogMass, VeryVerbose, TEXT("Archetype did not match due to missing Chunk Fragments: %s")
@@ -1359,11 +1380,11 @@ void FMassEntityManager::GetValidArchetypes(const FMassEntityQuery& Query, TArra
 			continue;
 		}
 
-		if (Archetype.GetChunkFragmentBitSet().HasNone(Query.GetRequiredNoneChunkFragments()) == false)
+		if (Archetype.GetChunkFragmentBitSet().HasNone(Requirements.GetRequiredNoneChunkFragments()) == false)
 		{
 			// has some Fragments required to be absent
 #if WITH_MASSENTITY_DEBUG
-			const FMassChunkFragmentBitSet UnwantedFragments = Query.GetRequiredNoneChunkFragments().GetOverlap(Archetype.GetChunkFragmentBitSet());
+			const FMassChunkFragmentBitSet UnwantedFragments = Requirements.GetRequiredNoneChunkFragments().GetOverlap(Archetype.GetChunkFragmentBitSet());
 			FStringOutputDevice Description;
 			UnwantedFragments.DebugGetStringDesc(Description);
 			UE_LOG(LogMass, VeryVerbose, TEXT("Archetype has Chunk Fragments required absent: %s")
@@ -1372,11 +1393,11 @@ void FMassEntityManager::GetValidArchetypes(const FMassEntityQuery& Query, TArra
 			continue;
 		}
 
-		if (Archetype.GetSharedFragmentBitSet().HasAll(Query.GetRequiredAllSharedFragments()) == false)
+		if (Archetype.GetSharedFragmentBitSet().HasAll(Requirements.GetRequiredAllSharedFragments()) == false)
 		{
 			// missing some required fragments, skip.
 #if WITH_MASSENTITY_DEBUG
-			const FMassSharedFragmentBitSet UnsatisfiedFragments = Query.GetRequiredAllSharedFragments() - Archetype.GetSharedFragmentBitSet();
+			const FMassSharedFragmentBitSet UnsatisfiedFragments = Requirements.GetRequiredAllSharedFragments() - Archetype.GetSharedFragmentBitSet();
 			FStringOutputDevice Description;
 			UnsatisfiedFragments.DebugGetStringDesc(Description);
 			UE_LOG(LogMass, VeryVerbose, TEXT("Archetype did not match due to missing Shared Fragments: %s")
@@ -1385,11 +1406,11 @@ void FMassEntityManager::GetValidArchetypes(const FMassEntityQuery& Query, TArra
 			continue;
 		}
 
-		if (Archetype.GetSharedFragmentBitSet().HasNone(Query.GetRequiredNoneSharedFragments()) == false)
+		if (Archetype.GetSharedFragmentBitSet().HasNone(Requirements.GetRequiredNoneSharedFragments()) == false)
 		{
 			// has some Fragments required to be absent
 #if WITH_MASSENTITY_DEBUG
-			const FMassSharedFragmentBitSet UnwantedFragments = Query.GetRequiredNoneSharedFragments().GetOverlap(Archetype.GetSharedFragmentBitSet());
+			const FMassSharedFragmentBitSet UnwantedFragments = Requirements.GetRequiredNoneSharedFragments().GetOverlap(Archetype.GetSharedFragmentBitSet());
 			FStringOutputDevice Description;
 			UnwantedFragments.DebugGetStringDesc(Description);
 			UE_LOG(LogMass, VeryVerbose, TEXT("Archetype has Shared Fragments required absent: %s")
diff --git a/Engine/Plugins/Runtime/MassEntity/Source/MassEntity/Private/MassEntityQuery.cpp b/Engine/Plugins/Runtime/MassEntity/Source/MassEntity/Private/MassEntityQuery.cpp
--- a/Engine/Plugins/Runtime/MassEntity/Source/MassEntity/Private/MassEntityQuery.cpp
+++ b/Engine/Plugins/Runtime/MassEntity/Source/MassEntity/Private/MassEntityQuery.cpp
@@ -107,7 +107,7 @@ void FMassEntityQuery::CacheArchetypes(const FMassEntityManager& InEntityManager
 	if (bUpdateArchetypes)
 	{
 		TArray<FMassArchetypeHandle> NewValidArchetypes;
-		InEntityManager.GetValidArchetypes(*this, NewValidArchetypes, LastUpdatedArchetypeDataVersion);
+		InEntityManager.GetMatchingArchetypes(*this, NewValidArchetypes, LastUpdatedArchetypeDataVersion);
 		LastUpdatedArchetypeDataVersion = InEntityManager.GetArchetypeDataVersion();
 		if (NewValidArchetypes.Num())
 		{
diff --git a/Engine/Plugins/Runtime/MassEntity/Source/MassEntity/Public/MassEntityManager.h b/Engine/Plugins/Runtime/MassEntity/Source/MassEntity/Public/MassEntityManager.h
--- a/Engine/Plugins/Runtime/MassEntity/Source/MassEntity/Public/MassEntityManager.h
+++ b/Engine/Plugins/Runtime/MassEntity/Source/MassEntity/Public/MassEntityManager.h
@@ -425,7 +425,7 @@ public:
 	int32 DebugGetArchetypeEntitiesCount(const FMassArchetypeHandle& Archetype) const;
 	int32 DebugGetArchetypeEntitiesCountPerChunk(const FMassArchetypeHandle& Archetype) const;
 	int32 DebugGetEntityCount() const { return Entities.Num() - NumReservedEntities - EntityFreeIndexList.Num(); }
-	int32 DebugGetArchetypesCount() const { return FragmentHashToArchetypeMap.Num(); }
+	int32 DebugGetArchetypesCount() const { return AllArchetypes.Num(); }
 	void DebugRemoveAllEntities();
 	void DebugForceArchetypeDataVersionBump() { ++ArchetypeDataVersion; }
 	void DebugGetArchetypeStrings(const FMassArchetypeHandle& Archetype, TArray<FName>& OutFragmentNames, TArray<FName>& OutTagNames);
@@ -436,7 +436,7 @@ public:
 #endif // WITH_MASSENTITY_DEBUG
 
 protected:
-	void GetValidArchetypes(const FMassEntityQuery& Query, TArray<FMassArchetypeHandle>& OutValidArchetypes, const uint32 FromArchetypeDataVersion) const;
+	void GetMatchingArchetypes(const FMassFragmentRequirements& Requirements, TArray<FMassArchetypeHandle>& OutValidArchetypes, const uint32 FromArchetypeDataVersion) const;
 	
 	/** 
 	 * A "similar" archetype is an archetype exactly the same as SourceArchetype except for one composition aspect 
@@ -489,6 +489,10 @@ private:
 	// Map to list of archetypes that contain the specified fragment type
 	TMap<const UScriptStruct*, TArray<TSharedPtr<FMassArchetypeData>>> FragmentTypeToArchetypeMap;
 
+	// Contains all archetypes ever created. The array always growing and a given archetypes remains at a given index 
+	// throughout its lifetime, and the index is never reused for another archetype. 
+	TArray<TSharedPtr<FMassArchetypeData>> AllArchetypes;
+
 	// Shared fragments
 	TArray<FConstSharedStruct> ConstSharedFragments;
 	// Hash/Index in array pair
