Changes to MassDebugger letting us fetch the information about recently selected entity (as oopposed to listening to the relevant multicast being the only way to get this information)

diff --git a/Engine/Plugins/Runtime/MassEntity/Source/MassEntity/Private/MassDebugger.cpp b/Engine/Plugins/Runtime/MassEntity/Source/MassEntity/Private/MassDebugger.cpp
--- a/Engine/Plugins/Runtime/MassEntity/Source/MassEntity/Private/MassDebugger.cpp
+++ b/Engine/Plugins/Runtime/MassEntity/Source/MassEntity/Private/MassDebugger.cpp
@@ -361,7 +361,7 @@ FMassDebugger::FOnEntitySelected FMassDebugger::OnEntitySelectedDelegate;
 
 FMassDebugger::FOnMassEntityManagerEvent FMassDebugger::OnEntityManagerInitialized;
 FMassDebugger::FOnMassEntityManagerEvent FMassDebugger::OnEntityManagerDeinitialized;
-TArray<TWeakPtr<const FMassEntityManager>> FMassDebugger::ActiveEntityManagers;
+TArray<FMassDebugger::FEnvironment> FMassDebugger::ActiveEnvironments;
 UE::FSpinLock FMassDebugger::EntityManagerRegistrationLock;
 
 TConstArrayView<FMassEntityQuery*> FMassDebugger::GetProcessorQueries(const UMassProcessor& Processor)
@@ -589,14 +589,34 @@ void FMassDebugger::OutputEntityDescription(FOutputDevice& Ar, const FMassEntity
 void FMassDebugger::SelectEntity(const FMassEntityManager& EntityManager, const FMassEntityHandle EntityHandle)
 {
 	UE::Mass::Debug::SetDebugEntityRange(EntityHandle.Index, EntityHandle.Index);
+
+	const int32 Index = ActiveEnvironments.IndexOfByPredicate([WeakManager = EntityManager.AsWeak()](const FEnvironment& Element)
+		{
+			return Element.EntityManager == WeakManager;
+		});
+	if (ensure(Index != INDEX_NONE))
+	{
+		ActiveEnvironments[Index].SelectedEntity = EntityHandle;
+	}
+
 	OnEntitySelectedDelegate.Broadcast(EntityManager, EntityHandle);
 }
 
+FMassEntityHandle FMassDebugger::GetSelectedEntity(const FMassEntityManager& EntityManager)
+{
+	const int32 Index = ActiveEnvironments.IndexOfByPredicate([WeakManager = EntityManager.AsWeak()](const FEnvironment& Element)
+		{
+			return Element.EntityManager == WeakManager;
+		});
+
+	return Index != INDEX_NONE ? ActiveEnvironments[Index].SelectedEntity : FMassEntityHandle();
+}
+
 void FMassDebugger::RegisterEntityManager(FMassEntityManager& EntityManager)
 {
 	UE::TScopeLock<UE::FSpinLock> ScopeLock(EntityManagerRegistrationLock);
 
-	ActiveEntityManagers.Add(EntityManager.AsShared());
+	ActiveEnvironments.Emplace(EntityManager);
 	OnEntityManagerInitialized.Broadcast(EntityManager);
 }
 
@@ -606,13 +626,20 @@ void FMassDebugger::UnregisterEntityManager(FMassEntityManager& EntityManager)
 
 	if (EntityManager.DoesSharedInstanceExist())
 	{
-		ActiveEntityManagers.Remove(EntityManager.AsWeak());
+		const int32 Index = ActiveEnvironments.IndexOfByPredicate([WeakManager = EntityManager.AsWeak()](const FEnvironment& Element) 
+		{
+			return Element.EntityManager == WeakManager;
+		});
+		if (Index != INDEX_NONE)
+		{
+			ActiveEnvironments.RemoveAt(Index, 1, /*bAllowShrinking=*/false);
+		}
 	}
 	else
 	{
-		ActiveEntityManagers.RemoveAll([](const TWeakPtr<const FMassEntityManager>& Item)
+		ActiveEnvironments.RemoveAll([](const FEnvironment& Item)
 			{
-				return Item.IsValid();
+				return Item.IsValid() == false;
 			});
 	}
 	OnEntityManagerDeinitialized.Broadcast(EntityManager);
diff --git a/Engine/Plugins/Runtime/MassEntity/Source/MassEntity/Public/MassDebugger.h b/Engine/Plugins/Runtime/MassEntity/Source/MassEntity/Public/MassDebugger.h
--- a/Engine/Plugins/Runtime/MassEntity/Source/MassEntity/Public/MassDebugger.h
+++ b/Engine/Plugins/Runtime/MassEntity/Source/MassEntity/Public/MassDebugger.h
@@ -76,6 +76,17 @@ struct MASSENTITY_API FMassDebugger
 	DECLARE_MULTICAST_DELEGATE_TwoParams(FOnEntitySelected, const FMassEntityManager&, const FMassEntityHandle);
 	DECLARE_MULTICAST_DELEGATE_OneParam(FOnMassEntityManagerEvent, const FMassEntityManager&);
 
+	struct FEnvironment
+	{
+		TWeakPtr<const FMassEntityManager> EntityManager;
+		FMassEntityHandle SelectedEntity;
+
+		explicit FEnvironment(const FMassEntityManager& InEntityManager)
+			: EntityManager(InEntityManager.AsWeak())
+		{}
+
+		bool IsValid() const { return EntityManager.IsValid(); }
+	};
 
 	static TConstArrayView<FMassEntityQuery*> GetProcessorQueries(const UMassProcessor& Processor);
 	/** fetches all queries registered for given Processor. Note that in order to get up to date information
@@ -103,6 +114,7 @@ struct MASSENTITY_API FMassDebugger
 	static void OutputEntityDescription(FOutputDevice& Ar, const FMassEntityManager& EntityManager, const FMassEntityHandle Entity, const TCHAR* InPrefix = TEXT(""));
 
 	static void SelectEntity(const FMassEntityManager& EntityManager, const FMassEntityHandle EntityHandle);
+	static FMassEntityHandle GetSelectedEntity(const FMassEntityManager& EntityManager);
 
 	static FOnEntitySelected OnEntitySelectedDelegate;
 
@@ -111,10 +123,10 @@ struct MASSENTITY_API FMassDebugger
 
 	static void RegisterEntityManager(FMassEntityManager& EntityManager);
 	static void UnregisterEntityManager(FMassEntityManager& EntityManager);
-	static TConstArrayView<TWeakPtr<const FMassEntityManager>> GetEntityManagers() { return ActiveEntityManagers; }
+	static TConstArrayView<FEnvironment> GetEnvironments() { return ActiveEnvironments; }
 
 private:
-	static TArray<TWeakPtr<const FMassEntityManager>> ActiveEntityManagers;
+	static TArray<FEnvironment> ActiveEnvironments;
 	static UE::FSpinLock EntityManagerRegistrationLock;
 };
 
diff --git a/Engine/Plugins/Runtime/MassEntity/Source/MassEntityDebugger/Private/SMassDebugger.cpp b/Engine/Plugins/Runtime/MassEntity/Source/MassEntityDebugger/Private/SMassDebugger.cpp
--- a/Engine/Plugins/Runtime/MassEntity/Source/MassEntityDebugger/Private/SMassDebugger.cpp
+++ b/Engine/Plugins/Runtime/MassEntity/Source/MassEntityDebugger/Private/SMassDebugger.cpp
@@ -261,9 +261,9 @@ void SMassDebugger::HandleEnvironmentChanged(TSharedPtr<FMassDebuggerEnvironment
 void SMassDebugger::RebuildEnvironmentsList()
 {
 #if WITH_MASSENTITY_DEBUG
-	for (const TWeakPtr<const FMassEntityManager>& WeakEntityManager : FMassDebugger::GetEntityManagers())
+	for (const FMassDebugger::FEnvironment& Environment : FMassDebugger::GetEnvironments())
 	{
-		if (const FMassEntityManager* EntityManagerPtr = WeakEntityManager.Pin().Get())
+		if (const FMassEntityManager* EntityManagerPtr = Environment.EntityManager.Pin().Get())
 		{
 			const UWorld* World = EntityManagerPtr->GetWorld();
 			if (World == nullptr || UE::Mass::Debugger::Private::IsSupportedWorldType(World->WorldType))
