Switched FMassEntityManager::InternalBatchCreateReservedEntities over from adding entities to archetypes one by one over to using a preexisting FMassArchetypeData::BatchAddEntities.

Also, since the afore mentioned BatchAddEntities produces a FMassArchetypeEntityCollection::FEntityRangeArray as a side product we reuse it now instead of creating a new one as part of FMassArchetypeEntityCollection creation later in the function.

diff --git a/Engine/Plugins/Runtime/MassEntity/Source/MassEntity/Private/MassEntityManager.cpp b/Engine/Plugins/Runtime/MassEntity/Source/MassEntity/Private/MassEntityManager.cpp
--- a/Engine/Plugins/Runtime/MassEntity/Source/MassEntity/Private/MassEntityManager.cpp
+++ b/Engine/Plugins/Runtime/MassEntity/Source/MassEntity/Private/MassEntityManager.cpp
@@ -583,23 +583,23 @@ TSharedRef<FMassEntityManager::FEntityCreationContext> FMassEntityManager::Inter
 {
 	// Functions calling into this one are required to verify that the archetype handle is valid
 	FMassArchetypeData* ArchetypeData = FMassArchetypeHelper::ArchetypeDataFromHandle(ArchetypeHandle);
+	checkf(ArchetypeData, TEXT("Functions calling into this one are required to verify that the archetype handle is valid"));
 
-	// @todo optimize
 	for (FMassEntityHandle Entity : ReservedEntities)
 	{
-		checkf(!IsEntityBuilt(Entity), TEXT("Batch creating reserved entities can only use entities that have not been constructed yet."));
-
+		check(IsEntityValid(Entity));
 		FEntityData& EntityData = Entities[Entity.Index];
+		checkf(!EntityData.CurrentArchetype.IsValid(), TEXT("Batch creating reserved entities can only use entities that have not been constructed yet."));
 		EntityData.CurrentArchetype = ArchetypeHandle.DataPtr;
 		EntityData.SerialNumber = Entity.SerialNumber;
-
-		ArchetypeData->AddEntity(Entity, SharedFragmentValues);
 	}
 
+	FMassArchetypeEntityCollection::FEntityRangeArray TargetArchetypeEntityRanges;
+	ArchetypeData->BatchAddEntities(ReservedEntities, SharedFragmentValues, TargetArchetypeEntityRanges);
+
 	FEntityCreationContext* CreationContext = new FEntityCreationContext(ReservedEntities.Num());
-	// @todo this could probably be optimized since one would assume we're adding elements to OutEntities in order.
-	// Then again, if that's the case, the sorting will be almost instant
-	new (&CreationContext->EntityCollection) FMassArchetypeEntityCollection(ArchetypeHandle, ReservedEntities, FMassArchetypeEntityCollection::NoDuplicates);
+	new (&CreationContext->EntityCollection) FMassArchetypeEntityCollection(ArchetypeHandle, MoveTemp(TargetArchetypeEntityRanges));
+
 	if (ObserverManager.HasObserversForBitSet(ArchetypeData->GetCompositionDescriptor().Fragments, EMassObservedOperation::Add)
 		|| ObserverManager.HasObserversForBitSet(ArchetypeData->GetCompositionDescriptor().Tags, EMassObservedOperation::Add))
 	{
diff --git a/Engine/Plugins/Runtime/MassEntity/Source/MassEntity/Public/MassArchetypeTypes.h b/Engine/Plugins/Runtime/MassEntity/Source/MassEntity/Public/MassArchetypeTypes.h
--- a/Engine/Plugins/Runtime/MassEntity/Source/MassEntity/Public/MassArchetypeTypes.h
+++ b/Engine/Plugins/Runtime/MassEntity/Source/MassEntity/Public/MassArchetypeTypes.h
@@ -109,7 +109,7 @@ public:
 	FMassArchetypeEntityCollection(const FMassArchetypeHandle& InArchetype, TConstArrayView<FMassEntityHandle> InEntities, EDuplicatesHandling DuplicatesHandling);
 	explicit FMassArchetypeEntityCollection(const FMassArchetypeHandle& InArchetypeHandle, const EInitializationType Initialization = EInitializationType::GatherAll);
 	explicit FMassArchetypeEntityCollection(TSharedPtr<FMassArchetypeData>& InArchetype, const EInitializationType Initialization = EInitializationType::GatherAll);
-	FMassArchetypeEntityCollection(FMassArchetypeHandle& InArchetypeHandle, FEntityRangeArray&& InEntityRanges)
+	FMassArchetypeEntityCollection(const FMassArchetypeHandle& InArchetypeHandle, FEntityRangeArray&& InEntityRanges)
 		: Ranges(MoveTemp(InEntityRanges))
 		, Archetype(InArchetypeHandle)
 	{}
