Critical fix to how Mass handles mutating command-instigated operations on destroyed entities.

When handling these operations, as triggered by commands, we assum the archetype-less entities have been reserved and never check it. This CL adds some checking so that we know when it happens in the future, but also the CL changes how we treat invalid entities while creating `FMassArchetypeEntityCollectionWithPayload`, which is the tool commands are using to batch together updates to diverse collections of entities.

This fix has been prompted by project-specific code issuing commands outside of Mass processing, and before Mass processing actually starts there's gameplay code that synchronously kills off relevant entities. That behavior is not recommended and the project-side code will be updated, but Mass needs to handle this case gracefully.

diff --git a/Engine/Plugins/Runtime/MassEntity/Source/MassEntity/Private/MassArchetypeTypes.cpp b/Engine/Plugins/Runtime/MassEntity/Source/MassEntity/Private/MassArchetypeTypes.cpp
--- a/Engine/Plugins/Runtime/MassEntity/Source/MassEntity/Private/MassArchetypeTypes.cpp
+++ b/Engine/Plugins/Runtime/MassEntity/Source/MassEntity/Private/MassArchetypeTypes.cpp
@@ -227,18 +227,32 @@ void FMassArchetypeEntityCollectionWithPayload::CreateEntityRangesWithPayload(co
 	for (int32 i = 0; i < Entities.Num(); ++i)
 	{
 		const FMassEntityHandle& Entity = Entities[i];
-		const FMassArchetypeHandle ArchetypeHandle = EntityManager.GetArchetypeForEntity(Entity);
-		const FMassArchetypeData* ArchetypePtr = FMassArchetypeHelper::ArchetypeDataFromHandle(ArchetypeHandle);
+		if (EntityManager.IsEntityValid(Entity))
+		{
+			// using Unsafe since we just checked that the entity is valid
+			const FMassArchetypeHandle ArchetypeHandle = EntityManager.GetArchetypeForEntityUnsafe(Entity);
+			const FMassArchetypeData* ArchetypePtr = FMassArchetypeHelper::ArchetypeDataFromHandle(ArchetypeHandle);
 		
-		// @todo if FMassArchetypeHandle used indices the look up would be a lot faster
-		int32 ArchetypeIndex = INDEX_NONE;
+			// @todo if FMassArchetypeHandle used indices the look up would be a lot faster
+			int32 ArchetypeIndex = INDEX_NONE;
 		
-		if (Archetypes.Find(FArchetypeInfo{ ArchetypeHandle, 0 }, ArchetypeIndex) == false)
+			if (Archetypes.Find(FArchetypeInfo{ ArchetypeHandle, 0 }, ArchetypeIndex) == false)
+			{
+				ArchetypeIndex = Archetypes.Add({ ArchetypeHandle, 0 });
+			}
+			++Archetypes[ArchetypeIndex].Count;
+			EntityData[i] = { ArchetypeIndex, ArchetypePtr ? ArchetypePtr->GetInternalIndexForEntity(Entity.Index) : Entity.Index };
+		}
+		else
 		{
-			ArchetypeIndex = Archetypes.Add({ ArchetypeHandle, 0 });
+			// for invalid entities we create an entry that will result in all of them batched together 
+			// (due to having the same archetype index, INDEX_NONE). Since the main logic loop below relies on entries 
+			// in Archetypes and the "invalid" EntityData doesn't correspond to any, all the invalid entities 
+			// will be silently filtered out.
+			EntityData[i] = FEntityInArchetype();
+			UE_LOG(LogMass, Warning, TEXT("%hs: Invalid entity handle passed in. Ignoring it, but check your code to make sure you don't mix synchronous entity-mutating Mass API function calls with Mass commands")
+				, __FUNCTION__);
 		}
-		++Archetypes[ArchetypeIndex].Count;
-		EntityData[i] = { ArchetypeIndex, ArchetypePtr ? ArchetypePtr->GetInternalIndexForEntity(Entity.Index) : Entity.Index };
 	}
 
 	UE::Mass::Utils::AbstractSort(Entities.Num(), [&EntityData](const int32 LHS, const int32 RHS)
diff --git a/Engine/Plugins/Runtime/MassEntity/Source/MassEntity/Private/MassEntityManager.cpp b/Engine/Plugins/Runtime/MassEntity/Source/MassEntity/Private/MassEntityManager.cpp
--- a/Engine/Plugins/Runtime/MassEntity/Source/MassEntity/Private/MassEntityManager.cpp
+++ b/Engine/Plugins/Runtime/MassEntity/Source/MassEntity/Private/MassEntityManager.cpp
@@ -514,6 +514,7 @@ void FMassEntityManager::BatchBuildEntities(const FMassArchetypeEntityCollection
 		check(Entities.IsValidIndex(Entity.Index));
 
 		FEntityData& EntityData = Entities[Entity.Index];
+		checkf(EntityData.SerialNumber != 0, TEXT("Failing this check means we're trying to build entities that have not been reserved yet. Consider using BatchCreateEntities"));
 		Entity.SerialNumber = EntityData.SerialNumber;
 		EntityData.CurrentArchetype = TargetArchetypeHandle.DataPtr;
 	}
@@ -590,6 +591,7 @@ TSharedRef<FMassEntityManager::FEntityCreationContext> FMassEntityManager::Inter
 		check(IsEntityValid(Entity));
 		FEntityData& EntityData = Entities[Entity.Index];
 		checkf(!EntityData.CurrentArchetype.IsValid(), TEXT("Batch creating reserved entities can only use entities that have not been constructed yet."));
+		checkf(EntityData.SerialNumber != 0, TEXT("Failing this check means we're trying to build entities that have not been reserved yet"));
 		EntityData.CurrentArchetype = ArchetypeHandle.DataPtr;
 		EntityData.SerialNumber = Entity.SerialNumber;
 	}
