Fix for MassEntity Add/Remove observers not being called in many situations
* Fixed observer not being called when batch adding entities with tags but not fragments that overlap the current observer set
* Fixed observer not being called when single entity added with archtype that contains an observed tag or fragment
* Fixed observer not being called when single observed fragment added to entity
* Fixed observer not being called when single entity removed that has archtype containing observed tag or fragment
* Fixed observer not being called when single observed tag or fragment removed from entity

diff --git a/Engine/Plugins/Runtime/MassEntity/Source/MassEntity/Private/MassEntityManager.cpp b/Engine/Plugins/Runtime/MassEntity/Source/MassEntity/Private/MassEntityManager.cpp
--- a/Engine/Plugins/Runtime/MassEntity/Source/MassEntity/Private/MassEntityManager.cpp
+++ b/Engine/Plugins/Runtime/MassEntity/Source/MassEntity/Private/MassEntityManager.cpp
@@ -529,7 +529,7 @@ void FMassEntityManager::BatchBuildEntities(const FMassArchetypeEntityCollection
 		TargetArchetypeHandle.DataPtr->BatchSetFragmentValues(TargetArchetypeEntityRanges, EncodedEntitiesWithPayload.GetPayload());
 	}
 
-	if (ObserverManager.HasObserversForBitSet(Composition.Fragments, EMassObservedOperation::Add))
+	if (ObserverManager.HasObserversForBitSet(Composition.Fragments, EMassObservedOperation::Add) || ObserverManager.HasObserversForBitSet(Composition.Tags, EMassObservedOperation::Add))
 	{
 		ObserverManager.OnCompositionChanged(
 			FMassArchetypeEntityCollection(TargetArchetypeHandle, MoveTemp(TargetArchetypeEntityRanges))
@@ -803,6 +803,8 @@ void FMassEntityManager::InternalBuildEntity(FMassEntityHandle Entity, const FMa
 	FEntityData& EntityData = Entities[Entity.Index];
 	EntityData.CurrentArchetype = ArchetypeHandle.DataPtr;
 	EntityData.CurrentArchetype->AddEntity(Entity, SharedFragmentValues);
+
+	ObserverManager.OnPostCompositionAdded(Entity, EntityData.CurrentArchetype->GetCompositionDescriptor());
 }
 
 void FMassEntityManager::InternalReleaseEntity(FMassEntityHandle Entity)
@@ -847,6 +849,11 @@ void FMassEntityManager::InternalAddFragmentListToEntity(FMassEntityHandle Entit
 		NewArchetype.CopyDebugNamesFrom(*OldArchetype);
 		EntityData.CurrentArchetype->MoveEntityToAnotherArchetype(Entity, NewArchetype);
 		EntityData.CurrentArchetype = NewArchetypeHandle.DataPtr;
+
+		FMassArchetypeCompositionDescriptor CompositionChangeDescriptor;
+		CompositionChangeDescriptor.Fragments += InFragments;
+
+		ObserverManager.OnPostCompositionAdded(Entity, CompositionChangeDescriptor);
 	}
 }
 
@@ -886,6 +893,11 @@ void FMassEntityManager::RemoveFragmentListFromEntity(FMassEntityHandle Entity,
 		const FMassArchetypeCompositionDescriptor NewComposition(OldArchetype->GetFragmentBitSet() - FragmentsToRemove, OldArchetype->GetTagBitSet(), OldArchetype->GetChunkFragmentBitSet(), OldArchetype->GetSharedFragmentBitSet());
 		const FMassArchetypeHandle NewArchetypeHandle = CreateArchetype(NewComposition);
 
+		FMassArchetypeCompositionDescriptor CompositionDelta;
+		// Find overlap.  It isn't guaranteed that the old archtype has all of the fragments being removed.
+		CompositionDelta.Fragments = OldArchetype->GetFragmentBitSet().GetOverlap(FragmentsToRemove);
+		ObserverManager.OnPreCompositionRemoved(Entity, CompositionDelta);
+
 		// Move the entity over
 		FMassArchetypeData& NewArchetype = FMassArchetypeHelper::ArchetypeDataFromHandleChecked(NewArchetypeHandle);
 		NewArchetype.CopyDebugNamesFrom(*OldArchetype);
@@ -943,6 +955,12 @@ void FMassEntityManager::AddTagToEntity(FMassEntityHandle Entity, const UScriptS
 		// Move the entity over
 		EntityData.CurrentArchetype->MoveEntityToAnotherArchetype(Entity, *NewArchetypeHandle.DataPtr.Get());
 		EntityData.CurrentArchetype = NewArchetypeHandle.DataPtr;
+
+		FMassArchetypeCompositionDescriptor CompositionDelta;
+		FMassTagBitSet TagDelta;
+		TagDelta.Add(*TagType);
+		CompositionDelta.Tags = TagDelta;
+		ObserverManager.OnPostCompositionAdded(Entity, CompositionDelta);
 	}
 }
 	
@@ -958,10 +976,15 @@ void FMassEntityManager::RemoveTagFromEntity(FMassEntityHandle Entity, const USc
 
 	if (CurrentArchetype->HasTagType(TagType))
 	{
+		FMassArchetypeCompositionDescriptor CompositionDelta;
+		FMassTagBitSet TagDelta;
+		TagDelta.Add(*TagType);
+		CompositionDelta.Tags = TagDelta;
+		ObserverManager.OnPreCompositionRemoved(Entity, CompositionDelta);
+		
 		// CurrentArchetype->GetTagBitSet() -  *TagType
-		FMassTagBitSet NewTags = CurrentArchetype->GetTagBitSet();
-		NewTags.Remove(*TagType);
-		const FMassArchetypeHandle NewArchetypeHandle = InternalCreateSimilarArchetype(EntityData.CurrentArchetype, NewTags);
+		const FMassTagBitSet NewTagComposition = CurrentArchetype->GetTagBitSet() - TagDelta;
+		const FMassArchetypeHandle NewArchetypeHandle = InternalCreateSimilarArchetype(EntityData.CurrentArchetype, NewTagComposition);
 		checkSlow(NewArchetypeHandle.IsValid());
 
 		// Move the entity over
