Added simple macros for entity-based Mass debugging

diff --git a/Engine/Plugins/Runtime/MassEntity/Source/MassEntity/Private/MassDebugger.cpp b/Engine/Plugins/Runtime/MassEntity/Source/MassEntity/Private/MassDebugger.cpp
--- a/Engine/Plugins/Runtime/MassEntity/Source/MassEntity/Private/MassDebugger.cpp
+++ b/Engine/Plugins/Runtime/MassEntity/Source/MassEntity/Private/MassDebugger.cpp
@@ -19,6 +19,17 @@
 
 namespace UE::Mass::Debug
 {
+	bool bAllowProceduralDebuggedEntitySelection = false;
+	bool bAllowBreakOnlDebuggedEntity = false;
+
+	FAutoConsoleVariableRef CVars[] =
+	{
+		FAutoConsoleVariableRef(TEXT("mass.debug.AllowProceduralDebuggedEntitySelection"), bAllowProceduralDebuggedEntitySelection
+			, TEXT("Guards whether MASS_SET_ENTITY_DEBUGGED calls take effect."), ECVF_Cheat)
+		, FAutoConsoleVariableRef(TEXT("mass.debug.AllowBreakOnlDebuggedEntity"), bAllowBreakOnlDebuggedEntity
+			, TEXT("Guards whether MASS_BREAK_IF_ENTITY_DEBUGGED calls take effect."), ECVF_Cheat)
+	};
+	
 
 	FString DebugGetFragmentAccessString(EMassFragmentAccess Access)
 	{
diff --git a/Engine/Plugins/Runtime/MassEntity/Source/MassEntity/Public/MassDebugger.h b/Engine/Plugins/Runtime/MassEntity/Source/MassEntity/Public/MassDebugger.h
--- a/Engine/Plugins/Runtime/MassEntity/Source/MassEntity/Public/MassDebugger.h
+++ b/Engine/Plugins/Runtime/MassEntity/Source/MassEntity/Public/MassDebugger.h
@@ -39,6 +39,17 @@ namespace UE::Mass::Debug
 } // namespace UE::Mass::Debug
 
 #if WITH_MASSENTITY_DEBUG
+
+namespace UE::Mass::Debug
+{
+	extern MASSENTITY_API bool bAllowProceduralDebuggedEntitySelection;
+	extern MASSENTITY_API bool bAllowBreakOnlDebuggedEntity;
+} // namespace UE::Mass::Debug
+
+#define MASS_IF_ENTITY_DEBUGGED(Manager, EntityHandle) (FMassDebugger::GetSelectedEntity(Manager) == EntityHandle)
+#define MASS_BREAK_IF_ENTITY_DEBUGGED(Manager, EntityHandle) { if (UE::Mass::Debug::bAllowBreakOnlDebuggedEntity && MASS_IF_ENTITY_DEBUGGED(Manager, EntityHandle)) { PLATFORM_BREAK();} }
+#define MASS_SET_ENTITY_DEBUGGED(Manager, EntityHandle) { if (UE::Mass::Debug::bAllowProceduralDebuggedEntitySelection) {FMassDebugger::SelectEntity(Manager, EntityHandle); }}
+
 namespace UE::Mass::Debug
 {
 	struct MASSENTITY_API FQueryRequirementsView
@@ -143,4 +154,8 @@ struct MASSENTITY_API FMassDebugger
 	static FString GetArchetypeRequirementCompatibilityDescription(const FMassFragmentRequirements& Requirements, const FMassArchetypeHandle& ArchetypeHandle) { return TEXT("[no debug information]"); }
 };
 
+#define MASS_IF_ENTITY_DEBUGGED(a, b) false
+#define MASS_BREAK_IF_ENTITY_DEBUGGED(a, b)
+#define MASS_SET_ENTITY_DEBUGGED(a, b)
+
 #endif // WITH_MASSENTITY_DEBUG
