Added an ensure to FMassArchetypeData::BatchMoveEntitiesToAnotherArchetype to catch invalid EntityRanges

diff --git a/Engine/Plugins/Runtime/MassEntity/Source/MassEntity/Private/MassArchetypeData.cpp b/Engine/Plugins/Runtime/MassEntity/Source/MassEntity/Private/MassArchetypeData.cpp
--- a/Engine/Plugins/Runtime/MassEntity/Source/MassEntity/Private/MassArchetypeData.cpp
+++ b/Engine/Plugins/Runtime/MassEntity/Source/MassEntity/Private/MassArchetypeData.cpp
@@ -993,6 +993,11 @@ void FMassArchetypeData::BatchMoveEntitiesToAnotherArchetype(const FMassArchetyp
 
 	for (const FMassArchetypeEntityCollection::FArchetypeEntityRange EntityRange : Subchunks)
 	{
+		if (!ensureMsgf(EntityRange.IsSet() && EntityRange.Length > 0, TEXT("We only expect to get valid EntityRanges at this point.")))
+		{
+			continue;
+		}
+
 		FMassArchetypeChunk& Chunk = Chunks[EntityRange.ChunkIndex];
 
 		// 0 - consider compacting new archetype to ensure larger empty spaces
