Fixed typo in MassDebugger CVar

diff --git a/Engine/Plugins/Runtime/MassEntity/Source/MassEntity/Private/MassDebugger.cpp b/Engine/Plugins/Runtime/MassEntity/Source/MassEntity/Private/MassDebugger.cpp
--- a/Engine/Plugins/Runtime/MassEntity/Source/MassEntity/Private/MassDebugger.cpp
+++ b/Engine/Plugins/Runtime/MassEntity/Source/MassEntity/Private/MassDebugger.cpp
@@ -20,13 +20,13 @@
 namespace UE::Mass::Debug
 {
 	bool bAllowProceduralDebuggedEntitySelection = false;
-	bool bAllowBreakOnlDebuggedEntity = false;
+	bool bAllowBreakOnDebuggedEntity = false;
 
 	FAutoConsoleVariableRef CVars[] =
 	{
 		FAutoConsoleVariableRef(TEXT("mass.debug.AllowProceduralDebuggedEntitySelection"), bAllowProceduralDebuggedEntitySelection
 			, TEXT("Guards whether MASS_SET_ENTITY_DEBUGGED calls take effect."), ECVF_Cheat)
-		, FAutoConsoleVariableRef(TEXT("mass.debug.AllowBreakOnlDebuggedEntity"), bAllowBreakOnlDebuggedEntity
+		, FAutoConsoleVariableRef(TEXT("mass.debug.AllowBreakOnDebuggedEntity"), bAllowBreakOnDebuggedEntity
 			, TEXT("Guards whether MASS_BREAK_IF_ENTITY_DEBUGGED calls take effect."), ECVF_Cheat)
 	};
 	
diff --git a/Engine/Plugins/Runtime/MassEntity/Source/MassEntity/Public/MassDebugger.h b/Engine/Plugins/Runtime/MassEntity/Source/MassEntity/Public/MassDebugger.h
--- a/Engine/Plugins/Runtime/MassEntity/Source/MassEntity/Public/MassDebugger.h
+++ b/Engine/Plugins/Runtime/MassEntity/Source/MassEntity/Public/MassDebugger.h
@@ -43,11 +43,11 @@ namespace UE::Mass::Debug
 namespace UE::Mass::Debug
 {
 	extern MASSENTITY_API bool bAllowProceduralDebuggedEntitySelection;
-	extern MASSENTITY_API bool bAllowBreakOnlDebuggedEntity;
+	extern MASSENTITY_API bool bAllowBreakOnDebuggedEntity;
 } // namespace UE::Mass::Debug
 
 #define MASS_IF_ENTITY_DEBUGGED(Manager, EntityHandle) (FMassDebugger::GetSelectedEntity(Manager) == EntityHandle)
-#define MASS_BREAK_IF_ENTITY_DEBUGGED(Manager, EntityHandle) { if (UE::Mass::Debug::bAllowBreakOnlDebuggedEntity && MASS_IF_ENTITY_DEBUGGED(Manager, EntityHandle)) { PLATFORM_BREAK();} }
+#define MASS_BREAK_IF_ENTITY_DEBUGGED(Manager, EntityHandle) { if (UE::Mass::Debug::bAllowBreakOnDebuggedEntity && MASS_IF_ENTITY_DEBUGGED(Manager, EntityHandle)) { PLATFORM_BREAK();} }
 #define MASS_SET_ENTITY_DEBUGGED(Manager, EntityHandle) { if (UE::Mass::Debug::bAllowProceduralDebuggedEntitySelection) {FMassDebugger::SelectEntity(Manager, EntityHandle); }}
 
 namespace UE::Mass::Debug
