Added const and non-const shared fragment accessors to FMassEntityManager

diff --git a/Engine/Plugins/Runtime/MassEntity/Source/MassEntity/Private/MassEntityManager.cpp b/Engine/Plugins/Runtime/MassEntity/Source/MassEntity/Private/MassEntityManager.cpp
--- a/Engine/Plugins/Runtime/MassEntity/Source/MassEntity/Private/MassEntityManager.cpp
+++ b/Engine/Plugins/Runtime/MassEntity/Source/MassEntity/Private/MassEntityManager.cpp
@@ -1226,6 +1226,44 @@ void* FMassEntityManager::InternalGetFragmentDataPtr(FMassEntityHandle Entity, c
 	return EntityData.CurrentArchetype->GetFragmentDataForEntity(FragmentType, Entity.Index);
 }
 
+const void* FMassEntityManager::InternalGetConstSharedFragmentDataChecked(FMassEntityHandle Entity, const UScriptStruct* ConstSharedFragmentType) const
+{
+	CheckIfEntityIsActive(Entity);
+	checkf((ConstSharedFragmentType != nullptr) && ConstSharedFragmentType->IsChildOf(FMassSharedFragment::StaticStruct()), TEXT("InternalGetConstSharedFragmentDataChecked called with an invalid fragment type '%s'"), *GetPathNameSafe(ConstSharedFragmentType));
+	const FEntityData& EntityData = Entities[Entity.Index];
+	const FConstSharedStruct* SharedFragment = EntityData.CurrentArchetype->GetSharedFragmentValues(Entity).GetConstSharedFragments().FindByPredicate(FStructTypeEqualOperator(ConstSharedFragmentType));
+	check(SharedFragment != nullptr);
+	return SharedFragment->GetMemory();
+}
+
+const void* FMassEntityManager::InternalGetConstSharedFragmentDataPtr(FMassEntityHandle Entity, const UScriptStruct* ConstSharedFragmentType) const
+{
+	CheckIfEntityIsActive(Entity);
+	checkf((ConstSharedFragmentType != nullptr) && ConstSharedFragmentType->IsChildOf(FMassSharedFragment::StaticStruct()), TEXT("InternalGetConstSharedFragmentData called with an invalid fragment type '%s'"), *GetPathNameSafe(ConstSharedFragmentType));
+	const FEntityData& EntityData = Entities[Entity.Index];
+	const FConstSharedStruct* SharedFragment = EntityData.CurrentArchetype->GetSharedFragmentValues(Entity).GetConstSharedFragments().FindByPredicate(FStructTypeEqualOperator(ConstSharedFragmentType));
+	return (SharedFragment != nullptr) ? SharedFragment->GetMemory() : nullptr;
+}
+
+void* FMassEntityManager::InternalGetSharedFragmentDataChecked(FMassEntityHandle Entity, const UScriptStruct* SharedFragmentType) const
+{
+	CheckIfEntityIsActive(Entity);
+	checkf((SharedFragmentType != nullptr) && SharedFragmentType->IsChildOf(FMassSharedFragment::StaticStruct()), TEXT("InternalGetSharedFragmentDataChecked called with an invalid fragment type '%s'"), *GetPathNameSafe(SharedFragmentType));
+	const FEntityData& EntityData = Entities[Entity.Index];
+	const FSharedStruct* SharedFragment = EntityData.CurrentArchetype->GetSharedFragmentValues(Entity).GetSharedFragments().FindByPredicate(FStructTypeEqualOperator(SharedFragmentType));
+	check(SharedFragment != nullptr);
+	return SharedFragment->GetMemory();
+}
+
+void* FMassEntityManager::InternalGetSharedFragmentDataPtr(FMassEntityHandle Entity, const UScriptStruct* SharedFragmentType) const
+{
+	CheckIfEntityIsActive(Entity);
+	checkf((SharedFragmentType != nullptr) && SharedFragmentType->IsChildOf(FMassSharedFragment::StaticStruct()), TEXT("InternalGetSharedFragmentData called with an invalid fragment type '%s'"), *GetPathNameSafe(SharedFragmentType));
+	const FEntityData& EntityData = Entities[Entity.Index];
+	const FSharedStruct* SharedFragment = EntityData.CurrentArchetype->GetSharedFragmentValues(Entity).GetSharedFragments().FindByPredicate(FStructTypeEqualOperator(SharedFragmentType));
+	return (SharedFragment != nullptr) ? SharedFragment->GetMemory() : nullptr;
+}
+
 bool FMassEntityManager::IsEntityValid(FMassEntityHandle Entity) const
 {
 	return (Entity.Index > 0) && Entities.IsValidIndex(Entity.Index) && (Entities[Entity.Index].SerialNumber == Entity.SerialNumber);
diff --git a/Engine/Plugins/Runtime/MassEntity/Source/MassEntity/Public/MassEntityManager.h b/Engine/Plugins/Runtime/MassEntity/Source/MassEntity/Public/MassEntityManager.h
--- a/Engine/Plugins/Runtime/MassEntity/Source/MassEntity/Public/MassEntityManager.h
+++ b/Engine/Plugins/Runtime/MassEntity/Source/MassEntity/Public/MassEntityManager.h
@@ -334,6 +334,40 @@ public:
 		return FStructView(FragmentType, static_cast<uint8*>(InternalGetFragmentDataPtr(Entity, FragmentType)));
 	}
 
+	template <typename ConstSharedFragmentType>
+	ConstSharedFragmentType& GetConstSharedFragmentDataChecked(FMassEntityHandle Entity) const
+	{
+		return *((ConstSharedFragmentType*)InternalGetConstSharedFragmentDataChecked(Entity, ConstSharedFragmentType::StaticStruct()));
+	}
+
+	template <typename ConstSharedFragmentType>
+	ConstSharedFragmentType* GetConstSharedFragmentDataPtr(FMassEntityHandle Entity) const
+	{
+		return (ConstSharedFragmentType*)InternalGetConstSharedFragmentDataPtr(Entity, ConstSharedFragmentType::StaticStruct());
+	}
+
+	FConstStructView GetConstSharedFragmentDataStruct(FMassEntityHandle Entity, const UScriptStruct* ConstSharedFragmentType) const
+	{
+		return FConstStructView(ConstSharedFragmentType, static_cast<const uint8*>(InternalGetConstSharedFragmentDataPtr(Entity, ConstSharedFragmentType)));
+	}
+
+	template <typename SharedFragmentType>
+	SharedFragmentType& GetSharedFragmentDataChecked(FMassEntityHandle Entity) const
+	{
+		return *((SharedFragmentType*)InternalGetSharedFragmentDataChecked(Entity, SharedFragmentType::StaticStruct()));
+	}
+
+	template <typename SharedFragmentType>
+	SharedFragmentType* GetSharedFragmentDataPtr(FMassEntityHandle Entity) const
+	{
+		return (SharedFragmentType*)InternalGetSharedFragmentDataPtr(Entity, SharedFragmentType::StaticStruct());
+	}
+
+	FConstStructView GetSharedFragmentDataStruct(FMassEntityHandle Entity, const UScriptStruct* SharedFragmentType) const
+	{
+		return FConstStructView(SharedFragmentType, static_cast<uint8*>(InternalGetSharedFragmentDataPtr(Entity, SharedFragmentType)));
+	}
+
 	uint32 GetArchetypeDataVersion() const { return ArchetypeDataVersion; }
 
 	/**
@@ -466,6 +500,10 @@ private:
 	void InternalAddFragmentListToEntity(FMassEntityHandle Entity, const FMassFragmentBitSet& InFragments);
 	void* InternalGetFragmentDataChecked(FMassEntityHandle Entity, const UScriptStruct* FragmentType) const;
 	void* InternalGetFragmentDataPtr(FMassEntityHandle Entity, const UScriptStruct* FragmentType) const;
+	const void* InternalGetConstSharedFragmentDataChecked(FMassEntityHandle Entity, const UScriptStruct* ConstSharedFragmentType) const;
+	const void* InternalGetConstSharedFragmentDataPtr(FMassEntityHandle Entity, const UScriptStruct* ConstSharedFragmentType) const;
+	void* InternalGetSharedFragmentDataChecked(FMassEntityHandle Entity, const UScriptStruct* SharedFragmentType) const;
+	void* InternalGetSharedFragmentDataPtr(FMassEntityHandle Entity, const UScriptStruct* SharedFragmentType) const;
 
 	TSharedRef<FEntityCreationContext> InternalBatchCreateReservedEntities(const FMassArchetypeHandle& ArchetypeHandle,
 		const FMassArchetypeSharedFragmentValues& SharedFragmentValues, TConstArrayView<FMassEntityHandle> ReservedEntities);
