Minor refactor to `FMassObserverManager::OnCompositionChanged` to make the `FMassProcessingContext` parameter not optional. In addition made `FMassObserverManager::OnPostEntitiesCreated` reuse that code rather than duplicating the logic.

diff --git a/Engine/Plugins/Runtime/MassEntity/Source/MassEntity/Private/MassObserverManager.cpp b/Engine/Plugins/Runtime/MassEntity/Source/MassEntity/Private/MassObserverManager.cpp
--- a/Engine/Plugins/Runtime/MassEntity/Source/MassEntity/Private/MassObserverManager.cpp
+++ b/Engine/Plugins/Runtime/MassEntity/Source/MassEntity/Private/MassObserverManager.cpp
@@ -9,40 +9,48 @@
 
 #include UE_INLINE_GENERATED_CPP_BY_NAME(MassObserverManager)
 
-namespace UE::Mass::ObserverManager::Private
+namespace UE::Mass::ObserverManager
 {
-// a helper function to reduce code duplication in FMassObserverManager::Initialize
-template<typename TBitSet, typename TPointerType>
-void SetUpObservers(FMassEntityManager& EntityManager, const TMap<TPointerType, FMassProcessorClassCollection>& RegisteredObserverTypes, TBitSet& ObservedBitSet, FMassObserversMap& Observers)
-{
-	ObservedBitSet.Reset();
-	UObject* Owner = EntityManager.GetOwner();
-	check(Owner);
-	const UWorld* World = Owner->GetWorld();
-	const EProcessorExecutionFlags WorldExecutionFlags = World ? UE::Mass::Utils::GetProcessorExecutionFlagsForWorld(*World) : EProcessorExecutionFlags::All;
+	namespace Tweakables
+	{
+		// Used as a template parameter for TInlineAllocator that we use when gathering UScriptStruct* of the observed types to process.
+		constexpr int InlineAllocatorElementsForOverlapTypes = 8;
+	} // Tweakables
 
-	for (auto It : RegisteredObserverTypes)
+	namespace Private
+	{
+	// a helper function to reduce code duplication in FMassObserverManager::Initialize
+	template<typename TBitSet, typename TPointerType>
+	void SetUpObservers(FMassEntityManager& EntityManager, const TMap<TPointerType, FMassProcessorClassCollection>& RegisteredObserverTypes, TBitSet& ObservedBitSet, FMassObserversMap& Observers)
 	{
-		if (It.Value.ClassCollection.Num() == 0)
+		ObservedBitSet.Reset();
+		UObject* Owner = EntityManager.GetOwner();
+		check(Owner);
+		const UWorld* World = Owner->GetWorld();
+		const EProcessorExecutionFlags WorldExecutionFlags = World ? UE::Mass::Utils::GetProcessorExecutionFlagsForWorld(*World) : EProcessorExecutionFlags::All;
+
+		for (auto It : RegisteredObserverTypes)
 		{
-			continue;
-		}
+			if (It.Value.ClassCollection.Num() == 0)
+			{
+				continue;
+			}
 
-		ObservedBitSet.Add(*It.Key);
-		FMassRuntimePipeline& Pipeline = (*Observers).FindOrAdd(It.Key);
+			ObservedBitSet.Add(*It.Key);
+			FMassRuntimePipeline& Pipeline = (*Observers).FindOrAdd(It.Key);
 
-		for (const TSubclassOf<UMassProcessor>& ProcessorClass : It.Value.ClassCollection)
-		{
-			if (ProcessorClass->GetDefaultObject<UMassProcessor>()->ShouldExecute(WorldExecutionFlags))
+			for (const TSubclassOf<UMassProcessor>& ProcessorClass : It.Value.ClassCollection)
 			{
-				Pipeline.AppendProcessor(ProcessorClass, *Owner);
+				if (ProcessorClass->GetDefaultObject<UMassProcessor>()->ShouldExecute(WorldExecutionFlags))
+				{
+					Pipeline.AppendProcessor(ProcessorClass, *Owner);
+				}
 			}
+			Pipeline.Initialize(*Owner);
 		}
-		Pipeline.Initialize(*Owner);
-	}
-};
-
-} // UE::Mass::ObserverManager::Private
+	};
+	} // Private
+} // UE::Mass::ObserverManager
 
 //----------------------------------------------------------------------//
 // FMassObserverManager
@@ -95,18 +103,8 @@ bool FMassObserverManager::OnPostEntitiesCreated(FMassProcessingContext& Process
 
 	check(ProcessingContext.EntityManager);
 	const FMassArchetypeCompositionDescriptor& ArchetypeComposition = ProcessingContext.EntityManager->GetArchetypeComposition(EntityCollection.GetArchetype());
-	const FMassFragmentBitSet Overlap = ObservedFragments[(uint8)EMassObservedOperation::Add].GetOverlap(ArchetypeComposition.Fragments);
-
-	if (Overlap.IsEmpty() == false)
-	{
-		TArray<const UScriptStruct*> OverlapTypes;
-		Overlap.ExportTypes(OverlapTypes);
-
-		HandleFragmentsImpl(ProcessingContext, EntityCollection, MakeArrayView(OverlapTypes), FragmentObservers[(uint8)EMassObservedOperation::Add]);
-		return true;
-	}
 
-	return false;
+	return OnCompositionChanged(ProcessingContext, EntityCollection, ArchetypeComposition, EMassObservedOperation::Add);
 }
 
 bool FMassObserverManager::OnPreEntitiesDestroyed(const FMassArchetypeEntityCollection& EntityCollection)
@@ -130,7 +128,7 @@ bool FMassObserverManager::OnPreEntitiesDestroyed(FMassProcessingContext& Proces
 	check(ProcessingContext.EntityManager);
 	const FMassArchetypeCompositionDescriptor& ArchetypeComposition = ProcessingContext.EntityManager->GetArchetypeComposition(EntityCollection.GetArchetype());
 	
-	return OnCompositionChanged(EntityCollection, ArchetypeComposition, EMassObservedOperation::Remove, &ProcessingContext);
+	return OnCompositionChanged(ProcessingContext, EntityCollection, ArchetypeComposition, EMassObservedOperation::Remove);
 }
 
 bool FMassObserverManager::OnPreEntityDestroyed(const FMassArchetypeCompositionDescriptor& ArchetypeComposition, const FMassEntityHandle Entity)
@@ -139,8 +137,10 @@ bool FMassObserverManager::OnPreEntityDestroyed(const FMassArchetypeCompositionD
 	return OnCompositionChanged(Entity, ArchetypeComposition, EMassObservedOperation::Remove);
 }
 
-bool FMassObserverManager::OnCompositionChanged(const FMassArchetypeEntityCollection& EntityCollection, const FMassArchetypeCompositionDescriptor& CompositionDelta, const EMassObservedOperation Operation, FMassProcessingContext* InProcessingContext)
+bool FMassObserverManager::OnCompositionChanged(FMassProcessingContext& ProcessingContext, const FMassArchetypeEntityCollection& EntityCollection, const FMassArchetypeCompositionDescriptor& CompositionDelta, const EMassObservedOperation Operation)
 {
+	using UE::Mass::ObserverManager::Tweakables::InlineAllocatorElementsForOverlapTypes;
+
 	const FMassFragmentBitSet FragmentOverlap = ObservedFragments[(uint8)Operation].GetOverlap(CompositionDelta.Fragments);
 	const bool bHasFragmentsOverlap = !FragmentOverlap.IsEmpty();
 	const FMassTagBitSet TagOverlap = ObservedTags[(uint8)Operation].GetOverlap(CompositionDelta.Tags);
@@ -148,16 +148,13 @@ bool FMassObserverManager::OnCompositionChanged(const FMassArchetypeEntityCollec
 
 	if (bHasFragmentsOverlap || bHasTagsOverlap)
 	{
-		FMassProcessingContext LocalContext(EntityManager, /*DeltaSeconds=*/0.f);
-		LocalContext.bFlushCommandBuffer = false;
-		FMassProcessingContext* ProcessingContext = InProcessingContext ? InProcessingContext : &LocalContext;
-		TArray<const UScriptStruct*> ObservedTypesOverlap;
+		TArray<const UScriptStruct*, TInlineAllocator<InlineAllocatorElementsForOverlapTypes>> ObservedTypesOverlap;
 
 		if (bHasFragmentsOverlap)
 		{
 			FragmentOverlap.ExportTypes(ObservedTypesOverlap);
 
-			HandleFragmentsImpl(*ProcessingContext, EntityCollection, ObservedTypesOverlap, FragmentObservers[(uint8)Operation]);
+			HandleFragmentsImpl(ProcessingContext, EntityCollection, ObservedTypesOverlap, FragmentObservers[(uint8)Operation]);
 		}
 
 		if (bHasTagsOverlap)
@@ -165,15 +162,19 @@ bool FMassObserverManager::OnCompositionChanged(const FMassArchetypeEntityCollec
 			ObservedTypesOverlap.Reset();
 			TagOverlap.ExportTypes(ObservedTypesOverlap);
 
-			HandleFragmentsImpl(*ProcessingContext, EntityCollection, ObservedTypesOverlap, TagObservers[(uint8)Operation]);
+			HandleFragmentsImpl(ProcessingContext, EntityCollection, ObservedTypesOverlap, TagObservers[(uint8)Operation]);
 		}
+
+		return true;
 	}
 
-	return bHasFragmentsOverlap || bHasTagsOverlap;
+	return false;
 }
 
 bool FMassObserverManager::OnCompositionChanged(const FMassEntityHandle Entity, const FMassArchetypeCompositionDescriptor& CompositionDelta, const EMassObservedOperation Operation)
 {
+	using UE::Mass::ObserverManager::Tweakables::InlineAllocatorElementsForOverlapTypes;
+
 	const FMassFragmentBitSet FragmentOverlap = ObservedFragments[(uint8)Operation].GetOverlap(CompositionDelta.Fragments);
 	const bool bHasFragmentsOverlap = !FragmentOverlap.IsEmpty();
 	const FMassTagBitSet TagOverlap = ObservedTags[(uint8)Operation].GetOverlap(CompositionDelta.Tags);
@@ -181,7 +182,7 @@ bool FMassObserverManager::OnCompositionChanged(const FMassEntityHandle Entity,
 
 	if (bHasFragmentsOverlap || bHasTagsOverlap)
 	{
-		TArray<const UScriptStruct*> ObservedTypesOverlap;
+		TArray<const UScriptStruct*, TInlineAllocator<InlineAllocatorElementsForOverlapTypes>> ObservedTypesOverlap;
 		FMassProcessingContext ProcessingContext(EntityManager, /*DeltaSeconds=*/0.f);
 		ProcessingContext.bFlushCommandBuffer = false;
 		const FMassArchetypeHandle ArchetypeHandle = EntityManager.GetArchetypeForEntity(Entity);
diff --git a/Engine/Plugins/Runtime/MassEntity/Source/MassEntity/Public/MassObserverManager.h b/Engine/Plugins/Runtime/MassEntity/Source/MassEntity/Public/MassObserverManager.h
--- a/Engine/Plugins/Runtime/MassEntity/Source/MassEntity/Public/MassObserverManager.h
+++ b/Engine/Plugins/Runtime/MassEntity/Source/MassEntity/Public/MassObserverManager.h
@@ -71,7 +71,7 @@ public:
 	}
 
 	bool OnPostEntitiesCreated(const FMassArchetypeEntityCollection& EntityCollection);
-	bool OnPostEntitiesCreated(FMassProcessingContext& ProcessingContext, const FMassArchetypeEntityCollection& EntityCollection);
+	bool OnPostEntitiesCreated(FMassProcessingContext& InProcessingContext, const FMassArchetypeEntityCollection& EntityCollection);
 
 	bool OnPreEntitiesDestroyed(const FMassArchetypeEntityCollection& EntityCollection);
 	bool OnPreEntitiesDestroyed(FMassProcessingContext& ProcessingContext, const FMassArchetypeEntityCollection& EntityCollection);
@@ -87,7 +87,14 @@ public:
 		return OnCompositionChanged(Entity, Composition, EMassObservedOperation::Remove);
 	}
 
-	bool OnCompositionChanged(const FMassArchetypeEntityCollection& EntityCollection, const FMassArchetypeCompositionDescriptor& Composition, const EMassObservedOperation Operation, FMassProcessingContext* InProcessingContext = nullptr);
+	bool OnCompositionChanged(FMassProcessingContext& InProcessingContext, const FMassArchetypeEntityCollection& EntityCollection, const FMassArchetypeCompositionDescriptor& Composition, const EMassObservedOperation Operation);
+
+	bool OnCompositionChanged(const FMassArchetypeEntityCollection& EntityCollection, const FMassArchetypeCompositionDescriptor& Composition, const EMassObservedOperation Operation)
+	{
+		FMassProcessingContext LocalContext(EntityManager, /*DeltaSeconds=*/0.f);
+		LocalContext.bFlushCommandBuffer = false;
+		return OnCompositionChanged(LocalContext, EntityCollection, Composition, Operation);
+	}
 
 	void OnPostFragmentOrTagAdded(const UScriptStruct& FragmentOrTagType, const FMassArchetypeEntityCollection& EntityCollection)
 	{
@@ -126,6 +133,17 @@ protected:
 	 * FMassObserverManager outside of an FMassEntityManager instance 
 	 */
 	FMassEntityManager& EntityManager;
+
+public:
+	UE_DEPRECATED(5.4, "This flavor of OnCompositionChanged is deprecated. Please use the one taking a FMassProcessingContext& parameter instead")
+	bool OnCompositionChanged(const FMassArchetypeEntityCollection& EntityCollection, const FMassArchetypeCompositionDescriptor& Composition, const EMassObservedOperation Operation, FMassProcessingContext* InProcessingContext)
+	{
+		if (InProcessingContext)
+		{
+			return OnCompositionChanged(*InProcessingContext, EntityCollection, Composition, Operation);
+		}
+		return OnCompositionChanged(EntityCollection, Composition, Operation);
+	}
 };
 
 template<>
