Massaged FMassEntityManager's InternalGetF*Data* functions to fail compilation rather than at runtime (wherever possible) and dropped two `Checked` functions in favor of reusing `Ptr` flavors to reduce code duplication

diff --git a/Engine/Plugins/Runtime/MassEntity/Source/MassEntity/Private/MassEntityManager.cpp b/Engine/Plugins/Runtime/MassEntity/Source/MassEntity/Private/MassEntityManager.cpp
--- a/Engine/Plugins/Runtime/MassEntity/Source/MassEntity/Private/MassEntityManager.cpp
+++ b/Engine/Plugins/Runtime/MassEntity/Source/MassEntity/Private/MassEntityManager.cpp
@@ -1212,56 +1212,37 @@ void FMassEntityManager::BatchSetEntityFragmentsValues(const FMassArchetypeEntit
 void* FMassEntityManager::InternalGetFragmentDataChecked(FMassEntityHandle Entity, const UScriptStruct* FragmentType) const
 {
 	CheckIfEntityIsActive(Entity);
-
-	checkf((FragmentType != nullptr) && FragmentType->IsChildOf(FMassFragment::StaticStruct()), TEXT("InternalGetFragmentDataChecked called with an invalid fragment type '%s'"), *GetPathNameSafe(FragmentType));
 	const FEntityData& EntityData = Entities[Entity.Index];
 	return EntityData.CurrentArchetype->GetFragmentDataForEntityChecked(FragmentType, Entity.Index);
 }
 
 void* FMassEntityManager::InternalGetFragmentDataPtr(FMassEntityHandle Entity, const UScriptStruct* FragmentType) const
 {
+	// note that FragmentType is guaranteed to be of valid type - it's either statically checked by the template versions
+	// or `checkf`ed by the non-template one
 	CheckIfEntityIsActive(Entity);
-	checkf((FragmentType != nullptr) && FragmentType->IsChildOf(FMassFragment::StaticStruct()), TEXT("InternalGetFragmentData called with an invalid fragment type '%s'"), *GetPathNameSafe(FragmentType));
 	const FEntityData& EntityData = Entities[Entity.Index];
 	return EntityData.CurrentArchetype->GetFragmentDataForEntity(FragmentType, Entity.Index);
 }
 
-const void* FMassEntityManager::InternalGetConstSharedFragmentDataChecked(FMassEntityHandle Entity, const UScriptStruct* ConstSharedFragmentType) const
-{
-	CheckIfEntityIsActive(Entity);
-	checkf((ConstSharedFragmentType != nullptr) && ConstSharedFragmentType->IsChildOf(FMassSharedFragment::StaticStruct()), TEXT("InternalGetConstSharedFragmentDataChecked called with an invalid fragment type '%s'"), *GetPathNameSafe(ConstSharedFragmentType));
-	const FEntityData& EntityData = Entities[Entity.Index];
-	const FConstSharedStruct* SharedFragment = EntityData.CurrentArchetype->GetSharedFragmentValues(Entity).GetConstSharedFragments().FindByPredicate(FStructTypeEqualOperator(ConstSharedFragmentType));
-	check(SharedFragment != nullptr);
-	return SharedFragment->GetMemory();
-}
-
-const void* FMassEntityManager::InternalGetConstSharedFragmentDataPtr(FMassEntityHandle Entity, const UScriptStruct* ConstSharedFragmentType) const
+const FConstSharedStruct* FMassEntityManager::InternalGetConstSharedFragmentPtr(FMassEntityHandle Entity, const UScriptStruct* ConstSharedFragmentType) const
 {
+	// note that ConstSharedFragmentType is guaranteed to be of valid type - it's either statically checked by the template versions
+	// or `checkf`ed by the non-template one
 	CheckIfEntityIsActive(Entity);
-	checkf((ConstSharedFragmentType != nullptr) && ConstSharedFragmentType->IsChildOf(FMassSharedFragment::StaticStruct()), TEXT("InternalGetConstSharedFragmentData called with an invalid fragment type '%s'"), *GetPathNameSafe(ConstSharedFragmentType));
 	const FEntityData& EntityData = Entities[Entity.Index];
 	const FConstSharedStruct* SharedFragment = EntityData.CurrentArchetype->GetSharedFragmentValues(Entity).GetConstSharedFragments().FindByPredicate(FStructTypeEqualOperator(ConstSharedFragmentType));
-	return (SharedFragment != nullptr) ? SharedFragment->GetMemory() : nullptr;
-}
-
-void* FMassEntityManager::InternalGetSharedFragmentDataChecked(FMassEntityHandle Entity, const UScriptStruct* SharedFragmentType) const
-{
-	CheckIfEntityIsActive(Entity);
-	checkf((SharedFragmentType != nullptr) && SharedFragmentType->IsChildOf(FMassSharedFragment::StaticStruct()), TEXT("InternalGetSharedFragmentDataChecked called with an invalid fragment type '%s'"), *GetPathNameSafe(SharedFragmentType));
-	const FEntityData& EntityData = Entities[Entity.Index];
-	const FSharedStruct* SharedFragment = EntityData.CurrentArchetype->GetSharedFragmentValues(Entity).GetSharedFragments().FindByPredicate(FStructTypeEqualOperator(SharedFragmentType));
-	check(SharedFragment != nullptr);
-	return SharedFragment->GetMemory();
+	return SharedFragment;
 }
 
-void* FMassEntityManager::InternalGetSharedFragmentDataPtr(FMassEntityHandle Entity, const UScriptStruct* SharedFragmentType) const
+const FSharedStruct* FMassEntityManager::InternalGetSharedFragmentPtr(FMassEntityHandle Entity, const UScriptStruct* SharedFragmentType) const
 {
+	// note that SharedFragmentType is guaranteed to be of valid type - it's either statically checked by the template versions
+	// or `checkf`ed by the non-template one
 	CheckIfEntityIsActive(Entity);
-	checkf((SharedFragmentType != nullptr) && SharedFragmentType->IsChildOf(FMassSharedFragment::StaticStruct()), TEXT("InternalGetSharedFragmentData called with an invalid fragment type '%s'"), *GetPathNameSafe(SharedFragmentType));
 	const FEntityData& EntityData = Entities[Entity.Index];
 	const FSharedStruct* SharedFragment = EntityData.CurrentArchetype->GetSharedFragmentValues(Entity).GetSharedFragments().FindByPredicate(FStructTypeEqualOperator(SharedFragmentType));
-	return (SharedFragment != nullptr) ? SharedFragment->GetMemory() : nullptr;
+	return SharedFragment;
 }
 
 bool FMassEntityManager::IsEntityValid(FMassEntityHandle Entity) const
diff --git a/Engine/Plugins/Runtime/MassEntity/Source/MassEntity/Public/MassEntityManager.h b/Engine/Plugins/Runtime/MassEntity/Source/MassEntity/Public/MassEntityManager.h
--- a/Engine/Plugins/Runtime/MassEntity/Source/MassEntity/Public/MassEntityManager.h
+++ b/Engine/Plugins/Runtime/MassEntity/Source/MassEntity/Public/MassEntityManager.h
@@ -317,55 +317,80 @@ public:
 	// Asserts that IsEntityBuilt
 	void CheckIfEntityIsActive(FMassEntityHandle Entity) const;
 
-	template <typename FragmentType>
+	template<typename FragmentType>
 	FragmentType& GetFragmentDataChecked(FMassEntityHandle Entity) const
 	{
+		static_assert(TIsDerivedFrom<FragmentType, FMassFragment>::IsDerived
+			, "Given struct doesn't represent a valid fragment type. Make sure to inherit from FMassFragment or one of its child-types.");
 		return *((FragmentType*)InternalGetFragmentDataChecked(Entity, FragmentType::StaticStruct()));
 	}
 
-	template <typename FragmentType>
+	template<typename FragmentType>
 	FragmentType* GetFragmentDataPtr(FMassEntityHandle Entity) const
 	{
+		static_assert(TIsDerivedFrom<FragmentType, FMassFragment>::IsDerived
+			, "Given struct doesn't represent a valid fragment type. Make sure to inherit from FMassFragment or one of its child-types.");
 		return (FragmentType*)InternalGetFragmentDataPtr(Entity, FragmentType::StaticStruct());
 	}
 
 	FStructView GetFragmentDataStruct(FMassEntityHandle Entity, const UScriptStruct* FragmentType) const
 	{
+		checkf((FragmentType != nullptr) && FragmentType->IsChildOf(FMassFragment::StaticStruct())
+			, TEXT("GetFragmentDataStruct called with an invalid fragment type '%s'"), *GetPathNameSafe(FragmentType));
 		return FStructView(FragmentType, static_cast<uint8*>(InternalGetFragmentDataPtr(Entity, FragmentType)));
 	}
 
-	template <typename ConstSharedFragmentType>
-	ConstSharedFragmentType& GetConstSharedFragmentDataChecked(FMassEntityHandle Entity) const
+	template<typename ConstSharedFragmentType>
+	ConstSharedFragmentType* GetConstSharedFragmentDataPtr(FMassEntityHandle Entity) const
 	{
-		return *((ConstSharedFragmentType*)InternalGetConstSharedFragmentDataChecked(Entity, ConstSharedFragmentType::StaticStruct()));
+		static_assert(TIsDerivedFrom<ConstSharedFragmentType, FMassSharedFragment>::IsDerived, "Given struct doesn't represent a valid shared fragment type. Make sure to inherit from FMassSharedFragment or one of its child-types.");
+		const FConstSharedStruct* ConstSharedStruct = InternalGetConstSharedFragmentPtr(Entity, ConstSharedFragmentType::StaticStruct());
+		return (ConstSharedFragmentType*)(ConstSharedStruct ? ConstSharedStruct->GetMemory() : nullptr);
 	}
 
-	template <typename ConstSharedFragmentType>
-	ConstSharedFragmentType* GetConstSharedFragmentDataPtr(FMassEntityHandle Entity) const
+	template<typename ConstSharedFragmentType>
+	ConstSharedFragmentType& GetConstSharedFragmentDataChecked(FMassEntityHandle Entity) const
 	{
-		return (ConstSharedFragmentType*)InternalGetConstSharedFragmentDataPtr(Entity, ConstSharedFragmentType::StaticStruct());
+		ConstSharedFragmentType* TypePtr = GetConstSharedFragmentDataPtr<ConstSharedFragmentType>(Entity);
+		check(TypePtr);
+		return *TypePtr;
 	}
 
 	FConstStructView GetConstSharedFragmentDataStruct(FMassEntityHandle Entity, const UScriptStruct* ConstSharedFragmentType) const
 	{
-		return FConstStructView(ConstSharedFragmentType, static_cast<const uint8*>(InternalGetConstSharedFragmentDataPtr(Entity, ConstSharedFragmentType)));
+		checkf((ConstSharedFragmentType != nullptr) && ConstSharedFragmentType->IsChildOf(FMassSharedFragment::StaticStruct())
+			, TEXT("GetConstSharedFragmentDataStruct called with an invalid fragment type '%s'"), *GetPathNameSafe(ConstSharedFragmentType));
+		const FConstSharedStruct* ConstSharedStruct = InternalGetConstSharedFragmentPtr(Entity, ConstSharedFragmentType);
+		return ConstSharedStruct
+			? FConstStructView(*ConstSharedStruct)
+			: FConstStructView();
 	}
 
-	template <typename SharedFragmentType>
-	SharedFragmentType& GetSharedFragmentDataChecked(FMassEntityHandle Entity) const
+	template<typename SharedFragmentType>
+	SharedFragmentType* GetSharedFragmentDataPtr(FMassEntityHandle Entity) const
 	{
-		return *((SharedFragmentType*)InternalGetSharedFragmentDataChecked(Entity, SharedFragmentType::StaticStruct()));
+		static_assert(TIsDerivedFrom<SharedFragmentType, FMassSharedFragment>::IsDerived
+			, "Given struct doesn't represent a valid shared fragment type. Make sure to inherit from FMassSharedFragment or one of its child-types.");
+		const FSharedStruct* FragmentPtr = InternalGetSharedFragmentPtr(Entity, SharedFragmentType::StaticStruct());
+		return (SharedFragmentType*)(FragmentPtr ? FragmentPtr->GetMemory() : nullptr);
 	}
 
-	template <typename SharedFragmentType>
-	SharedFragmentType* GetSharedFragmentDataPtr(FMassEntityHandle Entity) const
+	template<typename SharedFragmentType>
+	SharedFragmentType& GetSharedFragmentDataChecked(FMassEntityHandle Entity) const
 	{
-		return (SharedFragmentType*)InternalGetSharedFragmentDataPtr(Entity, SharedFragmentType::StaticStruct());
+		SharedFragmentType* TypePtr = GetSharedFragmentDataPtr<SharedFragmentType>(Entity);
+		check(TypePtr);
+		return *TypePtr;
 	}
 
 	FConstStructView GetSharedFragmentDataStruct(FMassEntityHandle Entity, const UScriptStruct* SharedFragmentType) const
 	{
-		return FConstStructView(SharedFragmentType, static_cast<uint8*>(InternalGetSharedFragmentDataPtr(Entity, SharedFragmentType)));
+		checkf((SharedFragmentType != nullptr) && SharedFragmentType->IsChildOf(FMassSharedFragment::StaticStruct())
+			, TEXT("GetSharedFragmentDataStruct called with an invalid fragment type '%s'"), *GetPathNameSafe(SharedFragmentType));
+		const FSharedStruct* FragmentPtr = InternalGetSharedFragmentPtr(Entity, SharedFragmentType);
+		return FragmentPtr
+			? FConstStructView(*FragmentPtr)
+			: FConstStructView();
 	}
 
 	uint32 GetArchetypeDataVersion() const { return ArchetypeDataVersion; }
@@ -498,12 +523,14 @@ private:
 	 *  fragment list. It's callers responsibility to ensure that's true. Failing this will cause a `check` fail.
 	 */
 	void InternalAddFragmentListToEntity(FMassEntityHandle Entity, const FMassFragmentBitSet& InFragments);
+	/** Note that it's the caller's responsibility to ensure `FragmentType` is a kind of FMassFragment */
 	void* InternalGetFragmentDataChecked(FMassEntityHandle Entity, const UScriptStruct* FragmentType) const;
+	/** Note that it's the caller's responsibility to ensure `FragmentType` is a kind of FMassFragment */
 	void* InternalGetFragmentDataPtr(FMassEntityHandle Entity, const UScriptStruct* FragmentType) const;
-	const void* InternalGetConstSharedFragmentDataChecked(FMassEntityHandle Entity, const UScriptStruct* ConstSharedFragmentType) const;
-	const void* InternalGetConstSharedFragmentDataPtr(FMassEntityHandle Entity, const UScriptStruct* ConstSharedFragmentType) const;
-	void* InternalGetSharedFragmentDataChecked(FMassEntityHandle Entity, const UScriptStruct* SharedFragmentType) const;
-	void* InternalGetSharedFragmentDataPtr(FMassEntityHandle Entity, const UScriptStruct* SharedFragmentType) const;
+	/** Note that it's the caller's responsibility to ensure `ConstSharedFragmentType` is a kind of FMassSharedFragment */
+	const FConstSharedStruct* InternalGetConstSharedFragmentPtr(FMassEntityHandle Entity, const UScriptStruct* ConstSharedFragmentType) const;
+	/** Note that it's the caller's responsibility to ensure `SharedFragmentType` is a kind of FMassSharedFragment */
+	const FSharedStruct* InternalGetSharedFragmentPtr(FMassEntityHandle Entity, const UScriptStruct* SharedFragmentType) const;
 
 	TSharedRef<FEntityCreationContext> InternalBatchCreateReservedEntities(const FMassArchetypeHandle& ArchetypeHandle,
 		const FMassArchetypeSharedFragmentValues& SharedFragmentValues, TConstArrayView<FMassEntityHandle> ReservedEntities);
