Added a missing comment

diff --git a/Engine/Plugins/Runtime/MassEntity/Source/MassEntity/Private/MassEntityManager.cpp b/Engine/Plugins/Runtime/MassEntity/Source/MassEntity/Private/MassEntityManager.cpp
--- a/Engine/Plugins/Runtime/MassEntity/Source/MassEntity/Private/MassEntityManager.cpp
+++ b/Engine/Plugins/Runtime/MassEntity/Source/MassEntity/Private/MassEntityManager.cpp
@@ -1211,6 +1211,8 @@ void FMassEntityManager::BatchSetEntityFragmentsValues(const FMassArchetypeEntit
 
 void* FMassEntityManager::InternalGetFragmentDataChecked(FMassEntityHandle Entity, const UScriptStruct* FragmentType) const
 {
+	// note that FragmentType is guaranteed to be of valid type - it's either statically checked by the template versions
+	// or `checkf`ed by the non-template one
 	CheckIfEntityIsActive(Entity);
 	const FEntityData& EntityData = Entities[Entity.Index];
 	return EntityData.CurrentArchetype->GetFragmentDataForEntityChecked(FragmentType, Entity.Index);
