Replaced FMassTagBitSet with FMassArchetypeCompositionDescriptor in the MassExecutionContext.

This follows the addition of DoesArchetypeHaveFragment which originally used the fragment view, making the lookup more expensive than it needed to be. This change gives access to the full bitsets used by an archetype to provide a more efficient lookup as well as room to expand in the future for other types of fragments.

diff --git a/Engine/Plugins/Runtime/MassEntity/Source/MassEntity/Private/MassArchetypeData.cpp b/Engine/Plugins/Runtime/MassEntity/Source/MassEntity/Private/MassArchetypeData.cpp
--- a/Engine/Plugins/Runtime/MassEntity/Source/MassEntity/Private/MassArchetypeData.cpp
+++ b/Engine/Plugins/Runtime/MassEntity/Source/MassEntity/Private/MassArchetypeData.cpp
@@ -416,7 +416,7 @@ void FMassArchetypeData::ExecuteFunction(FMassExecutionContext& RunContext, cons
 	}
 
 	// mz@todo to be removed
-	RunContext.SetCurrentArchetypesTagBitSet(GetTagBitSet());
+	RunContext.SetCurrentArchetypeCompositionDescriptor(GetCompositionDescriptor());
 
 	uint32 PrevSharedFragmentValuesHash = UINT32_MAX;
 	for (FMassArchetypeChunkIterator ChunkIterator(EntityRangeContainer); ChunkIterator; ++ChunkIterator)
@@ -456,7 +456,7 @@ void FMassArchetypeData::ExecuteFunction(FMassExecutionContext& RunContext, cons
 	}
 
 	// mz@todo to be removed
-	RunContext.SetCurrentArchetypesTagBitSet(GetTagBitSet());
+	RunContext.SetCurrentArchetypeCompositionDescriptor(GetCompositionDescriptor());
 
 	uint32 PrevSharedFragmentValuesHash = UINT32_MAX;
 	for (FMassArchetypeChunk& Chunk : Chunks)
@@ -493,7 +493,7 @@ void FMassArchetypeData::ExecutionFunctionForChunk(FMassExecutionContext RunCont
 		BindConstSharedFragmentRequirements(RunContext, Chunk.GetSharedFragmentValues(), RequirementMapping.ConstSharedFragments);
 		BindSharedFragmentRequirements(RunContext, Chunk.GetMutableSharedFragmentValues(), RequirementMapping.SharedFragments);
 
-		RunContext.SetCurrentArchetypesTagBitSet(GetTagBitSet());
+		RunContext.SetCurrentArchetypeCompositionDescriptor(GetCompositionDescriptor());
 		RunContext.SetCurrentChunkSerialModificationNumber(Chunk.GetSerialModificationNumber());
 		BindChunkFragmentRequirements(RunContext, RequirementMapping.ChunkFragments, Chunk);
 
diff --git a/Engine/Plugins/Runtime/MassEntity/Source/MassEntity/Private/MassExecutionContext.cpp b/Engine/Plugins/Runtime/MassEntity/Source/MassEntity/Private/MassExecutionContext.cpp
--- a/Engine/Plugins/Runtime/MassEntity/Source/MassEntity/Private/MassExecutionContext.cpp
+++ b/Engine/Plugins/Runtime/MassEntity/Source/MassEntity/Private/MassExecutionContext.cpp
@@ -29,7 +29,7 @@ void FMassExecutionContext::ClearExecutionData()
 	ChunkFragmentViews.Reset();
 	ConstSharedFragmentViews.Reset();
 	SharedFragmentViews.Reset();
-	CurrentArchetypesTagBitSet.Reset();
+	CurrentArchetypeCompositionDescriptor.Reset();
 	ChunkSerialModificationNumber = -1;
 }
 
diff --git a/Engine/Plugins/Runtime/MassEntity/Source/MassEntity/Public/MassExecutionContext.h b/Engine/Plugins/Runtime/MassEntity/Source/MassEntity/Public/MassExecutionContext.h
--- a/Engine/Plugins/Runtime/MassEntity/Source/MassEntity/Public/MassExecutionContext.h
+++ b/Engine/Plugins/Runtime/MassEntity/Source/MassEntity/Public/MassExecutionContext.h
@@ -49,7 +49,7 @@ private:
 	FInstancedStruct AuxData;
 	float DeltaTimeSeconds = 0.0f;
 	int32 ChunkSerialModificationNumber = -1;
-	FMassTagBitSet CurrentArchetypesTagBitSet;
+	FMassArchetypeCompositionDescriptor CurrentArchetypeCompositionDescriptor;
 
 	TSharedRef<FMassEntityManager> EntityManager;
 
@@ -121,30 +121,26 @@ public:
 
 	bool DoesArchetypeHaveFragment(const UScriptStruct& FragmentType) const
 	{
-		return FragmentViews.FindByPredicate(
-			[&FragmentType](const FFragmentView& Element) 
-			{ 
-				return Element.Requirement.StructType == &FragmentType; 
-			}) != nullptr;
+		return CurrentArchetypeCompositionDescriptor.Fragments.Contains(FragmentType);
 	}
 
 	template<typename T>
 	bool DoesArchetypeHaveFragment() const
 	{
 		static_assert(TIsDerivedFrom<T, FMassFragment>::IsDerived, "Given struct is not of a valid fragment type.");
-		return DoesArchetypeHaveFragment(T::StaticStruct());
+		return CurrentArchetypeCompositionDescriptor.Fragments.Contains<T>();
 	}
 
 	bool DoesArchetypeHaveTag(const UScriptStruct& TagType) const
 	{
-		return CurrentArchetypesTagBitSet.Contains(TagType);
+		return CurrentArchetypeCompositionDescriptor.Tags.Contains(TagType);
 	}
 
 	template<typename T>
 	bool DoesArchetypeHaveTag() const
 	{
 		static_assert(TIsDerivedFrom<T, FMassTag>::IsDerived, "Given struct is not of a valid tag type.");
-		return CurrentArchetypesTagBitSet.Contains<T>();
+		return CurrentArchetypeCompositionDescriptor.Tags.Contains<T>();
 	}
 
 	/** Chunk related operations */
@@ -338,9 +334,9 @@ public:
 	void FlushDeferred();
 
 	void ClearExecutionData();
-	void SetCurrentArchetypesTagBitSet(const FMassTagBitSet& BitSet)
+	void SetCurrentArchetypeCompositionDescriptor(const FMassArchetypeCompositionDescriptor& Descriptor)
 	{
-		CurrentArchetypesTagBitSet = BitSet;
+		CurrentArchetypeCompositionDescriptor = Descriptor;
 	}
 
 	/** 
@@ -447,4 +443,7 @@ public:
 	{
 		return CacheSubsystemRequirements(SubsystemRequirements);
 	}
+
+	UE_DEPRECATED(5.4, "Deprecated in favor of 'SetCurrentArchetypeCompositionDescriptor' as this provides information on the entire archetype.")
+	void SetCurrentArchetypesTagBitSet(const FMassTagBitSet&) {}
 };
