Added `FMassEntityManager.SharedFragmentsTypeMap` that maps a Mass shared fragment type to a all created `FSharedStruct` instances hosting that given type.

Modified `ForEachSharedFragment` to take advantage of the newly created map. The aim is to improve performance for iteration in context where there's a lot of shared fragments and multiple linear searches a frame can stack up to noticeable performance cost.
Also added a flavor of the function, `ForEachSharedFragmentConditional`, that lets caller specify a shared-fragment-related condtion function to further short-circuit shared fragments iteration.

Also constified return types of `GetOrCreateConstSharedFragmentByHash` and `GetOrCreateSharedFragmentByHash` to emphasize that callers are never expected to change the referenced `FSharedStruct` instances

diff --git a/Engine/Plugins/Runtime/MassEntity/Source/MassEntity/Public/MassEntityManager.h b/Engine/Plugins/Runtime/MassEntity/Source/MassEntity/Public/MassEntityManager.h
--- a/Engine/Plugins/Runtime/MassEntity/Source/MassEntity/Public/MassEntityManager.h
+++ b/Engine/Plugins/Runtime/MassEntity/Source/MassEntity/Public/MassEntityManager.h
@@ -421,7 +421,7 @@ public:
 	 * Shared fragment creation methods
 	 */
 	template<typename T>
-	FConstSharedStruct& GetOrCreateConstSharedFragmentByHash(const uint32 Hash, const T& Fragment)
+	const FConstSharedStruct& GetOrCreateConstSharedFragmentByHash(const uint32 Hash, const T& Fragment)
 	{
 		static_assert(TIsDerivedFrom<T, FMassSharedFragment>::IsDerived, "Given struct doesn't represent a valid shared fragment type. Make sure to inherit from FMassSharedFragment or one of its child-types.");
 		int32& Index = ConstSharedFragmentsMap.FindOrAddByHash(Hash, Hash, INDEX_NONE);
@@ -433,14 +433,14 @@ public:
 	}
 
 	template<typename T>
-	FConstSharedStruct& GetOrCreateConstSharedFragment(const T& Fragment)
+	const FConstSharedStruct& GetOrCreateConstSharedFragment(const T& Fragment)
 	{
 		const uint32 Hash = UE::StructUtils::GetStructCrc32(FConstStructView::Make(Fragment));
 		return GetOrCreateConstSharedFragmentByHash(Hash, Fragment);
 	}
 
 	template<typename T, typename... TArgs>
-	FSharedStruct& GetOrCreateSharedFragmentByHash(const uint32 Hash, TArgs&&... InArgs)
+	const FSharedStruct& GetOrCreateSharedFragmentByHash(const uint32 Hash, TArgs&&... InArgs)
 	{
 		static_assert(TIsDerivedFrom<T, FMassSharedFragment>::IsDerived, "Given struct doesn't represent a valid shared fragment type. Make sure to inherit from FMassSharedFragment or one of its child-types.");
 
@@ -448,20 +448,39 @@ public:
 		if (Index == INDEX_NONE)
 		{
 			Index = SharedFragments.Add(FSharedStruct::Make<T>(Forward<TArgs>(InArgs)...));
+			// note that even though we're copying the freshly created FSharedStruct instance it's perfectly fine since 
+			// FSharedStruct do guarantee there's not going to be data duplication (via a member shared pointer to hosted data)
+			TArray<FSharedStruct>& InstancedOfType = SharedFragmentsTypeMap.FindOrAdd(T::StaticStruct(), {});
+			InstancedOfType.Add(SharedFragments[Index]);
 		}
 
 		return SharedFragments[Index];
 	}
 
 	template<typename T>
-	void ForEachSharedFragment(TFunction< void(T& /*SharedFragment*/) > ExecuteFunction)
+	void ForEachSharedFragment(TFunctionRef< void(T& /*SharedFragment*/) > ExecuteFunction)
 	{
-		FStructTypeEqualOperator Predicate(T::StaticStruct());
-		for (FSharedStruct& Struct : SharedFragments)
+		if (TArray<FSharedStruct>* InstancedOfType = SharedFragmentsTypeMap.Find(T::StaticStruct()))
 		{
-			if (Predicate(Struct))
+			for (const FSharedStruct& SharedStruct : *InstancedOfType)
 			{
-				ExecuteFunction(Struct.Get<T>());
+				ExecuteFunction(SharedStruct.Get<T>());
+			}
+		}
+	}
+
+	template<typename T>
+	void ForEachSharedFragmentConditional(TFunctionRef< bool(T& /*SharedFragment*/) > ConditionFunction, TFunctionRef< void(T& /*SharedFragment*/) > ExecuteFunction)
+	{
+		if (TArray<FSharedStruct>* InstancedOfType = SharedFragmentsTypeMap.Find(T::StaticStruct()))
+		{
+			for (const FSharedStruct& SharedStruct : *InstancedOfType)
+			{
+				T& StructInstanceRef = SharedStruct.Get<T>();
+				if (ConditionFunction(StructInstanceRef))
+				{
+					ExecuteFunction(StructInstanceRef);
+				}
 			}
 		}
 	}
@@ -564,8 +583,10 @@ private:
 	TMap<uint32, int32> ConstSharedFragmentsMap;
 
 	TArray<FSharedStruct> SharedFragments;
-	// Hash/Index in array pair
+	// Hash/Index in array pair, indices point at SharedFragments
 	TMap<uint32, int32> SharedFragmentsMap;
+	// Maps specific struct type to a collection of FSharedStruct instances of that type
+	TMap<UScriptStruct*, TArray<FSharedStruct>> SharedFragmentsTypeMap;
 
 	FMassObserverManager ObserverManager;
 
