Added an innocent little function to FMassEntityDebugger to fetch all shared fragments representing a given FMassSharedFragment-based struct type.

While in the file I've fixed a couple of misspellings introduced recently.

diff --git a/Engine/Plugins/Runtime/MassEntity/Source/MassEntity/Public/MassEntityManager.h b/Engine/Plugins/Runtime/MassEntity/Source/MassEntity/Public/MassEntityManager.h
--- a/Engine/Plugins/Runtime/MassEntity/Source/MassEntity/Public/MassEntityManager.h
+++ b/Engine/Plugins/Runtime/MassEntity/Source/MassEntity/Public/MassEntityManager.h
@@ -366,6 +366,16 @@ public:
 			: FConstStructView();
 	}
 
+	template<typename SharedFragmentType>
+	TConstArrayView<FSharedStruct> GetSharedFragmentsOfType()
+	{
+		static_assert(TIsDerivedFrom<SharedFragmentType, FMassSharedFragment>::IsDerived
+			, "Given struct doesn't represent a valid shared fragment type. Make sure to inherit from FMassSharedFragment or one of its child-types.");
+
+		TArray<FSharedStruct>* InstancesOfType = SharedFragmentsTypeMap.Find(SharedFragmentType::StaticStruct());
+		return InstancesOfType ? *InstancesOfType : TConstArrayView<FSharedStruct>();
+	}
+
 	template<typename SharedFragmentType>
 	SharedFragmentType* GetSharedFragmentDataPtr(FMassEntityHandle Entity) const
 	{
@@ -450,8 +460,8 @@ public:
 			Index = SharedFragments.Add(FSharedStruct::Make<T>(Forward<TArgs>(InArgs)...));
 			// note that even though we're copying the freshly created FSharedStruct instance it's perfectly fine since 
 			// FSharedStruct do guarantee there's not going to be data duplication (via a member shared pointer to hosted data)
-			TArray<FSharedStruct>& InstancedOfType = SharedFragmentsTypeMap.FindOrAdd(T::StaticStruct(), {});
-			InstancedOfType.Add(SharedFragments[Index]);
+			TArray<FSharedStruct>& InstancesOfType = SharedFragmentsTypeMap.FindOrAdd(T::StaticStruct(), {});
+			InstancesOfType.Add(SharedFragments[Index]);
 		}
 
 		return SharedFragments[Index];
@@ -460,9 +470,9 @@ public:
 	template<typename T>
 	void ForEachSharedFragment(TFunctionRef< void(T& /*SharedFragment*/) > ExecuteFunction)
 	{
-		if (TArray<FSharedStruct>* InstancedOfType = SharedFragmentsTypeMap.Find(T::StaticStruct()))
+		if (TArray<FSharedStruct>* InstancesOfType = SharedFragmentsTypeMap.Find(T::StaticStruct()))
 		{
-			for (const FSharedStruct& SharedStruct : *InstancedOfType)
+			for (const FSharedStruct& SharedStruct : *InstancesOfType)
 			{
 				ExecuteFunction(SharedStruct.Get<T>());
 			}
@@ -472,9 +482,9 @@ public:
 	template<typename T>
 	void ForEachSharedFragmentConditional(TFunctionRef< bool(T& /*SharedFragment*/) > ConditionFunction, TFunctionRef< void(T& /*SharedFragment*/) > ExecuteFunction)
 	{
-		if (TArray<FSharedStruct>* InstancedOfType = SharedFragmentsTypeMap.Find(T::StaticStruct()))
+		if (TArray<FSharedStruct>* InstancesOfType = SharedFragmentsTypeMap.Find(T::StaticStruct()))
 		{
-			for (const FSharedStruct& SharedStruct : *InstancedOfType)
+			for (const FSharedStruct& SharedStruct : *InstancesOfType)
 			{
 				T& StructInstanceRef = SharedStruct.Get<T>();
 				if (ConditionFunction(StructInstanceRef))
