Added GetNumMatchingEntities overload to FMassEntityQuery for testing against an FMassArchetypeEntityCollection array.

Also conformed some Collections and Collection parameter names to EntityCollection(s) to match wider Mass convention.

diff --git a/Engine/Plugins/Runtime/MassEntity/Source/MassEntity/Private/MassEntityQuery.cpp b/Engine/Plugins/Runtime/MassEntity/Source/MassEntity/Private/MassEntityQuery.cpp
--- a/Engine/Plugins/Runtime/MassEntity/Source/MassEntity/Private/MassEntityQuery.cpp
+++ b/Engine/Plugins/Runtime/MassEntity/Source/MassEntity/Private/MassEntityQuery.cpp
@@ -138,18 +138,18 @@ void FMassEntityQuery::CacheArchetypes(const FMassEntityManager& InEntityManager
 	}
 }
 
-void FMassEntityQuery::ForEachEntityChunkInCollections(TConstArrayView<FMassArchetypeEntityCollection> Collections, FMassEntityManager& EntityManager, FMassExecutionContext& ExecutionContext, const FMassExecuteFunction& ExecuteFunction)
+void FMassEntityQuery::ForEachEntityChunkInCollections(TConstArrayView<FMassArchetypeEntityCollection> EntityCollections, FMassEntityManager& EntityManager, FMassExecutionContext& ExecutionContext, const FMassExecuteFunction& ExecuteFunction)
 {
-	for (const FMassArchetypeEntityCollection& Collection : Collections)
+	for (const FMassArchetypeEntityCollection& EntityCollection : EntityCollections)
 	{
-		ForEachEntityChunk(Collection, EntityManager, ExecutionContext, ExecuteFunction);
+		ForEachEntityChunk(EntityCollection, EntityManager, ExecutionContext, ExecuteFunction);
 	}
 }
 
-void FMassEntityQuery::ForEachEntityChunk(const FMassArchetypeEntityCollection& Collection, FMassEntityManager& EntityManager, FMassExecutionContext& ExecutionContext, const FMassExecuteFunction& ExecuteFunction)
+void FMassEntityQuery::ForEachEntityChunk(const FMassArchetypeEntityCollection& EntityCollection, FMassEntityManager& EntityManager, FMassExecutionContext& ExecutionContext, const FMassExecuteFunction& ExecuteFunction)
 {
 	// mz@todo I don't like that we're copying data here.
-	ExecutionContext.SetEntityCollection(Collection);
+	ExecutionContext.SetEntityCollection(EntityCollection);
 	ForEachEntityChunk(EntityManager, ExecutionContext, ExecuteFunction);
 	ExecutionContext.ClearEntityCollection();
 }
@@ -266,6 +266,22 @@ int32 FMassEntityQuery::GetNumMatchingEntities(FMassEntityManager& InEntityManag
 	return TotalEntities;
 }
 
+int32 FMassEntityQuery::GetNumMatchingEntities(TConstArrayView<FMassArchetypeEntityCollection> EntityCollections)
+{
+	int32 TotalEntities = 0;
+	for (const FMassArchetypeEntityCollection& EntityCollection : EntityCollections)
+	{
+		if (DoesArchetypeMatchRequirements(EntityCollection.GetArchetype()))
+		{
+			for (const FMassArchetypeEntityCollection::FArchetypeEntityRange& EntityRange : EntityCollection.GetRanges())
+			{
+				TotalEntities += EntityRange.Length;
+			}
+		}
+	}
+	return TotalEntities;
+}
+
 bool FMassEntityQuery::HasMatchingEntities(FMassEntityManager& InEntityManager)
 {
 	CacheArchetypes(InEntityManager);
diff --git a/Engine/Plugins/Runtime/MassEntity/Source/MassEntity/Public/MassEntityQuery.h b/Engine/Plugins/Runtime/MassEntity/Source/MassEntity/Public/MassEntityQuery.h
--- a/Engine/Plugins/Runtime/MassEntity/Source/MassEntity/Public/MassEntityQuery.h
+++ b/Engine/Plugins/Runtime/MassEntity/Source/MassEntity/Public/MassEntityQuery.h
@@ -39,9 +39,9 @@ public:
 	void ForEachEntityChunk(FMassEntityManager& EntityManager, FMassExecutionContext& ExecutionContext, const FMassExecuteFunction& ExecuteFunction);
 	
 	/** Will first verify that the archetype given with Collection matches the query's requirements, and if so will run the other, more generic ForEachEntityChunk implementation */
-	void ForEachEntityChunk(const FMassArchetypeEntityCollection& Collection, FMassEntityManager& EntitySubsystem, FMassExecutionContext& ExecutionContext, const FMassExecuteFunction& ExecuteFunction);
+	void ForEachEntityChunk(const FMassArchetypeEntityCollection& EntityCollection, FMassEntityManager& EntitySubsystem, FMassExecutionContext& ExecutionContext, const FMassExecuteFunction& ExecuteFunction);
 
-	void ForEachEntityChunkInCollections(TConstArrayView<FMassArchetypeEntityCollection> Collections, FMassEntityManager& EntityManager, FMassExecutionContext& ExecutionContext, const FMassExecuteFunction& ExecuteFunction);
+	void ForEachEntityChunkInCollections(TConstArrayView<FMassArchetypeEntityCollection> EntityCollections, FMassEntityManager& EntityManager, FMassExecutionContext& ExecutionContext, const FMassExecuteFunction& ExecuteFunction);
 
 	/** Will gather all archetypes from InEntityManager matching this->Requirements.
 	 *  Note that no work will be done if the cached data is up to date (as tracked by EntitySubsystemHash and 
@@ -82,6 +82,13 @@ public:
 	 */
 	int32 GetNumMatchingEntities(FMassEntityManager& InEntityManager);
 
+	/** 
+	 * Sums the entity range lengths for each collection in EntityCollections, where the collection's 
+	 * archetype matches the querie's requirements.
+	 * @return the number of entities this given query would process if called "now" for EntityCollections
+	 */
+	int32 GetNumMatchingEntities(TConstArrayView<FMassArchetypeEntityCollection> EntityCollections);
+
 	/**
 	 * Checks if any of ValidArchetypes has any entities.
 	 * Note that the function is not const because calling it can result in re-caching of ValidArchetypes
