Added a master-switch console variable that gates Mass runtime subsystems' creation. The change doesn't change the default behavior.

diff --git a/Engine/Plugins/AI/MassAI/Source/MassAIBehavior/Public/MassComponentHitSubsystem.h b/Engine/Plugins/AI/MassAI/Source/MassAIBehavior/Public/MassComponentHitSubsystem.h
--- a/Engine/Plugins/AI/MassAI/Source/MassAIBehavior/Public/MassComponentHitSubsystem.h
+++ b/Engine/Plugins/AI/MassAI/Source/MassAIBehavior/Public/MassComponentHitSubsystem.h
@@ -4,7 +4,7 @@
 
 #include "MassEntityTypes.h"
 #include "MassComponentHitTypes.h"
-#include "Subsystems/WorldSubsystem.h"
+#include "MassSubsystemBase.h"
 #include "Engine/HitResult.h"
 
 #include "MassComponentHitSubsystem.generated.h"
@@ -18,7 +18,7 @@ class UCapsuleComponent;
  * Subsystem that keeps track of the latest component hits and allow mass entities to retrieve and handle them
  */
 UCLASS()
-class MASSAIBEHAVIOR_API UMassComponentHitSubsystem : public UTickableWorldSubsystem
+class MASSAIBEHAVIOR_API UMassComponentHitSubsystem : public UMassTickableSubsystemBase
 {
 	GENERATED_BODY()
 
diff --git a/Engine/Plugins/AI/MassAI/Source/MassAIBehavior/Public/MassStateTreeSubsystem.h b/Engine/Plugins/AI/MassAI/Source/MassAIBehavior/Public/MassStateTreeSubsystem.h
--- a/Engine/Plugins/AI/MassAI/Source/MassAIBehavior/Public/MassStateTreeSubsystem.h
+++ b/Engine/Plugins/AI/MassAI/Source/MassAIBehavior/Public/MassStateTreeSubsystem.h
@@ -3,7 +3,7 @@
 #pragma once
 
 #include "MassStateTreeTypes.h"
-#include "Subsystems/WorldSubsystem.h"
+#include "MassSubsystemBase.h"
 #include "StateTreeExecutionContext.h"
 #include "MassExternalSubsystemTraits.h"
 #include "MassStateTreeSubsystem.generated.h"
@@ -26,16 +26,16 @@ struct MASSAIBEHAVIOR_API FMassStateTreeInstanceDataItem
 * A subsystem managing StateTree assets in Mass
 */
 UCLASS()
-class MASSAIBEHAVIOR_API UMassStateTreeSubsystem : public UWorldSubsystem
+class MASSAIBEHAVIOR_API UMassStateTreeSubsystem : public UMassSubsystemBase
 {
 	GENERATED_BODY()
 	
-public:
-
+protected:
 	// USubsystem BEGIN
 	virtual void Initialize(FSubsystemCollectionBase& Collection) override;
 	// USubsystem END
 
+public:
 	/**
 	 * Allocates new instance data for specified StateTree.
 	 * @param StateTree StateTree to allocated the data for.
diff --git a/Engine/Plugins/AI/MassAI/Source/MassNavigation/Public/MassNavigationSubsystem.h b/Engine/Plugins/AI/MassAI/Source/MassNavigation/Public/MassNavigationSubsystem.h
--- a/Engine/Plugins/AI/MassAI/Source/MassNavigation/Public/MassNavigationSubsystem.h
+++ b/Engine/Plugins/AI/MassAI/Source/MassNavigation/Public/MassNavigationSubsystem.h
@@ -4,7 +4,7 @@
 
 #include "MassEntityTypes.h"
 #include "HierarchicalHashGrid2D.h"
-#include "Subsystems/WorldSubsystem.h"
+#include "MassSubsystemBase.h"
 #include "MassNavigationSubsystem.generated.h"
 
 UENUM()
@@ -29,19 +29,20 @@ struct FMassNavigationObstacleItem
 typedef THierarchicalHashGrid2D<2, 4, FMassNavigationObstacleItem> FNavigationObstacleHashGrid2D;	// 2 levels of hierarchy, 4 ratio between levels
 
 UCLASS()
-class MASSNAVIGATION_API UMassNavigationSubsystem : public UWorldSubsystem
+class MASSNAVIGATION_API UMassNavigationSubsystem : public UMassSubsystemBase
 {
 	GENERATED_BODY()
 
 public:
 	UMassNavigationSubsystem();
 
-	virtual void Initialize(FSubsystemCollectionBase& Collection) override;
-
 	const FNavigationObstacleHashGrid2D& GetObstacleGrid() const { return AvoidanceObstacleGrid; }
 	FNavigationObstacleHashGrid2D& GetObstacleGridMutable() { return AvoidanceObstacleGrid; }
 
 protected:
+	// USubsystem BEGIN
+	virtual void Initialize(FSubsystemCollectionBase& Collection) override;
+	// USubsystem END
 
 	FNavigationObstacleHashGrid2D AvoidanceObstacleGrid;
 };
diff --git a/Engine/Plugins/AI/MassCrowd/Source/MassCrowd/Public/MassCrowdSubsystem.h b/Engine/Plugins/AI/MassCrowd/Source/MassCrowd/Public/MassCrowdSubsystem.h
--- a/Engine/Plugins/AI/MassCrowd/Source/MassCrowd/Public/MassCrowdSubsystem.h
+++ b/Engine/Plugins/AI/MassCrowd/Source/MassCrowd/Public/MassCrowdSubsystem.h
@@ -4,7 +4,7 @@
 
 #include "MassCrowdTypes.h"
 #include "ZoneGraphTypes.h"
-#include "Subsystems/WorldSubsystem.h"
+#include "MassSubsystemBase.h"
 #include "MassCrowdSubsystem.generated.h"
 
 class UZoneGraphAnnotationSubsystem;
@@ -59,7 +59,7 @@ struct FRegisteredCrowdLaneData
  * It will create custom runtime lane data to allow branching decisions.
  */
 UCLASS()
-class MASSCROWD_API UMassCrowdSubsystem : public UWorldSubsystem
+class MASSCROWD_API UMassCrowdSubsystem : public UMassSubsystemBase
 {
 	GENERATED_BODY()
 public:
@@ -166,9 +166,12 @@ protected:
 	friend class UZoneGraphCrowdLaneAnnotations;
 
 	void UpdateDensityMask();
+
+	// USubsystem BEGIN
 	virtual void Initialize(FSubsystemCollectionBase& Collection) override;
 	virtual void PostInitialize() override;
 	virtual void Deinitialize() override;
+	// USubsystem END
 
 	void PostZoneGraphDataAdded(const AZoneGraphData* ZoneGraphData);
 	void PreZoneGraphDataRemoved(const AZoneGraphData* ZoneGraphData);
diff --git a/Engine/Plugins/Runtime/MassEntity/Source/MassEntity/Private/MassSubsystemBase.cpp b/Engine/Plugins/Runtime/MassEntity/Source/MassEntity/Private/MassSubsystemBase.cpp
new file mode 100644
--- /dev/null
+++ b/Engine/Plugins/Runtime/MassEntity/Source/MassEntity/Private/MassSubsystemBase.cpp
@@ -0,0 +1,43 @@
+// Copyright Epic Games, Inc. All Rights Reserved.
+
+#include "MassSubsystemBase.h"
+#include "HAL/IConsoleManager.h"
+
+
+namespace UE::Mass::Private
+{
+bool bRuntimeSubsystemsEnabled = true;
+
+namespace
+{
+	FAutoConsoleVariableRef AnonymousCVars[] =
+	{
+		{ TEXT("mass.RuntimeSubsystemsEnabled")
+		, bRuntimeSubsystemsEnabled
+		, TEXT("true by default, setting to false will prevent auto-creation of game-time Mass-related subsystems. Needs to be set before world loading.")
+		, ECVF_Default }
+	};
+}
+} // UE::Mass::Private
+
+
+//-----------------------------------------------------------------------------
+// UMassSubsystemBase
+//-----------------------------------------------------------------------------
+bool UMassSubsystemBase::AreRuntimeMassSubsystemsAllowed(UObject* Outer)
+{
+	return UE::Mass::Private::bRuntimeSubsystemsEnabled;
+}
+
+bool UMassSubsystemBase::ShouldCreateSubsystem(UObject* Outer) const 
+{
+	return UMassSubsystemBase::AreRuntimeMassSubsystemsAllowed(Outer);
+}
+
+//-----------------------------------------------------------------------------
+// UMassTickableSubsystemBase
+//-----------------------------------------------------------------------------
+bool UMassTickableSubsystemBase::ShouldCreateSubsystem(UObject* Outer) const
+{
+	return UMassSubsystemBase::AreRuntimeMassSubsystemsAllowed(Outer);
+}
diff --git a/Engine/Plugins/Runtime/MassEntity/Source/MassEntity/Public/MassEntitySubsystem.h b/Engine/Plugins/Runtime/MassEntity/Source/MassEntity/Public/MassEntitySubsystem.h
--- a/Engine/Plugins/Runtime/MassEntity/Source/MassEntity/Public/MassEntitySubsystem.h
+++ b/Engine/Plugins/Runtime/MassEntity/Source/MassEntity/Public/MassEntitySubsystem.h
@@ -2,7 +2,7 @@
 
 #pragma once
 
-#include "Subsystems/WorldSubsystem.h"
+#include "MassSubsystemBase.h"
 #include "MassEntityManager.h"
 #include "MassEntitySubsystem.generated.h"
 
@@ -13,7 +13,7 @@
  * use this by default. 
  */
 UCLASS()
-class MASSENTITY_API UMassEntitySubsystem : public UWorldSubsystem
+class MASSENTITY_API UMassEntitySubsystem : public UMassSubsystemBase
 {
 	GENERATED_BODY()
 
diff --git a/Engine/Plugins/Runtime/MassEntity/Source/MassEntity/Public/MassSubsystemBase.h b/Engine/Plugins/Runtime/MassEntity/Source/MassEntity/Public/MassSubsystemBase.h
new file mode 100644
--- /dev/null
+++ b/Engine/Plugins/Runtime/MassEntity/Source/MassEntity/Public/MassSubsystemBase.h
@@ -0,0 +1,40 @@
+// Copyright Epic Games, Inc. All Rights Reserved.
+
+#pragma once
+
+#include "Subsystems/WorldSubsystem.h"
+#include "MassSubsystemBase.generated.h"
+
+
+/** 
+ * The sole responsibility of this world subsystem class is to serve functionality common to all 
+ * Mass-related UWorldSubsystem-based subsystems, like whether the subsystems should get created at all. 
+ */
+UCLASS(Abstract)
+class MASSENTITY_API UMassSubsystemBase : public UWorldSubsystem
+{
+	GENERATED_BODY()
+
+public:
+	static bool AreRuntimeMassSubsystemsAllowed(UObject* Outer);
+
+protected:
+	//~USubsystem interface
+	virtual bool ShouldCreateSubsystem(UObject* Outer) const override;
+	//~End of USubsystem interface
+};
+
+/**
+ * The sole responsibility of this tickable world subsystem class is to serve functionality common to all
+ * Mass-related UTickableWorldSubsystem-based subsystems, like whether the subsystems should get created at all.
+ */
+UCLASS(Abstract)
+class MASSENTITY_API UMassTickableSubsystemBase : public UTickableWorldSubsystem
+{
+	GENERATED_BODY()
+
+protected:
+	//~USubsystem interface
+	virtual bool ShouldCreateSubsystem(UObject* Outer) const override;
+	//~End of USubsystem interface
+};
diff --git a/Engine/Plugins/Runtime/MassGameplay/Source/MassActors/Public/MassActorSpawnerSubsystem.h b/Engine/Plugins/Runtime/MassGameplay/Source/MassActors/Public/MassActorSpawnerSubsystem.h
--- a/Engine/Plugins/Runtime/MassGameplay/Source/MassActors/Public/MassActorSpawnerSubsystem.h
+++ b/Engine/Plugins/Runtime/MassGameplay/Source/MassActors/Public/MassActorSpawnerSubsystem.h
@@ -8,7 +8,7 @@
 #include "MassCommonTypes.h"
 #include "Templates/SubclassOf.h"
 #include "GameFramework/Actor.h"
-#include "Subsystems/WorldSubsystem.h"
+#include "MassSubsystemBase.h"
 
 #include "MassActorSpawnerSubsystem.generated.h"
 
@@ -116,17 +116,19 @@ public:
  * A subsystem managing spawning of actors for all mass subsystems
  */
 UCLASS(transient)
-class MASSACTORS_API UMassActorSpawnerSubsystem : public UWorldSubsystem
+class MASSACTORS_API UMassActorSpawnerSubsystem : public UMassSubsystemBase
 {
 	GENERATED_BODY()
-public:
 
+protected:
 	// USubsystem BEGIN
 	virtual void Initialize(FSubsystemCollectionBase& Collection) override;
 	virtual void Deinitialize() override;
-	static void AddReferencedObjects(UObject* InThis, FReferenceCollector& Collector);
 	// USubsystem END
 
+public:
+	static void AddReferencedObjects(UObject* InThis, FReferenceCollector& Collector);
+
 	/** Request an actor to spawn
 	 * Note: If you do not provide a spawn delegate, the requester is responsible to remove the request by hand.
 	 * It will be auto removed after the execution of the spawn delegate.
diff --git a/Engine/Plugins/Runtime/MassGameplay/Source/MassActors/Public/MassActorSubsystem.h b/Engine/Plugins/Runtime/MassGameplay/Source/MassActors/Public/MassActorSubsystem.h
--- a/Engine/Plugins/Runtime/MassGameplay/Source/MassActors/Public/MassActorSubsystem.h
+++ b/Engine/Plugins/Runtime/MassGameplay/Source/MassActors/Public/MassActorSubsystem.h
@@ -7,7 +7,7 @@
 #include "MassEntityTemplate.h"
 #include "Misc/MTAccessDetector.h"
 #include "UObject/ObjectKey.h"
-#include "Subsystems/WorldSubsystem.h"
+#include "MassSubsystemBase.h"
 #include "MassActorSubsystem.generated.h"
 
 struct FMassEntityHandle;
@@ -130,7 +130,7 @@ protected:
  * A subsystem managing communication between Actors and Mass
  */
 UCLASS()
-class MASSACTORS_API UMassActorSubsystem : public UWorldSubsystem
+class MASSACTORS_API UMassActorSubsystem : public UMassSubsystemBase
 {
 	GENERATED_BODY()
 
diff --git a/Engine/Plugins/Runtime/MassGameplay/Source/MassActors/Public/MassAgentSubsystem.h b/Engine/Plugins/Runtime/MassGameplay/Source/MassActors/Public/MassAgentSubsystem.h
--- a/Engine/Plugins/Runtime/MassGameplay/Source/MassActors/Public/MassAgentSubsystem.h
+++ b/Engine/Plugins/Runtime/MassGameplay/Source/MassActors/Public/MassAgentSubsystem.h
@@ -5,7 +5,7 @@
 #include "MassEntityTypes.h"
 #include "MassProcessingTypes.h"
 #include "MassEntityTemplate.h"
-#include "Subsystems/WorldSubsystem.h"
+#include "MassSubsystemBase.h"
 #include "MassAgentSubsystem.generated.h"
 
 class AActor;
@@ -36,7 +36,7 @@ struct MASSACTORS_API FMassAgentInitializationQueue
  * A subsystem managing communication between Actors and Mass
  */
 UCLASS()
-class MASSACTORS_API UMassAgentSubsystem : public UWorldSubsystem
+class MASSACTORS_API UMassAgentSubsystem : public UMassSubsystemBase
 {
 	GENERATED_BODY()
 
diff --git a/Engine/Plugins/Runtime/MassGameplay/Source/MassGameplayDebug/Public/MassDebuggerSubsystem.h b/Engine/Plugins/Runtime/MassGameplay/Source/MassGameplayDebug/Public/MassDebuggerSubsystem.h
--- a/Engine/Plugins/Runtime/MassGameplay/Source/MassGameplayDebug/Public/MassDebuggerSubsystem.h
+++ b/Engine/Plugins/Runtime/MassGameplay/Source/MassGameplayDebug/Public/MassDebuggerSubsystem.h
@@ -3,7 +3,7 @@
 #pragma once
 
 #include "MassGameplayDebugTypes.h"
-#include "Subsystems/WorldSubsystem.h"
+#include "MassSubsystemBase.h"
 #include "MassDebuggerSubsystem.generated.h"
 
 
@@ -11,7 +11,7 @@ class UMassDebugVisualizationComponent;
 class AMassDebugVisualizer;
 
 UCLASS()
-class MASSGAMEPLAYDEBUG_API UMassDebuggerSubsystem : public UWorldSubsystem
+class MASSGAMEPLAYDEBUG_API UMassDebuggerSubsystem : public UMassSubsystemBase
 {
 	GENERATED_BODY()
 public:
diff --git a/Engine/Plugins/Runtime/MassGameplay/Source/MassLOD/Public/MassLODSubsystem.h b/Engine/Plugins/Runtime/MassGameplay/Source/MassLOD/Public/MassLODSubsystem.h
--- a/Engine/Plugins/Runtime/MassGameplay/Source/MassLOD/Public/MassLODSubsystem.h
+++ b/Engine/Plugins/Runtime/MassGameplay/Source/MassLOD/Public/MassLODSubsystem.h
@@ -6,7 +6,7 @@
 #include "MassProcessor.h"
 #include "IndexedHandle.h"
 #include "MassLODTypes.h"
-#include "Subsystems/WorldSubsystem.h"
+#include "MassSubsystemBase.h"
 #include "MassExternalSubsystemTraits.h"
 #include "MassLODSubsystem.generated.h"
 
@@ -61,7 +61,7 @@ DECLARE_MULTICAST_DELEGATE_ThreeParams(FOnViewerRemoved, FMassViewerHandle Viewe
  * Manager responsible to manage and synchronized available viewers
  */
 UCLASS(config = Mass, defaultconfig)
-class MASSLOD_API UMassLODSubsystem : public UTickableWorldSubsystem
+class MASSLOD_API UMassLODSubsystem : public UMassTickableSubsystemBase
 {
 	GENERATED_BODY()
 
@@ -97,9 +97,12 @@ public:
 	FOnViewerRemoved& GetOnViewerRemovedDelegate() { return OnViewerRemovedDelegate; }
 
 protected:
+	// USubsystem BEGIN
 	virtual void Initialize(FSubsystemCollectionBase& Collection) override;
-	virtual TStatId GetStatId() const override;
 	virtual void Deinitialize() override;
+	// USubsystem END
+
+	virtual TStatId GetStatId() const override;
 
 	/** Called at the start of the PrePhysics mass processing phase and calls SynchronizeViewers */ 
 	void OnPrePhysicsPhaseStarted(float DeltaTime);
diff --git a/Engine/Plugins/Runtime/MassGameplay/Source/MassReplication/Public/MassReplicationSubsystem.h b/Engine/Plugins/Runtime/MassGameplay/Source/MassReplication/Public/MassReplicationSubsystem.h
--- a/Engine/Plugins/Runtime/MassGameplay/Source/MassReplication/Public/MassReplicationSubsystem.h
+++ b/Engine/Plugins/Runtime/MassGameplay/Source/MassReplication/Public/MassReplicationSubsystem.h
@@ -5,7 +5,7 @@
 #include "MassCommonTypes.h"
 #include "MassEntityTypes.h"
 #include "MassReplicationTypes.h"
-#include "Subsystems/WorldSubsystem.h"
+#include "MassSubsystemBase.h"
 
 #include "MassReplicationSubsystem.generated.h"
 
@@ -91,7 +91,7 @@ namespace UE::MassReplication
  *  ClientReplicationInfo relate to all the player controllers that have a parent or child UNetConnection to a single Client machine (split screen etc).
  */
 UCLASS()
-class MASSREPLICATION_API UMassReplicationSubsystem : public UWorldSubsystem
+class MASSREPLICATION_API UMassReplicationSubsystem : public UMassSubsystemBase
 {
 	GENERATED_BODY()
 
diff --git a/Engine/Plugins/Runtime/MassGameplay/Source/MassRepresentation/Public/MassRepresentationSubsystem.h b/Engine/Plugins/Runtime/MassGameplay/Source/MassRepresentation/Public/MassRepresentationSubsystem.h
--- a/Engine/Plugins/Runtime/MassGameplay/Source/MassRepresentation/Public/MassRepresentationSubsystem.h
+++ b/Engine/Plugins/Runtime/MassGameplay/Source/MassRepresentation/Public/MassRepresentationSubsystem.h
@@ -7,7 +7,7 @@
 #include "Misc/MTAccessDetector.h"
 #include "MassRepresentationTypes.h"
 #include "MassActorSpawnerSubsystem.h"
-#include "Subsystems/WorldSubsystem.h"
+#include "MassSubsystemBase.h"
 #include "MassRepresentationSubsystem.generated.h"
 
 class UMassVisualizationComponent;
@@ -25,7 +25,7 @@ class UWorldPartitionSubsystem;
  * Subsystem responsible for all visual of mass agents, will handle actors spawning and static mesh instances
  */
 UCLASS()
-class MASSREPRESENTATION_API UMassRepresentationSubsystem : public UWorldSubsystem
+class MASSREPRESENTATION_API UMassRepresentationSubsystem : public UMassSubsystemBase
 {
 	GENERATED_BODY()
 
diff --git a/Engine/Plugins/Runtime/MassGameplay/Source/MassSignals/Public/MassSignalSubsystem.h b/Engine/Plugins/Runtime/MassGameplay/Source/MassSignals/Public/MassSignalSubsystem.h
--- a/Engine/Plugins/Runtime/MassGameplay/Source/MassSignals/Public/MassSignalSubsystem.h
+++ b/Engine/Plugins/Runtime/MassGameplay/Source/MassSignals/Public/MassSignalSubsystem.h
@@ -3,7 +3,7 @@
 #pragma once
 
 #include "MassEntityManager.h"
-#include "Subsystems/WorldSubsystem.h"
+#include "MassSubsystemBase.h"
 #include "MassSignalSubsystem.generated.h"
 
 namespace UE::MassSignal 
@@ -15,7 +15,7 @@ namespace UE::MassSignal
 * A subsystem for handling Signals in Mass
 */
 UCLASS()
-class MASSSIGNALS_API UMassSignalSubsystem : public UTickableWorldSubsystem
+class MASSSIGNALS_API UMassSignalSubsystem : public UMassTickableSubsystemBase
 {
 	GENERATED_BODY()
 	
diff --git a/Engine/Plugins/Runtime/MassGameplay/Source/MassSimulation/Public/MassSimulationSubsystem.h b/Engine/Plugins/Runtime/MassGameplay/Source/MassSimulation/Public/MassSimulationSubsystem.h
--- a/Engine/Plugins/Runtime/MassGameplay/Source/MassSimulation/Public/MassSimulationSubsystem.h
+++ b/Engine/Plugins/Runtime/MassGameplay/Source/MassSimulation/Public/MassSimulationSubsystem.h
@@ -5,7 +5,7 @@
 #include "Delegates/IDelegateInstance.h"
 #include "MassProcessingTypes.h"
 #include "MassProcessingPhaseManager.h"
-#include "Subsystems/WorldSubsystem.h"
+#include "MassSubsystemBase.h"
 #include "MassSimulationSubsystem.generated.h"
 
 
@@ -14,7 +14,7 @@ struct FMassEntityManager;
 DECLARE_LOG_CATEGORY_EXTERN(LogMassSim, Log, All);
 
 UCLASS(config = Game, defaultconfig)
-class MASSSIMULATION_API UMassSimulationSubsystem : public UWorldSubsystem
+class MASSSIMULATION_API UMassSimulationSubsystem : public UMassSubsystemBase
 {
 	GENERATED_BODY()
 public:
@@ -37,10 +37,12 @@ public:
 	bool IsDuringMassProcessing() const;
 
 protected:
+	// UWorldSubsystem BEGIN
 	virtual void Initialize(FSubsystemCollectionBase& Collection) override;
 	virtual void PostInitialize() override;
 	virtual void OnWorldBeginPlay(UWorld& InWorld) override;
 	virtual void Deinitialize() override;
+	// UWorldSubsystem END
 	virtual void BeginDestroy() override;
 	
 	void RebuildTickPipeline();
diff --git a/Engine/Plugins/Runtime/MassGameplay/Source/MassSpawner/Public/MassSpawnerSubsystem.h b/Engine/Plugins/Runtime/MassGameplay/Source/MassSpawner/Public/MassSpawnerSubsystem.h
--- a/Engine/Plugins/Runtime/MassGameplay/Source/MassSpawner/Public/MassSpawnerSubsystem.h
+++ b/Engine/Plugins/Runtime/MassGameplay/Source/MassSpawner/Public/MassSpawnerSubsystem.h
@@ -4,7 +4,7 @@
 
 #include "MassEntityManager.h"
 #include "InstancedStruct.h"
-#include "Subsystems/WorldSubsystem.h"
+#include "MassSubsystemBase.h"
 #include "MassEntityTemplateRegistry.h"
 #include "MassSpawnerSubsystem.generated.h"
 
@@ -16,7 +16,7 @@ struct FMassEntityTemplateID;
 class UMassSimulationSubsystem;
 
 UCLASS()
-class MASSSPAWNER_API UMassSpawnerSubsystem : public UWorldSubsystem
+class MASSSPAWNER_API UMassSpawnerSubsystem : public UMassSubsystemBase
 {
 	GENERATED_BODY()
 
@@ -47,8 +47,10 @@ public:
 	const FMassEntityTemplate* GetMassEntityTemplate(FMassEntityTemplateID TemplateID) const;
 
 protected:
+	// UWorldSubsystem BEGIN
 	virtual void Initialize(FSubsystemCollectionBase& Collection) override;
 	virtual void Deinitialize() override;
+	// UWorldSubsystem END
 
 	void DoSpawning(const FMassEntityTemplate& EntityTemplate, const int32 NumToSpawn, FConstStructView SpawnData, TSubclassOf<UMassProcessor> InitializerClass, TArray<FMassEntityHandle>& OutEntities);
 
