Fixed TArrays of C arrays.

diff --git a/Engine/Plugins/Online/OnlineServicesEOS/Source/Private/Online/AuthEOS.cpp b/Engine/Plugins/Online/OnlineServicesEOS/Source/Private/Online/AuthEOS.cpp
--- a/Engine/Plugins/Online/OnlineServicesEOS/Source/Private/Online/AuthEOS.cpp
+++ b/Engine/Plugins/Online/OnlineServicesEOS/Source/Private/Online/AuthEOS.cpp
@@ -3,6 +3,7 @@
 #include "Online/AuthEOS.h"
 
 #include "Algo/Transform.h"
+#include "Containers/StaticArray.h"
 #include "EOSShared.h"
 #include "IEOSSDKManager.h"
 #include "Online/AuthErrors.h"
@@ -367,7 +368,7 @@ TFuture<FAccountId> FAuthEOS::ResolveAccountId(const FAccountId& LocalAccountId,
 	return ResolveAccountIdImpl(*this, LocalAccountId, ProductUserId);
 }
 
-using FEpicAccountIdStrBuffer = char[EOS_EPICACCOUNTID_MAX_LENGTH + 1];
+using FEpicAccountIdStrBuffer = TStaticArray<char, EOS_EPICACCOUNTID_MAX_LENGTH + 1>;
 
 TFuture<TArray<FAccountId>> FAuthEOS::ResolveAccountIds(const FAccountId& LocalAccountId, const TArray<EOS_EpicAccountId>& InEpicAccountIds)
 {
@@ -414,7 +415,7 @@ TFuture<TArray<FAccountId>> FAuthEOS::ResolveAccountIds(const FAccountId& LocalA
 		FEpicAccountIdStrBuffer& EpicAccountIdStr = EpicAccountIdStrsToQuery.Emplace_GetRef();
 		int32_t BufferSize = sizeof(EpicAccountIdStr);
 		if (!EOS_EpicAccountId_IsValid(EpicAccountId) ||
-			EOS_EpicAccountId_ToString(EpicAccountId, EpicAccountIdStr, &BufferSize) != EOS_EResult::EOS_Success)
+			EOS_EpicAccountId_ToString(EpicAccountId, EpicAccountIdStr.GetData(), &BufferSize) != EOS_EResult::EOS_Success)
 		{
 			checkNoEntry();
 			return MakeFulfilledPromise<TArray<FAccountId>>().GetFuture();
@@ -454,7 +455,7 @@ TFuture<TArray<FAccountId>> FAuthEOS::ResolveAccountIds(const FAccountId& LocalA
 					{
 						FEpicAccountIdStrBuffer EpicAccountIdStr;
 						int32_t BufferSize = sizeof(EpicAccountIdStr);
-						verify(EOS_EpicAccountId_ToString(EpicAccountId, EpicAccountIdStr, &BufferSize) == EOS_EResult::EOS_Success);
+						verify(EOS_EpicAccountId_ToString(EpicAccountId, EpicAccountIdStr.GetData(), &BufferSize) == EOS_EResult::EOS_Success);
 						Options.TargetExternalUserId = &EpicAccountIdStr[0];
 						const EOS_ProductUserId ProductUserId = EOS_Connect_GetExternalAccountMapping(ConnectHandle, &Options);
 						AccountId = CreateAccountId(EpicAccountId, ProductUserId);
@@ -540,10 +541,10 @@ TFuture<TArray<FAccountId>> FAuthEOS::ResolveAccountIds(const FAccountId& LocalA
 						FEpicAccountIdStrBuffer EpicAccountIdStr;
 						int32_t BufferLength = sizeof(EpicAccountIdStr);
 						EOS_EpicAccountId EpicAccountId = nullptr;
-						const EOS_EResult Result = EOS_Connect_GetProductUserIdMapping(ConnectHandle, &Options, EpicAccountIdStr, &BufferLength);
+						const EOS_EResult Result = EOS_Connect_GetProductUserIdMapping(ConnectHandle, &Options, EpicAccountIdStr.GetData(), &BufferLength);
 						if (Result == EOS_EResult::EOS_Success)
 						{
-							EpicAccountId = EOS_EpicAccountId_FromString(EpicAccountIdStr);
+							EpicAccountId = EOS_EpicAccountId_FromString(EpicAccountIdStr.GetData());
 							check(EOS_EpicAccountId_IsValid(EpicAccountId));
 						}
 						AccountId = CreateAccountId(EpicAccountId, ProductUserId);
diff --git a/Engine/Plugins/Runtime/GeometryProcessing/Source/DynamicMesh/Private/Polygroups/PolygroupsGenerator.cpp b/Engine/Plugins/Runtime/GeometryProcessing/Source/DynamicMesh/Private/Polygroups/PolygroupsGenerator.cpp
--- a/Engine/Plugins/Runtime/GeometryProcessing/Source/DynamicMesh/Private/Polygroups/PolygroupsGenerator.cpp
+++ b/Engine/Plugins/Runtime/GeometryProcessing/Source/DynamicMesh/Private/Polygroups/PolygroupsGenerator.cpp
@@ -3,6 +3,7 @@
 #include "Polygroups/PolygroupsGenerator.h"
 
 #include "Clustering/FaceNormalClustering.h"
+#include "Containers/StaticArray.h"
 #include "DynamicMesh/MeshNormals.h"
 #include "Selections/MeshConnectedComponents.h"
 #include "Util/IndexUtil.h"
@@ -203,7 +204,7 @@ bool FPolygroupsGenerator::FindSourceMeshPolygonPolygroups(
 
 	// For each triangle, there are 3 possible adjacent triangles that form valid quads. 
 	// TriPotentialQuads is an [N][3] array of those potential quads, stored as FTriPairQuad
-	TArray<FTriPairQuad[3]> TriPotentialQuads;
+	TArray<TStaticArray<FTriPairQuad, 3>> TriPotentialQuads;
 	TriPotentialQuads.SetNum(MaxTriangleID);
 	for (int32 tid : Mesh->TriangleIndicesItr())
 	{
diff --git a/Engine/Plugins/Runtime/MassEntity/Source/MassEntity/Private/MassCommandBuffer.cpp b/Engine/Plugins/Runtime/MassEntity/Source/MassEntity/Private/MassCommandBuffer.cpp
--- a/Engine/Plugins/Runtime/MassEntity/Source/MassEntity/Private/MassCommandBuffer.cpp
+++ b/Engine/Plugins/Runtime/MassEntity/Source/MassEntity/Private/MassCommandBuffer.cpp
@@ -1,6 +1,7 @@
 // Copyright Epic Games, Inc. All Rights Reserved.
 
 #include "MassCommandBuffer.h"
+#include "Containers/AnsiString.h"
 #include "MassObserverManager.h"
 #include "MassEntityUtils.h"
 #include "HAL/IConsoleManager.h"
@@ -22,13 +23,10 @@ FAutoConsoleVariableRef CVarEnableDetailedCommandStats(TEXT("massentities.Enable
 
 /** CSV stat names */
 static FString DefaultBatchedName = TEXT("BatchedCommand");
-static TArray<FName> CommandBatchedFNames;
+static TMap<FName, TPair<FString, FAnsiString>> CommandBatchedFNames;
 
 /** CSV custom stat names (ANSI) */
-static const int32 MaxNameLength = 64;
-typedef ANSICHAR ANSIName[MaxNameLength];
-static ANSIName DefaultANSIBatchedName = "BatchedCommand";
-static TArray<ANSIName> CommandANSIBatchedNames;
+static FAnsiString DefaultANSIBatchedName = "BatchedCommand";
 
 /**
  * Provides valid names for CSV profiling.
@@ -36,31 +34,24 @@ static TArray<ANSIName> CommandANSIBatchedNames;
  * @param OutName is the name to use for csv custom stats
  * @param OutANSIName is the name to use for csv stats
  */
-void GetCommandStatNames(FMassBatchedCommand& Command, FString& OutName, ANSIName*& OutANSIName)
+void GetCommandStatNames(FMassBatchedCommand& Command, FString*& OutName, FAnsiString*& OutANSIName)
 {
 	OutANSIName = &DefaultANSIBatchedName;
-	OutName = DefaultANSIBatchedName;
+	OutName     = &DefaultBatchedName;
 	if (!bEnableDetailedStats)
 	{
 		return;
 	}
 
 	const FName CommandFName = Command.GetFName();
-	OutName = CommandFName.ToString();
 
-	const int32 Index = CommandBatchedFNames.Find(CommandFName);
-	if (Index == INDEX_NONE)
+	TPair<FString, FAnsiString>& Names = CommandBatchedFNames.FindOrAdd(CommandFName);
+	OutName     = &Names.Get<FString>();
+	OutANSIName = &Names.Get<FAnsiString>();
+	if (OutName->IsEmpty())
 	{
-		CommandBatchedFNames.Emplace(CommandFName);
-		OutANSIName = &CommandANSIBatchedNames.AddZeroed_GetRef();
-		// Use prefix for easier parsing in reports
-		//const FString CounterName = FString::Printf(TEXT("Num%s"), *OutName);
-		//FMemory::Memcpy(OutANSIName, StringCast<ANSICHAR>(*CounterName).Get(), FMath::Min(CounterName.Len(), MaxNameLength - 1) * sizeof(ANSICHAR));
-		FMemory::Memcpy(OutANSIName, StringCast<ANSICHAR>(*OutName).Get(), FMath::Min(OutName.Len(), MaxNameLength - 1) * sizeof(ANSICHAR));
-	}
-	else
-	{
-		OutANSIName = &CommandANSIBatchedNames[Index];
+		*OutName     = CommandFName.ToString();
+		*OutANSIName = **OutName;
 	}
 }
 
@@ -142,13 +133,13 @@ void FMassCommandBuffer::Flush(FMassEntityManager& EntityManager)
 			using namespace UE::Mass::Command;
 
 			// Extract name (default or detailed)
-			ANSIName* ANSIName = &DefaultANSIBatchedName;
-			FString Name = DefaultBatchedName;
+			FAnsiString* ANSIName = nullptr;
+			FString*     Name     = nullptr;
 			GetCommandStatNames(*Command, Name, ANSIName);
 
 			// Push stats
-			FScopedCsvStat ScopedCsvStat(*ANSIName, CSV_CATEGORY_INDEX(MassEntities));
-			FCsvProfiler::RecordCustomStat(*Name, CSV_CATEGORY_INDEX(MassEntitiesCounters), Command->GetNumOperationsStat(), ECsvCustomStatOp::Accumulate);
+			FScopedCsvStat ScopedCsvStat(**ANSIName, CSV_CATEGORY_INDEX(MassEntities));
+			FCsvProfiler::RecordCustomStat(**Name, CSV_CATEGORY_INDEX(MassEntitiesCounters), Command->GetNumOperationsStat(), ECsvCustomStatOp::Accumulate);
 #endif // CSV_PROFILER
 
 			Command->Execute(EntityManager);
