Fixed FMassArchetypeEntityCollectionWithPayload::CreateEntityRangesWithPayload not handling duplicates. The change required adding another mutating function to FMassGenericPayloadView.
Also added unit tests for the found issue.

diff --git a/Engine/Plugins/Runtime/MassEntity/Source/MassEntity/Private/MassArchetypeTypes.cpp b/Engine/Plugins/Runtime/MassEntity/Source/MassEntity/Private/MassArchetypeTypes.cpp
--- a/Engine/Plugins/Runtime/MassEntity/Source/MassEntity/Private/MassArchetypeTypes.cpp
+++ b/Engine/Plugins/Runtime/MassEntity/Source/MassEntity/Private/MassArchetypeTypes.cpp
@@ -198,6 +198,11 @@ void FMassArchetypeEntityCollectionWithPayload::CreateEntityRangesWithPayload(co
 		int32 ArchetypeIndex = INDEX_NONE;
 		int32 TrueIndex = INDEX_NONE;
 
+		bool operator==(const FEntityInArchetype& Other) const
+		{
+			return ArchetypeIndex == Other.ArchetypeIndex && TrueIndex == Other.TrueIndex;
+		}
+
 		bool operator<(const FEntityInArchetype& Other) const
 		{
 			return ArchetypeIndex < Other.ArchetypeIndex || (ArchetypeIndex == Other.ArchetypeIndex && TrueIndex < Other.TrueIndex);
@@ -255,8 +260,14 @@ void FMassArchetypeEntityCollectionWithPayload::CreateEntityRangesWithPayload(co
 		}
 	}
 
-	UE::Mass::Utils::AbstractSort(Entities.Num(), [&EntityData](const int32 LHS, const int32 RHS)
+	// A paranoid programmer might point out that there are no guarantees that a sorting algorithm will compare all elements.
+	// While that's true we make an assumption here, that the elements next to each other will in fact all get compared
+	// and since all we care about with `bDuplicatesFound` is whether same elements exist (that will be right next to each other
+	// in the final lineup) we feel safe in the assumption.
+	bool bDuplicatesFound = false;
+	UE::Mass::Utils::AbstractSort(Entities.Num(), [&EntityData, &bDuplicatesFound](const int32 LHS, const int32 RHS)
 		{
+			bDuplicatesFound = bDuplicatesFound || (EntityData[LHS] == EntityData[RHS]);
 			return EntityData[LHS] < EntityData[RHS];
 		}
 		, [&EntityData, &Payload](const int32 A, const int32 B)
@@ -264,6 +275,51 @@ void FMassArchetypeEntityCollectionWithPayload::CreateEntityRangesWithPayload(co
 			::Swap(EntityData[A], EntityData[B]);
 			Payload.Swap(A, B);
 		});
+	ensureMsgf(bDuplicatesFound == false || (DuplicatesHandling != FMassArchetypeEntityCollection::NoDuplicates)
+		, TEXT("Caller declared lack of duplicates in the input data, but duplicates have been found"));
+
+#if !UE_BUILD_SHIPPING
+	// in non shipping builds we still want to verify that the assumption expressed in bDuplicatesFound comment above
+	// is correct
+	if (!bDuplicatesFound && (DuplicatesHandling == FMassArchetypeEntityCollection::FoldDuplicates))
+	{
+		for (int32 EntryIndex = 0; EntryIndex < EntityData.Num() - 1; ++EntryIndex)
+		{
+			checkf(EntityData[EntryIndex] != EntityData[EntryIndex + 1], TEXT("Assumption regarding comparison between identical elements while sorting is wrong!"));
+		}
+	}
+#endif // !UE_BUILD_SHIPPING
+
+	if (bDuplicatesFound && (DuplicatesHandling == FMassArchetypeEntityCollection::FoldDuplicates))
+	{
+		// we cannot remove elements from Payload, since it's a view to existing data, we need to sort the data in 
+		// such a way that all the duplicates end up at the end of the view. We can then ignore the appropriate
+		// number of elements.
+		
+		// processing Num - 1 elements since there's no point in checking the last one - there's nothing to compare it against
+		for (int32 EntryIndex = 0; EntryIndex < EntityData.Num() - 1; ++EntryIndex)
+		{	
+			FEntityInArchetype& Entry = EntityData[EntryIndex];
+			if (Entry != EntityData[EntryIndex + 1])
+			{
+				continue;
+			}
+
+			int32 DuplicateIndex = EntryIndex + 1;
+			while (DuplicateIndex + 1 < EntityData.Num() && Entry == EntityData[DuplicateIndex + 1])
+			{
+				++DuplicateIndex;
+			};
+
+			const int32 NumDuplicates = DuplicateIndex - EntryIndex;
+
+			EntityData.RemoveAt(EntryIndex + 1, NumDuplicates, /*bAllowShrinking=*/false);
+			Payload.SwapElementsToEnd(EntryIndex + 1, NumDuplicates);
+			// even though we don't remove the elements from payload we later limit the number of elements used with
+			// ArchetypeInfo.Count, so we need to update that
+			Archetypes[Entry.ArchetypeIndex].Count -= NumDuplicates;
+		}
+	}
 
 	int32 ProcessedEntitiesCount = 0;
 	int32 ArchetypeIndex = 0;
diff --git a/Engine/Plugins/Runtime/MassEntity/Source/MassEntity/Private/MassEntityTypes.cpp b/Engine/Plugins/Runtime/MassEntity/Source/MassEntity/Private/MassEntityTypes.cpp
--- a/Engine/Plugins/Runtime/MassEntity/Source/MassEntity/Private/MassEntityTypes.cpp
+++ b/Engine/Plugins/Runtime/MassEntity/Source/MassEntity/Private/MassEntityTypes.cpp
@@ -13,6 +13,9 @@ DEFINE_TYPEBITSET(FMassChunkFragmentBitSet);
 DEFINE_TYPEBITSET(FMassSharedFragmentBitSet);
 DEFINE_TYPEBITSET(FMassExternalSubsystemBitSet);
 
+//-----------------------------------------------------------------------------
+// FMassArchetypeSharedFragmentValues
+//-----------------------------------------------------------------------------
 uint32 FMassArchetypeSharedFragmentValues::CalculateHash() const
 {
 	checkf(bSorted, TEXT("Expecting the containers to be sorted for the hash caluclation to be correct"));
@@ -32,3 +35,38 @@ uint32 FMassArchetypeSharedFragmentValues::CalculateHash() const
 	return Hash;
 }
 
+//-----------------------------------------------------------------------------
+// FMassGenericPayloadView
+//-----------------------------------------------------------------------------
+void FMassGenericPayloadView::SwapElementsToEnd(const int32 StartIndex, int32 NumToMove)
+{
+	check(StartIndex >= 0 && NumToMove >= 0);
+
+	if (UNLIKELY(NumToMove <= 0 || StartIndex < 0))
+	{
+		return;
+	}
+
+	TArray<uint8, TInlineAllocator<16>> MovedElements;
+
+	for (FStructArrayView& StructArrayView : Content)
+	{
+		check((StartIndex + NumToMove) <= StructArrayView.Num());
+		if (StartIndex + NumToMove >= StructArrayView.Num() - 1)
+		{
+			// nothing to do here, the elements are already at the back
+			continue;
+		}
+
+		uint8* ViewData = static_cast<uint8*>(StructArrayView.GetData());
+		const uint32 ElementSize = StructArrayView.GetTypeSize();
+		const uint32 MovedStartOffset = StartIndex * ElementSize;
+		const uint32 MovedSize = NumToMove * ElementSize;
+		const uint32 MoveOffset = (StructArrayView.Num() - (StartIndex + NumToMove)) * ElementSize;
+
+		MovedElements.Reset();
+		MovedElements.Append(ViewData + MovedStartOffset, MovedSize);
+		FMemory::Memmove(ViewData + MovedStartOffset, ViewData + MovedStartOffset + MovedSize, MoveOffset);
+		FMemory::Memcpy(ViewData + MovedStartOffset + MoveOffset, MovedElements.GetData(), MovedSize);
+	}
+}
\ No newline at end of file
diff --git a/Engine/Plugins/Runtime/MassEntity/Source/MassEntity/Public/MassEntityTypes.h b/Engine/Plugins/Runtime/MassEntity/Source/MassEntity/Public/MassEntityTypes.h
--- a/Engine/Plugins/Runtime/MassEntity/Source/MassEntity/Public/MassEntityTypes.h
+++ b/Engine/Plugins/Runtime/MassEntity/Source/MassEntity/Public/MassEntityTypes.h
@@ -393,6 +393,9 @@ struct FMassGenericPayloadView
 		}
 	}
 
+	/** Moves NumToMove elements to the back of the viewed collection. */
+	void SwapElementsToEnd(int32 StartIndex, int32 NumToMove);
+
 	TArrayView<FStructArrayView> Content;
 };
 
diff --git a/Engine/Plugins/Runtime/MassEntity/Source/MassEntityTestSuite/Private/MassChunkCollectionTest.cpp b/Engine/Plugins/Runtime/MassEntity/Source/MassEntityTestSuite/Private/MassChunkCollectionTest.cpp
--- a/Engine/Plugins/Runtime/MassEntity/Source/MassEntityTestSuite/Private/MassChunkCollectionTest.cpp
+++ b/Engine/Plugins/Runtime/MassEntity/Source/MassEntityTestSuite/Private/MassChunkCollectionTest.cpp
@@ -132,6 +132,24 @@ struct FEntityCollection_CreateCrossChunk : FEntityTestBase
 };
 IMPLEMENT_AI_INSTANT_TEST(FEntityCollection_CreateCrossChunk, "System.Mass.EntityCollection.Create.CrossChunk");
 
+struct FEntityCollection_CreateWithDuplicatesTrivial : FEntityCollectionTestBase
+{
+	virtual bool InstantTest() override
+	{
+		TArray<FMassEntityHandle> EntitiesWithDuplicates;
+		EntitiesWithDuplicates.Add(Entities[2]);
+		EntitiesWithDuplicates.Add(Entities[2]);
+
+		FMassArchetypeEntityCollection EntityCollection(FloatsArchetype, EntitiesWithDuplicates, FMassArchetypeEntityCollection::FoldDuplicates);
+		FMassArchetypeEntityCollection::FConstEntityRangeArrayView Ranges = EntityCollection.GetRanges();
+		AITEST_EQUAL("The result should have a single subchunk", Ranges.Num(), 1);
+		AITEST_EQUAL("The resulting subchunk should be of length 1", Ranges[0].Length, 1);
+
+		return true;
+	}
+};
+IMPLEMENT_AI_INSTANT_TEST(FEntityCollection_CreateWithDuplicatesTrivial, "System.Mass.EntityCollection.Create.TrivialDuplicates");
+
 struct FEntityCollection_CreateWithDuplicates : FEntityCollectionTestBase
 {
 	virtual bool InstantTest() override
@@ -155,7 +173,7 @@ struct FEntityCollection_CreateWithDuplicates : FEntityCollectionTestBase
 IMPLEMENT_AI_INSTANT_TEST(FEntityCollection_CreateWithDuplicates, "System.Mass.EntityCollection.Create.Duplicates");
 
 #if WITH_MASSENTITY_DEBUG
-struct FEntityCollection_WithPayload : FEntityCollectionTestBase
+struct FEntityCollection_WithPayloadBase : FEntityCollectionTestBase
 {
 	virtual bool SetUp() override
 	{
@@ -166,30 +184,105 @@ struct FEntityCollection_WithPayload : FEntityCollectionTestBase
 		EntityManager->BatchCreateEntities(FloatsArchetype, EntitiesPerChunk * 2, Entities);
 		return true;
 	}
+};
+
+struct FEntityCollection_TrivialDuplicatesWithPayload : FEntityCollection_WithPayloadBase
+{
+	virtual bool InstantTest() override
+	{
+		const int32 EntitiesPerChunk = EntityManager->DebugGetArchetypeEntitiesCountPerChunk(FloatsArchetype);
+		const int32 TotalCount = Entities.Num();
+		const int32 SubSetCount = int(0.6 * Entities.Num()); // using >0.5 to ensure some entities picked being in sequence and/or in different chunks
+		TArray<FMassEntityHandle> EntitiesSubSet;
+		TArray<FTestFragment_Int> Payload;
+
+		EntitiesSubSet.Add(Entities[EntitiesPerChunk + 20]); 
+		Payload.Add(FTestFragment_Int(int32(0)));
+		EntitiesSubSet.Add(Entities[EntitiesPerChunk + 20]);
+		Payload.Add(FTestFragment_Int(int32(1)));
+
+		// transform typed payload array into generic one for sorting purposes
+		FStructArrayView PaloadView(Payload);
+		TArray<FMassArchetypeEntityCollectionWithPayload> Result;
+		FMassArchetypeEntityCollectionWithPayload::CreateEntityRangesWithPayload(*EntityManager, EntitiesSubSet, FMassArchetypeEntityCollection::FoldDuplicates
+			, FMassGenericPayloadView(MakeArrayView(&PaloadView, 1)), Result);
 
+		AITEST_EQUAL(TEXT("We expect only a single result"), Result.Num(), 1);
+		AITEST_EQUAL(TEXT("We expect only a single resulting range"), Result[0].GetEntityCollection().GetRanges().Num(), 1);
+		AITEST_EQUAL(TEXT("We expect only a single entity in the resulting range"), Result[0].GetEntityCollection().GetRanges()[0].Length, 1);
+
+		return true;
+	}
+};
+IMPLEMENT_AI_INSTANT_TEST(FEntityCollection_TrivialDuplicatesWithPayload, "System.Mass.EntityCollection.Create.TrivialDuplicatesWithPayload");
+
+// @todo also add another archetype
+struct FEntityCollection_MultiDuplicatesWithPayload : FEntityCollection_WithPayloadBase
+{
+	virtual bool InstantTest() override
+	{
+		const int32 EntitiesPerChunk = EntityManager->DebugGetArchetypeEntitiesCountPerChunk(FloatsArchetype);
+		const int32 TotalCount = Entities.Num();
+		const int32 SubSetCount = int(0.6 * Entities.Num()); // using >0.5 to ensure some entities picked being in sequence and/or in different chunks
+		TArray<FMassEntityHandle> EntitiesSubSet;
+		TArray<FTestFragment_Int> Payload;
+
+		constexpr int32 NumUniques = 3;
+		constexpr int32 NumDuplicatesEach = 4;
+		int32 FragmentValue = 0;
+		for (int32 Iteration = 0; Iteration < NumDuplicatesEach; ++Iteration, ++FragmentValue)
+		{
+			for (int32 Unique = 0; Unique < NumUniques; ++Unique)
+			{
+				EntitiesSubSet.Add(Entities[EntitiesPerChunk + 20 + Unique]);
+				Payload.Add(FTestFragment_Int(Unique));
+			}
+		}
+
+		// transform typed payload array into generic one for sorting purposes
+		FStructArrayView PaloadView(Payload);
+		TArray<FMassArchetypeEntityCollectionWithPayload> Result;
+		FMassArchetypeEntityCollectionWithPayload::CreateEntityRangesWithPayload(*EntityManager, EntitiesSubSet, FMassArchetypeEntityCollection::FoldDuplicates
+			, FMassGenericPayloadView(MakeArrayView(&PaloadView, 1)), Result);
+
+		AITEST_EQUAL(TEXT("We expect only a single result"), Result.Num(), 1);
+		AITEST_EQUAL(TEXT("We expect only a single resulting range"), Result[0].GetEntityCollection().GetRanges().Num(), 1);
+		AITEST_EQUAL(TEXT("We expect exatly NumUniques entities in the resulting range"), Result[0].GetEntityCollection().GetRanges()[0].Length, NumUniques);
+		const FMassGenericPayloadViewSlice& PayloadSlice = Result[0].GetPayload();
+		for (int32 Unique = 0; Unique < NumUniques; ++Unique)
+		{
+			AITEST_EQUAL("The surviving payload value should match the expected", PayloadSlice[0].GetAt<FTestFragment_Int>(Unique).Value, Unique);
+		}
+
+		return true;
+	}
+};
+IMPLEMENT_AI_INSTANT_TEST(FEntityCollection_MultiDuplicatesWithPayload, "System.Mass.EntityCollection.Create.MultiDuplicatesWithPayload");
+
+struct FEntityCollection_WithPayload : FEntityCollection_WithPayloadBase
+{
 	virtual bool InstantTest() override
 	{
-		// Entities contains 100 FloatArchetype entities by now, created in SetUp
 		const int32 TotalCount = Entities.Num();
 		const int32 SubSetCount = int(0.6 * Entities.Num()); // using >0.5 to ensure some entities picked being in sequence and/or in different chunks
 		TArray<FMassEntityHandle> EntitiesSubSet;
 		TArray<FTestFragment_Int> Payload;
 
-		TArray<uint8> Indices;
+		TArray<int32> Indices;
 		Indices.AddUninitialized(TotalCount);
-		for (int i = 0; i < Indices.Num(); ++i)
+		for (int32 i = 0; i < Indices.Num(); ++i)
 		{
-			Indices[i] = uint8(i);
+			Indices[i] = i;
 		}
 
 		FMath::SRandInit(TotalCount);
 		Algo::RandomShuffle(Indices);
 		Indices.SetNum(SubSetCount);
 
-		for (uint8 i : Indices)
+		for (int32 i : Indices)
 		{
 			EntitiesSubSet.Add(Entities[i]);
-			Payload.Add(FTestFragment_Int(int32(i)));
+			Payload.Add(FTestFragment_Int(i));
 		}
 
 		// transform typed payload array into generic one for sorting purposes
