Add option to make mass PrePhysics phase a high priority task, to minimise GameThread waits by kicking off parallel work as soon as possible.
Enabled by default, disable via "mass.MakePrePhysicsTickFunctionHighPriority".

diff --git a/Engine/Plugins/Runtime/MassEntity/Source/MassEntity/Private/MassProcessingPhaseManager.cpp b/Engine/Plugins/Runtime/MassEntity/Source/MassEntity/Private/MassProcessingPhaseManager.cpp
--- a/Engine/Plugins/Runtime/MassEntity/Source/MassEntity/Private/MassProcessingPhaseManager.cpp
+++ b/Engine/Plugins/Runtime/MassEntity/Source/MassEntity/Private/MassProcessingPhaseManager.cpp
@@ -20,9 +20,11 @@ DECLARE_CYCLE_STAT(TEXT("Mass Phase Tick"), STAT_Mass_PhaseTick, STATGROUP_Mass)
 namespace UE::Mass::Tweakables
 {
 	bool bFullyParallel = MASS_DO_PARALLEL;
+	bool bMakePrePhysicsTickFunctionHighPriority = true;
 
 	FAutoConsoleVariableRef CVars[] = {
 		{TEXT("mass.FullyParallel"), bFullyParallel, TEXT("Enables mass processing distribution to all available thread (via the task graph)")},
+		{TEXT("mass.MakePrePhysicsTickFunctionHighPriority"), bMakePrePhysicsTickFunctionHighPriority, TEXT("Whether to make the PrePhysics tick function high priority - can minimise GameThread waits by starting parallel work as soon as possible")},
 	};
 }
 
@@ -315,6 +317,12 @@ void FMassProcessingPhaseManager::EnableTickFunctions(const UWorld& World)
 
 	for (FMassProcessingPhase& Phase : ProcessingPhases)
 	{
+		if (UE::Mass::Tweakables::bMakePrePhysicsTickFunctionHighPriority && (Phase.Phase == EMassProcessingPhase::PrePhysics))
+		{
+			constexpr bool bHighPriority = true;
+			Phase.SetPriorityIncludingPrerequisites(bHighPriority);
+		}
+
 		Phase.RegisterTickFunction(World.PersistentLevel);
 		Phase.SetTickFunctionEnable(true);
 #if WITH_MASSENTITY_DEBUG
