Minor changes to Mass subsystem's init flow to support adding subsystems at runtime (i.e. after UWorld::BeginPlay).

diff --git a/Engine/Plugins/AI/MassCrowd/Source/MassCrowd/Private/MassCrowdSubsystem.cpp b/Engine/Plugins/AI/MassCrowd/Source/MassCrowd/Private/MassCrowdSubsystem.cpp
--- a/Engine/Plugins/AI/MassCrowd/Source/MassCrowd/Private/MassCrowdSubsystem.cpp
+++ b/Engine/Plugins/AI/MassCrowd/Source/MassCrowd/Private/MassCrowdSubsystem.cpp
@@ -82,6 +82,8 @@ void UMassCrowdSubsystem::Initialize(FSubsystemCollectionBase& Collection)
 		RebuildLaneData();
 	});
 #endif
+
+	HandleLateCreation();
 }
 
 void UMassCrowdSubsystem::Deinitialize()
diff --git a/Engine/Plugins/Runtime/MassEntity/Source/MassEntity/Private/MassEntitySubsystem.cpp b/Engine/Plugins/Runtime/MassEntity/Source/MassEntity/Private/MassEntitySubsystem.cpp
--- a/Engine/Plugins/Runtime/MassEntity/Source/MassEntity/Private/MassEntitySubsystem.cpp
+++ b/Engine/Plugins/Runtime/MassEntity/Source/MassEntity/Private/MassEntitySubsystem.cpp
@@ -24,11 +24,14 @@ void UMassEntitySubsystem::GetResourceSizeEx(FResourceSizeEx& CumulativeResource
 
 void UMassEntitySubsystem::Initialize(FSubsystemCollectionBase& Collection)
 {
+	Super::Initialize(Collection);
 	EntityManager->Initialize();
+	HandleLateCreation();
 }
 
 void UMassEntitySubsystem::PostInitialize()
 {
+	Super::PostInitialize();
 	// this needs to be done after all the subsystems have been initialized since some processors might want to access
 	// them during processors' initialization
 	EntityManager->PostInitialize();
@@ -38,6 +41,7 @@ void UMassEntitySubsystem::Deinitialize()
 {
 	EntityManager->Deinitialize();
 	EntityManager.Reset();
+	Super::Deinitialize();
 }
 
 #if WITH_MASSENTITY_DEBUG
diff --git a/Engine/Plugins/Runtime/MassEntity/Source/MassEntity/Private/MassSubsystemBase.cpp b/Engine/Plugins/Runtime/MassEntity/Source/MassEntity/Private/MassSubsystemBase.cpp
--- a/Engine/Plugins/Runtime/MassEntity/Source/MassEntity/Private/MassSubsystemBase.cpp
+++ b/Engine/Plugins/Runtime/MassEntity/Source/MassEntity/Private/MassSubsystemBase.cpp
@@ -2,10 +2,32 @@
 
 #include "MassSubsystemBase.h"
 #include "HAL/IConsoleManager.h"
+#include "Engine/World.h"
 
 
 namespace UE::Mass::Private
 {
+/** 
+ * A helper function calling PostInitialize and OnWorldBeginPlay for the given subsystem, provided the world has already begun play.
+ * @see UMassSubsystemBase::HandleLateCreation for more detail
+ */
+void HandleLateCreation(UWorldSubsystem& MassWorldSubsystem, const UE::Mass::Subsystems::FInitializationState InitializationState)
+{
+	// handle late creation
+	UWorld* World = MassWorldSubsystem.GetWorld();
+	if (World)
+	{
+		if (World->IsInitialized() == true && InitializationState.bPostInitializeCalled == false)
+		{
+			MassWorldSubsystem.PostInitialize();
+		}
+		if (World->HasBegunPlay() == true && InitializationState.bOnWorldBeginPlayCalled == false)
+		{
+			MassWorldSubsystem.OnWorldBeginPlay(*World);
+		}
+	}
+}
+
 bool bRuntimeSubsystemsEnabled = true;
 
 namespace
@@ -34,6 +56,37 @@ bool UMassSubsystemBase::ShouldCreateSubsystem(UObject* Outer) const
 	return UMassSubsystemBase::AreRuntimeMassSubsystemsAllowed(Outer) && Super::ShouldCreateSubsystem(Outer);
 }
 
+void UMassSubsystemBase::Initialize(FSubsystemCollectionBase& Collection)
+{
+	// This ensure is here to make sure we handle HandleLateCreation() gracefully, we dont expect it to ever trigger unless users start to manually call the functions
+	ensureMsgf(InitializationState.bInitializeCalled == false, TEXT("%hs called multiple times"), __FUNCTION__);
+	InitializationState.bInitializeCalled = true;
+}
+
+void UMassSubsystemBase::PostInitialize()
+{
+	// This ensure is here to make sure we handle HandleLateCreation() gracefully, we dont expect it to ever trigger unless users start to manually call the functions
+	ensureMsgf(InitializationState.bPostInitializeCalled == false, TEXT("%hs called multiple times"), __FUNCTION__);
+	InitializationState.bPostInitializeCalled = true;
+}
+
+void UMassSubsystemBase::Deinitialize()
+{
+	InitializationState = UE::Mass::Subsystems::FInitializationState();
+}
+
+void UMassSubsystemBase::OnWorldBeginPlay(UWorld& InWorld)
+{
+	// This ensure is here to make sure we handle HandleLateCreation() gracefully, we dont expect it to ever trigger unless users start to manually call the functions
+	ensureMsgf(InitializationState.bOnWorldBeginPlayCalled == false, TEXT("%hs called multiple times"), __FUNCTION__);
+	InitializationState.bOnWorldBeginPlayCalled = true;
+}
+
+void UMassSubsystemBase::HandleLateCreation()
+{
+	UE::Mass::Private::HandleLateCreation(*this, InitializationState);
+}
+
 //-----------------------------------------------------------------------------
 // UMassTickableSubsystemBase
 //-----------------------------------------------------------------------------
@@ -41,3 +94,34 @@ bool UMassTickableSubsystemBase::ShouldCreateSubsystem(UObject* Outer) const
 {
 	return UMassSubsystemBase::AreRuntimeMassSubsystemsAllowed(Outer) && Super::ShouldCreateSubsystem(Outer);
 }
+
+void UMassTickableSubsystemBase::Initialize(FSubsystemCollectionBase& Collection)
+{
+	// This ensure is here to make sure we handle HandleLateCreation() gracefully, we dont expect it to ever trigger unless users start to manually call the functions
+	ensureMsgf(InitializationState.bInitializeCalled == false, TEXT("%hs called multiple times"), __FUNCTION__);
+	InitializationState.bInitializeCalled = true;
+}
+
+void UMassTickableSubsystemBase::PostInitialize()
+{
+	// This ensure is here to make sure we handle HandleLateCreation() gracefully, we dont expect it to ever trigger unless users start to manually call the functions
+	ensureMsgf(InitializationState.bPostInitializeCalled == false, TEXT("%hs called multiple times"), __FUNCTION__);
+	InitializationState.bPostInitializeCalled = true;
+}
+
+void UMassTickableSubsystemBase::Deinitialize()
+{
+	InitializationState = UE::Mass::Subsystems::FInitializationState();
+}
+
+void UMassTickableSubsystemBase::OnWorldBeginPlay(UWorld& InWorld)
+{
+	// This ensure is here to make sure we handle HandleLateCreation() gracefully, we dont expect it to ever trigger unless users start to manually call the functions
+	ensureMsgf(InitializationState.bOnWorldBeginPlayCalled == false, TEXT("%hs called multiple times"), __FUNCTION__);
+	InitializationState.bOnWorldBeginPlayCalled = true;
+}
+
+void UMassTickableSubsystemBase::HandleLateCreation()
+{
+	UE::Mass::Private::HandleLateCreation(*this, InitializationState);
+}
diff --git a/Engine/Plugins/Runtime/MassEntity/Source/MassEntity/Public/MassSubsystemBase.h b/Engine/Plugins/Runtime/MassEntity/Source/MassEntity/Public/MassSubsystemBase.h
--- a/Engine/Plugins/Runtime/MassEntity/Source/MassEntity/Public/MassSubsystemBase.h
+++ b/Engine/Plugins/Runtime/MassEntity/Source/MassEntity/Public/MassSubsystemBase.h
@@ -6,6 +6,16 @@
 #include "MassSubsystemBase.generated.h"
 
 
+namespace UE::Mass::Subsystems
+{
+	struct FInitializationState
+	{
+		uint8 bInitializeCalled : 1 = false;
+		uint8 bPostInitializeCalled : 1 = false;
+		uint8 bOnWorldBeginPlayCalled : 1 = false;
+	};
+}
+
 /** 
  * The sole responsibility of this world subsystem class is to serve functionality common to all 
  * Mass-related UWorldSubsystem-based subsystems, like whether the subsystems should get created at all. 
@@ -21,7 +31,23 @@ public:
 protected:
 	//~USubsystem interface
 	virtual bool ShouldCreateSubsystem(UObject* Outer) const override;
+	virtual void Initialize(FSubsystemCollectionBase& Collection) override;
+	virtual void PostInitialize() override;
+	virtual void Deinitialize() override;
+	virtual void OnWorldBeginPlay(UWorld& InWorld) override;
 	//~End of USubsystem interface
+
+	/**
+	 * Needs to be called in Initialize for subsystems we want to behave properly when dynamically added after UWorld::BeginPlay
+	 * (for example via GameplayFeatureActions). This is required for subsystems relying on their PostInitialize and/or OnWorldBeginPlay called.
+	 */
+	void HandleLateCreation();
+
+	/**
+	 * Tracks which initialization function had already been called. Requires the child classes to call Super implementation
+	 * for their Initialize, PostInitialize, Deinitialize and OnWorldBeginPlayCalled overrides
+	 */
+	UE::Mass::Subsystems::FInitializationState InitializationState;
 };
 
 /**
@@ -36,5 +62,22 @@ class MASSENTITY_API UMassTickableSubsystemBase : public UTickableWorldSubsystem
 protected:
 	//~USubsystem interface
 	virtual bool ShouldCreateSubsystem(UObject* Outer) const override;
+	virtual void Initialize(FSubsystemCollectionBase& Collection) override;
+	virtual void PostInitialize() override;
+	virtual void Deinitialize() override;
+	virtual void OnWorldBeginPlay(UWorld& InWorld) override;
 	//~End of USubsystem interface
+
+	/**
+	 * Needs to be called in Initialize for subsystems we want to behave properly when dynamically added after UWorld::BeginPlay
+	 * (for example via GameplayFeatureActions). This is required for subsystems relying on their PostInitialize and/or OnWorldBeginPlay called.
+	 */
+	void HandleLateCreation();
+
+private:
+	/** 
+	 * Tracks which initialization function had already been called. Requires the child classes to call Super implementation
+	 * for their Initialize, PostInitialize, Deinitialize and OnWorldBeginPlayCalled overrides
+	 */
+	UE::Mass::Subsystems::FInitializationState InitializationState;
 };
diff --git a/Engine/Plugins/Runtime/MassEntity/Source/MassEntityEditor/Private/MassEntityEditorSubsystem.cpp b/Engine/Plugins/Runtime/MassEntity/Source/MassEntityEditor/Private/MassEntityEditorSubsystem.cpp
--- a/Engine/Plugins/Runtime/MassEntity/Source/MassEntityEditor/Private/MassEntityEditorSubsystem.cpp
+++ b/Engine/Plugins/Runtime/MassEntity/Source/MassEntityEditor/Private/MassEntityEditorSubsystem.cpp
@@ -103,11 +103,14 @@ void UMassEntityEditorSubsystem::Initialize(FSubsystemCollectionBase& Collection
 	PhaseManager->Initialize(*this, ProcessingPhasesConfig);
 
 	PhaseManager->Start(EntityManager);
+
+	Super::Initialize(Collection);
 }
 
 void UMassEntityEditorSubsystem::Deinitialize()
 {
 	StopAndCleanUp();
+	Super::Deinitialize();
 }
 
 void UMassEntityEditorSubsystem::StopAndCleanUp()
diff --git a/Engine/Plugins/Runtime/MassGameplay/Source/MassActors/Private/MassActorSubsystem.cpp b/Engine/Plugins/Runtime/MassGameplay/Source/MassActors/Private/MassActorSubsystem.cpp
--- a/Engine/Plugins/Runtime/MassGameplay/Source/MassActors/Private/MassActorSubsystem.cpp
+++ b/Engine/Plugins/Runtime/MassGameplay/Source/MassActors/Private/MassActorSubsystem.cpp
@@ -93,6 +93,7 @@ void UMassActorSubsystem::Initialize(FSubsystemCollectionBase& Collection)
 void UMassActorSubsystem::Deinitialize()
 {
 	ActorManager.Reset();
+	Super::Deinitialize();
 }
 
 //----------------------------------------------------------------------//
diff --git a/Engine/Plugins/Runtime/MassGameplay/Source/MassActors/Private/MassAgentSubsystem.cpp b/Engine/Plugins/Runtime/MassGameplay/Source/MassActors/Private/MassAgentSubsystem.cpp
--- a/Engine/Plugins/Runtime/MassGameplay/Source/MassActors/Private/MassAgentSubsystem.cpp
+++ b/Engine/Plugins/Runtime/MassGameplay/Source/MassActors/Private/MassAgentSubsystem.cpp
@@ -69,6 +69,7 @@ void UMassAgentSubsystem::Initialize(FSubsystemCollectionBase& Collection)
 void UMassAgentSubsystem::Deinitialize()
 {
 	EntityManager.Reset();
+	Super::Deinitialize();
 }
 
 FMassEntityTemplateID UMassAgentSubsystem::RegisterAgentComponent(UMassAgentComponent& AgentComp)
diff --git a/Engine/Plugins/Runtime/MassGameplay/Source/MassGameplayDebug/Private/MassDebuggerSubsystem.cpp b/Engine/Plugins/Runtime/MassGameplay/Source/MassGameplayDebug/Private/MassDebuggerSubsystem.cpp
--- a/Engine/Plugins/Runtime/MassGameplay/Source/MassGameplayDebug/Private/MassDebuggerSubsystem.cpp
+++ b/Engine/Plugins/Runtime/MassGameplay/Source/MassGameplayDebug/Private/MassDebuggerSubsystem.cpp
@@ -30,6 +30,7 @@ void UMassDebuggerSubsystem::Deinitialize()
 #if WITH_MASSENTITY_DEBUG
 	FMassDebugger::OnEntitySelectedDelegate.Remove(OnEntitySelectedHandle);
 #endif // WITH_MASSENTITY_DEBUG
+	Super::Deinitialize();
 }
 
 #if WITH_MASSENTITY_DEBUG
diff --git a/Engine/Plugins/Runtime/MassGameplay/Source/MassGameplayEditor/Private/MassActorEditorSubsystem.cpp b/Engine/Plugins/Runtime/MassGameplay/Source/MassGameplayEditor/Private/MassActorEditorSubsystem.cpp
--- a/Engine/Plugins/Runtime/MassGameplay/Source/MassGameplayEditor/Private/MassActorEditorSubsystem.cpp
+++ b/Engine/Plugins/Runtime/MassGameplay/Source/MassGameplayEditor/Private/MassActorEditorSubsystem.cpp
@@ -15,4 +15,6 @@ void UMassActorEditorSubsystem::Initialize(FSubsystemCollectionBase& Collection)
 	check(MassEditorEditorSubsystem);
 	TSharedRef<FMassEntityManager> MassEntityManager = MassEditorEditorSubsystem->GetMutableEntityManager();
 	ActorManager = MakeShareable(new FMassActorManager(MassEntityManager));
+
+	Super::Initialize(Collection);
 }
diff --git a/Engine/Plugins/Runtime/MassGameplay/Source/MassLOD/Private/MassLODSubsystem.cpp b/Engine/Plugins/Runtime/MassGameplay/Source/MassLOD/Private/MassLODSubsystem.cpp
--- a/Engine/Plugins/Runtime/MassGameplay/Source/MassLOD/Private/MassLODSubsystem.cpp
+++ b/Engine/Plugins/Runtime/MassGameplay/Source/MassLOD/Private/MassLODSubsystem.cpp
@@ -117,6 +117,8 @@ void UMassLODSubsystem::Deinitialize()
 			SimSystem->GetOnProcessingPhaseStarted(EMassProcessingPhase::PrePhysics).RemoveAll(this);
 		}
 	}
+
+	Super::Deinitialize();
 }
 
 int32 UMassLODSubsystem::GetValidViewerIdx(const FMassViewerHandle& ViewerHandle) const
diff --git a/Engine/Plugins/Runtime/MassGameplay/Source/MassReplication/Private/MassReplicationSubsystem.cpp b/Engine/Plugins/Runtime/MassGameplay/Source/MassReplication/Private/MassReplicationSubsystem.cpp
--- a/Engine/Plugins/Runtime/MassGameplay/Source/MassReplication/Private/MassReplicationSubsystem.cpp
+++ b/Engine/Plugins/Runtime/MassGameplay/Source/MassReplication/Private/MassReplicationSubsystem.cpp
@@ -49,6 +49,8 @@ void UMassReplicationSubsystem::Deinitialize()
 	World = nullptr;
 	MassLODSubsystem = nullptr;
 	EntityManager.Reset();
+
+	Super::Deinitialize();
 }
 
 UMassReplicationSubsystem::UMassReplicationSubsystem()
diff --git a/Engine/Plugins/Runtime/MassGameplay/Source/MassRepresentation/Private/MassRepresentationSubsystem.cpp b/Engine/Plugins/Runtime/MassGameplay/Source/MassRepresentation/Private/MassRepresentationSubsystem.cpp
--- a/Engine/Plugins/Runtime/MassGameplay/Source/MassRepresentation/Private/MassRepresentationSubsystem.cpp
+++ b/Engine/Plugins/Runtime/MassGameplay/Source/MassRepresentation/Private/MassRepresentationSubsystem.cpp
@@ -390,6 +390,8 @@ void UMassRepresentationSubsystem::Deinitialize()
 		}
 	}
 	EntityManager.Reset();
+
+	Super::Deinitialize();
 }
 
 void UMassRepresentationSubsystem::OnProcessingPhaseStarted(const float DeltaSeconds, const EMassProcessingPhase Phase) const
diff --git a/Engine/Plugins/Runtime/MassGameplay/Source/MassSimulation/Private/MassSimulationSubsystem.cpp b/Engine/Plugins/Runtime/MassGameplay/Source/MassSimulation/Private/MassSimulationSubsystem.cpp
--- a/Engine/Plugins/Runtime/MassGameplay/Source/MassSimulation/Private/MassSimulationSubsystem.cpp
+++ b/Engine/Plugins/Runtime/MassGameplay/Source/MassSimulation/Private/MassSimulationSubsystem.cpp
@@ -80,6 +80,8 @@ void UMassSimulationSubsystem::Initialize(FSubsystemCollectionBase& Collection)
 	EntityManager = EntitySubsystem->GetMutableEntityManager().AsShared();
 	
 	GetOnProcessingPhaseStarted(EMassProcessingPhase::PrePhysics).AddUObject(this, &UMassSimulationSubsystem::OnProcessingPhaseStarted, EMassProcessingPhase::PrePhysics);
+
+	HandleLateCreation();
 }
 
 void UMassSimulationSubsystem::Deinitialize()
@@ -119,6 +121,8 @@ void UMassSimulationSubsystem::PostInitialize()
 
 void UMassSimulationSubsystem::OnWorldBeginPlay(UWorld& InWorld)
 {
+	Super::OnWorldBeginPlay(InWorld);
+
 	// To evaluate the effective processors execution mode, we need to wait on OnWorldBeginPlay before calling
 	// RebuildTickPipeline as we are sure by this time the network is setup correctly.
 	RebuildTickPipeline();
diff --git a/Engine/Plugins/Runtime/MassGameplay/Source/MassSpawner/Private/MassSpawnerSubsystem.cpp b/Engine/Plugins/Runtime/MassGameplay/Source/MassSpawner/Private/MassSpawnerSubsystem.cpp
--- a/Engine/Plugins/Runtime/MassGameplay/Source/MassSpawner/Private/MassSpawnerSubsystem.cpp
+++ b/Engine/Plugins/Runtime/MassGameplay/Source/MassSpawner/Private/MassSpawnerSubsystem.cpp
@@ -43,6 +43,8 @@ void UMassSpawnerSubsystem::Deinitialize()
 {
 	EntityManager.Reset();
 	TemplateRegistryInstance.ShutDown();
+
+	Super::Deinitialize();
 }
 
 void UMassSpawnerSubsystem::SpawnEntities(const FMassEntityTemplate& EntityTemplate, const uint32 NumberToSpawn, TArray<FMassEntityHandle>& OutEntities)
