TEDS: Make it possible to unregister observer queries.

diff --git a/Engine/Plugins/Experimental/TypedElementsDataStorage/Source/TypedElementsDataStorage/Private/Queries/TypedElementExtendedQueryStore.cpp b/Engine/Plugins/Experimental/TypedElementsDataStorage/Source/TypedElementsDataStorage/Private/Queries/TypedElementExtendedQueryStore.cpp
--- a/Engine/Plugins/Experimental/TypedElementsDataStorage/Source/TypedElementsDataStorage/Private/Queries/TypedElementExtendedQueryStore.cpp
+++ b/Engine/Plugins/Experimental/TypedElementsDataStorage/Source/TypedElementsDataStorage/Private/Queries/TypedElementExtendedQueryStore.cpp
@@ -690,7 +690,12 @@ void FTypedElementExtendedQueryStore::UnregisterQueryData(Handle Query, FTypedEl
 		}
 		else if (QueryData.Processor->IsA<UTypedElementQueryObserverCallbackAdapterProcessorBase>())
 		{
-			checkf(false, TEXT("Observer queries can not be unregistered."));
+			UTypedElementQueryObserverCallbackAdapterProcessorBase* Observer =
+				static_cast<UTypedElementQueryObserverCallbackAdapterProcessorBase*>(QueryData.Processor.Get());
+			if (ensure(Observer && PhaseManager.GetEntityManager().IsValid()))
+			{
+				PhaseManager.GetEntityManager()->GetObserverManager().RemoveObserverInstance(*Observer->GetObservedType(), Observer->GetObservedOperation(), *Observer);
+			}
 		}
 		else
 		{
diff --git a/Engine/Plugins/Runtime/MassEntity/Source/MassEntity/Private/MassObserverManager.cpp b/Engine/Plugins/Runtime/MassEntity/Source/MassEntity/Private/MassObserverManager.cpp
--- a/Engine/Plugins/Runtime/MassEntity/Source/MassEntity/Private/MassObserverManager.cpp
+++ b/Engine/Plugins/Runtime/MassEntity/Source/MassEntity/Private/MassObserverManager.cpp
@@ -279,3 +279,35 @@ void FMassObserverManager::AddObserverInstance(const UScriptStruct& FragmentOrTa
 	}
 }
 
+void FMassObserverManager::RemoveObserverInstance(const UScriptStruct& FragmentOrTagType, const EMassObservedOperation Operation, UMassProcessor& ObserverProcessor)
+{
+	if (!ensure(FragmentOrTagType.IsChildOf(FMassFragment::StaticStruct()) || FragmentOrTagType.IsChildOf(FMassTag::StaticStruct())))
+	{
+		return;
+	}
+
+	bool bIsFragmentObserver = FragmentOrTagType.IsChildOf(FMassFragment::StaticStruct());
+
+	TMap<TObjectPtr<const UScriptStruct>, FMassRuntimePipeline>& ObserversMap =
+		bIsFragmentObserver ? *FragmentObservers[(uint8)Operation] : *TagObservers[(uint8)Operation];
+
+	FMassRuntimePipeline* Pipeline = ObserversMap.Find(&FragmentOrTagType);
+	if (!ensureMsgf(Pipeline, TEXT("Trying to remove an observer for a fragment/tag that does not seem to be observed.")))
+	{
+		return;
+	}
+	Pipeline->RemoveProcessor(ObserverProcessor);
+
+	if (Pipeline->Num() == 0)
+	{
+		ObserversMap.Remove(&FragmentOrTagType);
+		if (bIsFragmentObserver)
+		{
+			ObservedFragments[(uint8)Operation].Remove(FragmentOrTagType);
+		}
+		else
+		{
+			ObservedTags[(uint8)Operation].Remove(FragmentOrTagType);
+		}
+	}
+}
diff --git a/Engine/Plugins/Runtime/MassEntity/Source/MassEntity/Public/MassObserverManager.h b/Engine/Plugins/Runtime/MassEntity/Source/MassEntity/Public/MassObserverManager.h
--- a/Engine/Plugins/Runtime/MassEntity/Source/MassEntity/Public/MassObserverManager.h
+++ b/Engine/Plugins/Runtime/MassEntity/Source/MassEntity/Public/MassObserverManager.h
@@ -109,6 +109,7 @@ public:
 	void OnFragmentOrTagOperation(const UScriptStruct& FragmentOrTagType, const FMassArchetypeEntityCollection& EntityCollection, const EMassObservedOperation Operation);
 
 	void AddObserverInstance(const UScriptStruct& FragmentOrTagType, const EMassObservedOperation Operation, UMassProcessor& ObserverProcessor);
+	void RemoveObserverInstance(const UScriptStruct& FragmentOrTagType, const EMassObservedOperation Operation, UMassProcessor& ObserverProcessor);
 
 protected:
 	friend FMassEntityManager;
