Move mass profiling marker to reduce bloat

diff --git a/Engine/Plugins/Runtime/MassEntity/Source/MassEntity/Private/MassSubsystemAccess.cpp b/Engine/Plugins/Runtime/MassEntity/Source/MassEntity/Private/MassSubsystemAccess.cpp
--- a/Engine/Plugins/Runtime/MassEntity/Source/MassEntity/Private/MassSubsystemAccess.cpp
+++ b/Engine/Plugins/Runtime/MassEntity/Source/MassEntity/Private/MassSubsystemAccess.cpp
@@ -23,6 +23,8 @@ FMassSubsystemAccess::FMassSubsystemAccess(UWorld* InWorld)
 
 USubsystem* FMassSubsystemAccess::FetchSubsystemInstance(UWorld* World, TSubclassOf<USubsystem> SubsystemClass)
 {
+	QUICK_SCOPE_CYCLE_COUNTER(Mass_FetchSubsystemInstance);
+
 	check(SubsystemClass);
 	if (SubsystemClass->IsChildOf<UWorldSubsystem>())
 	{
diff --git a/Engine/Plugins/Runtime/MassEntity/Source/MassEntity/Public/MassSubsystemAccess.h b/Engine/Plugins/Runtime/MassEntity/Source/MassEntity/Public/MassSubsystemAccess.h
--- a/Engine/Plugins/Runtime/MassEntity/Source/MassEntity/Public/MassSubsystemAccess.h
+++ b/Engine/Plugins/Runtime/MassEntity/Source/MassEntity/Public/MassSubsystemAccess.h
@@ -184,8 +184,6 @@ protected:
 	template<typename T>
 	T* GetSubsystemInternal(const uint32 SystemIndex)
 	{
-		QUICK_SCOPE_CYCLE_COUNTER(Mass_GetSubsystemInternal)
-
 		if (UNLIKELY(Subsystems.IsValidIndex(SystemIndex) == false))
 		{
 			Subsystems.AddZeroed(Subsystems.Num() - SystemIndex + 1);
