Fixed UMassTickableSubsystemBase not calling the initialization super-implementation which resulted in it not ticking at runtime.

Also:
* added super calls to UMassSubsystemBase as well.
* switched MassLODSubsystem's parent from UMassTickableSubsystemBase over to UMassSubsystemBase, since it doesn't utilize the ticking supplied by TickableWorldSubsystem subclasses

diff --git a/Engine/Plugins/Runtime/MassEntity/Source/MassEntity/Private/MassSubsystemBase.cpp b/Engine/Plugins/Runtime/MassEntity/Source/MassEntity/Private/MassSubsystemBase.cpp
--- a/Engine/Plugins/Runtime/MassEntity/Source/MassEntity/Private/MassSubsystemBase.cpp
+++ b/Engine/Plugins/Runtime/MassEntity/Source/MassEntity/Private/MassSubsystemBase.cpp
@@ -58,6 +58,8 @@ bool UMassSubsystemBase::ShouldCreateSubsystem(UObject* Outer) const
 
 void UMassSubsystemBase::Initialize(FSubsystemCollectionBase& Collection)
 {
+	Super::Initialize(Collection);
+
 	// This ensure is here to make sure we handle HandleLateCreation() gracefully, we dont expect it to ever trigger unless users start to manually call the functions
 	ensureMsgf(InitializationState.bInitializeCalled == false, TEXT("%hs called multiple times"), __FUNCTION__);
 	InitializationState.bInitializeCalled = true;
@@ -65,6 +67,8 @@ void UMassSubsystemBase::Initialize(FSubsystemCollectionBase& Collection)
 
 void UMassSubsystemBase::PostInitialize()
 {
+	Super::PostInitialize();
+
 	// This ensure is here to make sure we handle HandleLateCreation() gracefully, we dont expect it to ever trigger unless users start to manually call the functions
 	ensureMsgf(InitializationState.bPostInitializeCalled == false, TEXT("%hs called multiple times"), __FUNCTION__);
 	InitializationState.bPostInitializeCalled = true;
@@ -73,10 +77,14 @@ void UMassSubsystemBase::PostInitialize()
 void UMassSubsystemBase::Deinitialize()
 {
 	InitializationState = UE::Mass::Subsystems::FInitializationState();
+
+	Super::Deinitialize();
 }
 
 void UMassSubsystemBase::OnWorldBeginPlay(UWorld& InWorld)
 {
+	Super::OnWorldBeginPlay(InWorld);
+
 	// This ensure is here to make sure we handle HandleLateCreation() gracefully, we dont expect it to ever trigger unless users start to manually call the functions
 	ensureMsgf(InitializationState.bOnWorldBeginPlayCalled == false, TEXT("%hs called multiple times"), __FUNCTION__);
 	InitializationState.bOnWorldBeginPlayCalled = true;
@@ -97,6 +105,8 @@ bool UMassTickableSubsystemBase::ShouldCreateSubsystem(UObject* Outer) const
 
 void UMassTickableSubsystemBase::Initialize(FSubsystemCollectionBase& Collection)
 {
+	Super::Initialize(Collection);
+
 	// This ensure is here to make sure we handle HandleLateCreation() gracefully, we dont expect it to ever trigger unless users start to manually call the functions
 	ensureMsgf(InitializationState.bInitializeCalled == false, TEXT("%hs called multiple times"), __FUNCTION__);
 	InitializationState.bInitializeCalled = true;
@@ -104,6 +114,8 @@ void UMassTickableSubsystemBase::Initialize(FSubsystemCollectionBase& Collection
 
 void UMassTickableSubsystemBase::PostInitialize()
 {
+	Super::PostInitialize();
+
 	// This ensure is here to make sure we handle HandleLateCreation() gracefully, we dont expect it to ever trigger unless users start to manually call the functions
 	ensureMsgf(InitializationState.bPostInitializeCalled == false, TEXT("%hs called multiple times"), __FUNCTION__);
 	InitializationState.bPostInitializeCalled = true;
@@ -112,10 +124,14 @@ void UMassTickableSubsystemBase::PostInitialize()
 void UMassTickableSubsystemBase::Deinitialize()
 {
 	InitializationState = UE::Mass::Subsystems::FInitializationState();
+
+	Super::Deinitialize();
 }
 
 void UMassTickableSubsystemBase::OnWorldBeginPlay(UWorld& InWorld)
 {
+	Super::OnWorldBeginPlay(InWorld);
+
 	// This ensure is here to make sure we handle HandleLateCreation() gracefully, we dont expect it to ever trigger unless users start to manually call the functions
 	ensureMsgf(InitializationState.bOnWorldBeginPlayCalled == false, TEXT("%hs called multiple times"), __FUNCTION__);
 	InitializationState.bOnWorldBeginPlayCalled = true;
diff --git a/Engine/Plugins/Runtime/MassGameplay/Source/MassLOD/Private/MassLODSubsystem.cpp b/Engine/Plugins/Runtime/MassGameplay/Source/MassLOD/Private/MassLODSubsystem.cpp
--- a/Engine/Plugins/Runtime/MassGameplay/Source/MassLOD/Private/MassLODSubsystem.cpp
+++ b/Engine/Plugins/Runtime/MassGameplay/Source/MassLOD/Private/MassLODSubsystem.cpp
@@ -105,11 +105,6 @@ void UMassLODSubsystem::OnPrePhysicsPhaseStarted(float DeltaTime)
 	SynchronizeViewers();
 }
 
-TStatId UMassLODSubsystem::GetStatId() const
-{
-	RETURN_QUICK_DECLARE_CYCLE_STAT(UMassLODSubsystem, STATGROUP_Tickables);
-}
-
 void UMassLODSubsystem::Deinitialize()
 {
 	// Remove all viewers from the viewer info list
diff --git a/Engine/Plugins/Runtime/MassGameplay/Source/MassLOD/Public/MassLODSubsystem.h b/Engine/Plugins/Runtime/MassGameplay/Source/MassLOD/Public/MassLODSubsystem.h
--- a/Engine/Plugins/Runtime/MassGameplay/Source/MassLOD/Public/MassLODSubsystem.h
+++ b/Engine/Plugins/Runtime/MassGameplay/Source/MassLOD/Public/MassLODSubsystem.h
@@ -73,7 +73,7 @@ DECLARE_MULTICAST_DELEGATE_ThreeParams(FOnViewerRemoved, FMassViewerHandle Viewe
  * Manager responsible to manage and synchronized available viewers
  */
 UCLASS(config = Mass, defaultconfig)
-class MASSLOD_API UMassLODSubsystem : public UMassTickableSubsystemBase
+class MASSLOD_API UMassLODSubsystem : public UMassSubsystemBase
 {
 	GENERATED_BODY()
 
@@ -124,8 +124,6 @@ protected:
 	virtual void Deinitialize() override;
 	// USubsystem END
 
-	virtual TStatId GetStatId() const override;
-
 	/** Called at the start of the PrePhysics mass processing phase and calls SynchronizeViewers */ 
 	void OnPrePhysicsPhaseStarted(float DeltaTime);
 
