Fix to handling duplicate invalid entities in MassArchetypeEntityCollectionWithPayload::CreateEntityRangesWithPayload. Also added handling of invalid entity handles in FMassArchetypeEntityCollection.

The CL contains two new unit tests exercising the new code.

diff --git a/Engine/Plugins/Runtime/MassEntity/Source/MassEntity/Private/MassArchetypeData.h b/Engine/Plugins/Runtime/MassEntity/Source/MassEntity/Private/MassArchetypeData.h
--- a/Engine/Plugins/Runtime/MassEntity/Source/MassEntity/Private/MassArchetypeData.h
+++ b/Engine/Plugins/Runtime/MassEntity/Source/MassEntity/Private/MassArchetypeData.h
@@ -233,7 +233,8 @@ public:
 	void* GetFragmentDataForEntityChecked(const UScriptStruct* FragmentType, int32 EntityIndex) const;
 	void* GetFragmentDataForEntity(const UScriptStruct* FragmentType, int32 EntityIndex) const;
 
-	FORCEINLINE int32 GetInternalIndexForEntity(const int32 EntityIndex) const { return EntityMap.FindChecked(EntityIndex); }
+	FORCEINLINE const int32* GetInternalIndexForEntity(const int32 EntityIndex) const { return EntityMap.Find(EntityIndex); }
+	FORCEINLINE int32 GetInternalIndexForEntityChecked(const int32 EntityIndex) const { return EntityMap.FindChecked(EntityIndex); }
 	int32 GetNumEntitiesPerChunk() const { return NumEntitiesPerChunk; }
 	int32 GetBytesPerEntity() const { return TotalBytesPerEntity; }
 
diff --git a/Engine/Plugins/Runtime/MassEntity/Source/MassEntity/Private/MassArchetypeTypes.cpp b/Engine/Plugins/Runtime/MassEntity/Source/MassEntity/Private/MassArchetypeTypes.cpp
--- a/Engine/Plugins/Runtime/MassEntity/Source/MassEntity/Private/MassArchetypeTypes.cpp
+++ b/Engine/Plugins/Runtime/MassEntity/Source/MassEntity/Private/MassArchetypeTypes.cpp
@@ -31,12 +31,18 @@ FMassArchetypeEntityCollection::FMassArchetypeEntityCollection(const FMassArchet
 	// InEntities has a real chance of not being sorted by AbsoluteIndex. We gotta fix that to optimize how we process the data 
 	TArray<int32> TrueIndices;
 	TrueIndices.AddUninitialized(InEntities.Num());
-	int32 i = 0;
+	int32 NumValidEntities = 0;
 	if (ArchetypeData)
 	{
 		for (const FMassEntityHandle& Entity : InEntities)
 		{
-			TrueIndices[i++] = ArchetypeData->GetInternalIndexForEntity(Entity.Index);
+			if (Entity.IsValid())
+			{
+				if (const int32* TrueIndex = ArchetypeData->GetInternalIndexForEntity(Entity.Index))
+				{
+					TrueIndices[NumValidEntities++] = *TrueIndex;
+				}
+			}
 		}
 	}
 	else
@@ -46,10 +52,14 @@ FMassArchetypeEntityCollection::FMassArchetypeEntityCollection(const FMassArchet
 		// FMassArchetypeEntityCollection as the generic batched API wrapper for entities
 		for (const FMassEntityHandle& Entity : InEntities)
 		{
-			TrueIndices[i++] = Entity.Index;
+			if (Entity.IsValid())
+			{
+				TrueIndices[NumValidEntities++] = Entity.Index;
+			}
 		}
 	}
 
+	TrueIndices.SetNum(NumValidEntities, /*bAllowShrinking=*/false);
 	TrueIndices.Sort();
 
 #if DO_GUARD_SLOW
@@ -203,9 +213,11 @@ void FMassArchetypeEntityCollectionWithPayload::CreateEntityRangesWithPayload(co
 			return ArchetypeIndex == Other.ArchetypeIndex && TrueIndex == Other.TrueIndex;
 		}
 
-		bool operator<(const FEntityInArchetype& Other) const
+		/** @return whether A should come before B in an ordered collection */
+		static bool Compare(const FEntityInArchetype& A, const FEntityInArchetype& B)
 		{
-			return ArchetypeIndex < Other.ArchetypeIndex || (ArchetypeIndex == Other.ArchetypeIndex && TrueIndex < Other.TrueIndex);
+			// using "greater" to ensure INDEX_NONE archetypes end up at the end of the collection
+			return A.ArchetypeIndex > B.ArchetypeIndex || (A.ArchetypeIndex == B.ArchetypeIndex && A.TrueIndex < B.TrueIndex);
 		}
 	};
 
@@ -246,7 +258,7 @@ void FMassArchetypeEntityCollectionWithPayload::CreateEntityRangesWithPayload(co
 				ArchetypeIndex = Archetypes.Add({ ArchetypeHandle, 0 });
 			}
 			++Archetypes[ArchetypeIndex].Count;
-			EntityData[i] = { ArchetypeIndex, ArchetypePtr ? ArchetypePtr->GetInternalIndexForEntity(Entity.Index) : Entity.Index };
+			EntityData[i] = { ArchetypeIndex, ArchetypePtr ? ArchetypePtr->GetInternalIndexForEntityChecked(Entity.Index) : Entity.Index };
 		}
 		else
 		{
@@ -268,7 +280,7 @@ void FMassArchetypeEntityCollectionWithPayload::CreateEntityRangesWithPayload(co
 	UE::Mass::Utils::AbstractSort(Entities.Num(), [&EntityData, &bDuplicatesFound](const int32 LHS, const int32 RHS)
 		{
 			bDuplicatesFound = bDuplicatesFound || (EntityData[LHS] == EntityData[RHS]);
-			return EntityData[LHS] < EntityData[RHS];
+			return FEntityInArchetype::Compare(EntityData[LHS], EntityData[RHS]);
 		}
 		, [&EntityData, &Payload](const int32 A, const int32 B)
 		{
@@ -300,7 +312,13 @@ void FMassArchetypeEntityCollectionWithPayload::CreateEntityRangesWithPayload(co
 		for (int32 EntryIndex = 0; EntryIndex < EntityData.Num() - 1; ++EntryIndex)
 		{	
 			FEntityInArchetype& Entry = EntityData[EntryIndex];
-			if (Entry != EntityData[EntryIndex + 1])
+			if (Entry.ArchetypeIndex == INDEX_NONE)
+			{
+				// we're reached INDEX_NONE archetypes, which are at the end of EntityData
+				// breaking since there's nothing more to process. 
+				break;
+			}
+			else if (Entry != EntityData[EntryIndex + 1])
 			{
 				continue;
 			}
@@ -322,12 +340,10 @@ void FMassArchetypeEntityCollectionWithPayload::CreateEntityRangesWithPayload(co
 	}
 
 	int32 ProcessedEntitiesCount = 0;
-	int32 ArchetypeIndex = 0;
 	for (FArchetypeInfo& ArchetypeInfo : Archetypes)
-	{		
+	{
 		TArrayView<FEntityInArchetype> EntityDataSubset = MakeArrayView(&EntityData[ProcessedEntitiesCount], ArchetypeInfo.Count);
-		ensure(EntityDataSubset[0].ArchetypeIndex == ArchetypeIndex);
-		ensure(EntityDataSubset.Last().ArchetypeIndex == ArchetypeIndex);
+		ensure(EntityDataSubset[0].ArchetypeIndex == EntityDataSubset.Last().ArchetypeIndex);
 		TStridedView<int32> TrueIndices = MakeStridedView(EntityDataSubset, &FEntityInArchetype::TrueIndex);
 
 		FMassGenericPayloadViewSlice PayloadSubView(Payload, ProcessedEntitiesCount, ArchetypeInfo.Count);
@@ -335,7 +351,6 @@ void FMassArchetypeEntityCollectionWithPayload::CreateEntityRangesWithPayload(co
 		OutEntityCollections.Add(FMassArchetypeEntityCollectionWithPayload(ArchetypeInfo.Archetype, TrueIndices, MoveTemp(PayloadSubView)));
 
 		ProcessedEntitiesCount += ArchetypeInfo.Count;
-		++ArchetypeIndex;
 	}
 }
 
diff --git a/Engine/Plugins/Runtime/MassEntity/Source/MassEntityTestSuite/Private/MassChunkCollectionTest.cpp b/Engine/Plugins/Runtime/MassEntity/Source/MassEntityTestSuite/Private/MassChunkCollectionTest.cpp
--- a/Engine/Plugins/Runtime/MassEntity/Source/MassEntityTestSuite/Private/MassChunkCollectionTest.cpp
+++ b/Engine/Plugins/Runtime/MassEntity/Source/MassEntityTestSuite/Private/MassChunkCollectionTest.cpp
@@ -172,6 +172,85 @@ struct FEntityCollection_CreateWithDuplicates : FEntityCollectionTestBase
 };
 IMPLEMENT_AI_INSTANT_TEST(FEntityCollection_CreateWithDuplicates, "System.Mass.EntityCollection.Create.Duplicates");
 
+struct FEntityCollection_CreateWithInvalidDuplicates : FEntityCollectionTestBase
+{
+	virtual bool InstantTest() override
+	{
+		{
+			TArray<FMassEntityHandle> EntitiesSubSet;
+
+			EntitiesSubSet.Add(FMassEntityHandle());
+			EntitiesSubSet.Add(Entities[0]);
+			EntitiesSubSet.Add(Entities[0]);
+			EntitiesSubSet.Add(FMassEntityHandle());
+
+			FMassArchetypeEntityCollection Collection(FloatsArchetype, EntitiesSubSet, FMassArchetypeEntityCollection::FoldDuplicates);
+
+			// The resulting Collection should have only a single Range consisting of a single entity (matching Entities[0])
+			AITEST_EQUAL(TEXT("We expect only a single resulting range"), Collection.GetRanges().Num(), 1);
+			AITEST_EQUAL(TEXT("We expect only a single entity in the resulting range"), Collection.GetRanges()[0].SubchunkStart, 0);
+			AITEST_EQUAL(TEXT("We expect only a single entity in the resulting range"), Collection.GetRanges()[0].Length, 1);
+		}
+
+		{
+			TArray<FMassEntityHandle> EntitiesSubSet;
+
+			EntitiesSubSet.Add(Entities[4]);
+			EntitiesSubSet.Add(FMassEntityHandle());
+			EntitiesSubSet.Add(FMassEntityHandle()); 
+			EntitiesSubSet.Add(Entities[3]);
+			EntitiesSubSet.Add(FMassEntityHandle());
+			EntitiesSubSet.Add(Entities[1]);
+
+			FMassArchetypeEntityCollection Collection(FloatsArchetype, EntitiesSubSet, FMassArchetypeEntityCollection::FoldDuplicates);
+
+			// The resulting Collection should have two Ranges for a single archetype, one of them with two entities (3, 4).
+			AITEST_EQUAL(TEXT("We expect two resulting range"), Collection.GetRanges().Num(), 2);
+			AITEST_EQUAL(TEXT("We expect the first range to consist of a single entity"), Collection.GetRanges()[0].Length, 1);
+			AITEST_EQUAL(TEXT("We expect the second range to consist of a two entities"), Collection.GetRanges()[1].Length, 2);
+			// the specific composition of resulting ranges is being tested by other tests
+		}
+
+		return true;
+	}
+};
+IMPLEMENT_AI_INSTANT_TEST(FEntityCollection_CreateWithInvalidDuplicates, "System.Mass.EntityCollection.Create.InvalidDuplicates");
+
+struct FEntityCollection_CreateWithInvalidDuplicatesWithPayload : FEntityCollectionTestBase
+{
+	virtual bool InstantTest() override
+	{
+		TArray<FMassEntityHandle> EntitiesSubSet;
+		TArray<FTestFragment_Int> Payload;
+
+		EntitiesSubSet.Add(FMassEntityHandle());
+		Payload.Add(FTestFragment_Int(int32(2)));
+
+		EntitiesSubSet.Add(Entities[0]);
+		Payload.Add(FTestFragment_Int(int32(0)));
+
+		EntitiesSubSet.Add(Entities[0]);
+		Payload.Add(FTestFragment_Int(int32(1)));
+
+		EntitiesSubSet.Add(FMassEntityHandle());
+		Payload.Add(FTestFragment_Int(int32(3)));
+
+		// transform typed payload array into generic one for sorting purposes
+		FStructArrayView PaloadView(Payload);
+		TArray<FMassArchetypeEntityCollectionWithPayload> Result;
+		FMassArchetypeEntityCollectionWithPayload::CreateEntityRangesWithPayload(*EntityManager, EntitiesSubSet, FMassArchetypeEntityCollection::FoldDuplicates
+			, FMassGenericPayloadView(MakeArrayView(&PaloadView, 1)), Result);
+
+		AITEST_EQUAL(TEXT("We expect only a single result"), Result.Num(), 1);
+		AITEST_EQUAL(TEXT("We expect only a single resulting range"), Result[0].GetEntityCollection().GetRanges().Num(), 1);
+		AITEST_EQUAL(TEXT("We expect only a single entity in the resulting range"), Result[0].GetEntityCollection().GetRanges()[0].SubchunkStart, 0);
+		AITEST_EQUAL(TEXT("We expect only a single entity in the resulting range"), Result[0].GetEntityCollection().GetRanges()[0].Length, 1);
+
+		return true;
+	}
+};
+IMPLEMENT_AI_INSTANT_TEST(FEntityCollection_CreateWithInvalidDuplicatesWithPayload, "System.Mass.EntityCollection.Create.InvalidDuplicatesWithPayloadWithPayload");
+
 #if WITH_MASSENTITY_DEBUG
 struct FEntityCollection_WithPayloadBase : FEntityCollectionTestBase
 {
