Critical fix to MassArchetypeEntityCollectionWithPayload::CreateEntityRangesWithPayload issue introduced in recent 30226914. The change modifies a piece of code previous change affected but not updated.

diff --git a/Engine/Plugins/Runtime/MassEntity/Source/MassEntity/Private/MassArchetypeTypes.cpp b/Engine/Plugins/Runtime/MassEntity/Source/MassEntity/Private/MassArchetypeTypes.cpp
--- a/Engine/Plugins/Runtime/MassEntity/Source/MassEntity/Private/MassArchetypeTypes.cpp
+++ b/Engine/Plugins/Runtime/MassEntity/Source/MassEntity/Private/MassArchetypeTypes.cpp
@@ -340,10 +340,13 @@ void FMassArchetypeEntityCollectionWithPayload::CreateEntityRangesWithPayload(co
 	}
 
 	int32 ProcessedEntitiesCount = 0;
-	for (FArchetypeInfo& ArchetypeInfo : Archetypes)
+	// processing from the back since that's how EntityData is sorted - higher-index archetypes come first
+	for (int32 ArchetypeIndex = Archetypes.Num() - 1; ArchetypeIndex >= 0; --ArchetypeIndex)
 	{
+		FArchetypeInfo& ArchetypeInfo = Archetypes[ArchetypeIndex];
 		TArrayView<FEntityInArchetype> EntityDataSubset = MakeArrayView(&EntityData[ProcessedEntitiesCount], ArchetypeInfo.Count);
-		ensure(EntityDataSubset[0].ArchetypeIndex == EntityDataSubset.Last().ArchetypeIndex);
+		ensure(EntityDataSubset[0].ArchetypeIndex == ArchetypeIndex);
+		ensure(EntityDataSubset.Last().ArchetypeIndex == ArchetypeIndex);
 		TStridedView<int32> TrueIndices = MakeStridedView(EntityDataSubset, &FEntityInArchetype::TrueIndex);
 
 		FMassGenericPayloadViewSlice PayloadSubView(Payload, ProcessedEntitiesCount, ArchetypeInfo.Count);
