Improving Mass entity memory debugging display in MassDebugger

The CL also changes some memory-related archetype members to SIZE_T

diff --git a/Engine/Plugins/Runtime/MassEntity/Source/MassEntity/Private/MassArchetypeData.cpp b/Engine/Plugins/Runtime/MassEntity/Source/MassEntity/Private/MassArchetypeData.cpp
--- a/Engine/Plugins/Runtime/MassEntity/Source/MassEntity/Private/MassArchetypeData.cpp
+++ b/Engine/Plugins/Runtime/MassEntity/Source/MassEntity/Private/MassArchetypeData.cpp
@@ -33,11 +33,12 @@ namespace UE::Mass
 
 FMassArchetypeData::FMassArchetypeData(const FMassArchetypeCreationParams& CreationParams)
 	: NumEntitiesPerChunk(UE::Mass::Private::UninitializedInt32)
-	, TotalBytesPerEntity(UE::Mass::Private::UninitializedInt32)
 	, EntityListOffsetWithinChunk(UE::Mass::Private::UninitializedInt32)
 	, ChunkMemorySize(UE::Mass::SanitizeChunkMemorySize(CreationParams.ChunkMemorySize ? CreationParams.ChunkMemorySize : GET_MASS_CONFIG_VALUE(ChunkMemorySize)))
 {
+#if WITH_MASSENTITY_DEBUG
 	DebugNames.Add(CreationParams.DebugName);
+#endif // WITH_MASSENTITY_DEBUG
 }
 
 void FMassArchetypeData::ForEachFragmentType(TFunction< void(const UScriptStruct* /*Fragment*/)> Function) const
@@ -55,6 +56,15 @@ bool FMassArchetypeData::HasFragmentType(const UScriptStruct* FragmentType) cons
 
 void FMassArchetypeData::Initialize(const FMassArchetypeCompositionDescriptor& InCompositionDescriptor, const uint32 ArchetypeDataVersion)
 {
+	if (!ensureMsgf(Chunks.Num() == 0, TEXT("Trying to re-initialize non-empty Mass Archetype is not supported")))
+	{
+		return;
+	}
+	if (!ensureMsgf(CreatedArchetypeDataVersion == 0, TEXT("MassArchetype has already been initialized")))
+	{
+		return;
+	}
+
 	CreatedArchetypeDataVersion = ArchetypeDataVersion;
 	CompositionDescriptor.Fragments = InCompositionDescriptor.Fragments;
 	ConfigureFragments();
@@ -111,10 +121,10 @@ void FMassArchetypeData::ConfigureFragments()
 	SortedFragmentList.Sort(FScriptStructSortOperator());
 
 	// Figure out how many bytes all of the individual fragments (and metadata) will cost per entity
-	int32 FragmentSizeTallyBytes = 0;
+	SIZE_T FragmentSizeTallyBytes = 0;
 
 	// Alignment padding computation is currently very conservative and over-estimated.
-	int32 AlignmentPadding = 0;
+	SIZE_T AlignmentPadding = 0;
 	
 	// Save room for the 'metadata' (entity array)
 	FragmentSizeTallyBytes += sizeof(FMassEntityHandle);
@@ -129,8 +139,8 @@ void FMassArchetypeData::ConfigureFragments()
 		checkSlow(FragmentType);
 		FragmentConfigs[FragmentIndex].FragmentType = FragmentType;
 
-		AlignmentPadding += FragmentType->GetMinAlignment();
-		FragmentSizeTallyBytes += FragmentType->GetStructureSize();
+		AlignmentPadding += SIZE_T(FragmentType->GetMinAlignment());
+		FragmentSizeTallyBytes += SIZE_T(FragmentType->GetStructureSize());
 
 		FragmentIndexMap.Add(FragmentType, FragmentIndex);
 	}
@@ -142,12 +152,12 @@ void FMassArchetypeData::ConfigureFragments()
 	NumEntitiesPerChunk = ChunkAvailableSize / TotalBytesPerEntity;
 
 	// Set up the offsets for each fragment into the chunk data
-	int32 CurrentOffset = NumEntitiesPerChunk * sizeof(FMassEntityHandle);
+	SIZE_T CurrentOffset = NumEntitiesPerChunk * sizeof(FMassEntityHandle);
 	for (FMassArchetypeFragmentConfig& FragmentData : FragmentConfigs)
 	{
 		CurrentOffset = Align(CurrentOffset, FragmentData.FragmentType->GetMinAlignment());
 		FragmentData.ArrayOffsetWithinChunk = CurrentOffset;
-		const int32 SizeOfThisFragmentArray = NumEntitiesPerChunk * FragmentData.FragmentType->GetStructureSize();
+		const SIZE_T SizeOfThisFragmentArray = NumEntitiesPerChunk * FragmentData.FragmentType->GetStructureSize();
 		CurrentOffset += SizeOfThisFragmentArray;
 	}
 }
@@ -845,6 +855,13 @@ FString FMassArchetypeData::DebugGetDescription() const
 #endif
 }
 
+#if WITH_MASSENTITY_DEBUG
+void FMassArchetypeData::DebugGetEntityMemoryNumbers(SIZE_T& OutActiveChunksMemorySize, SIZE_T& OutActiveEntitiesMemorySize) const
+{
+	OutActiveChunksMemorySize = GetChunkAllocSize() *  Chunks.Num();
+	OutActiveEntitiesMemorySize = TotalBytesPerEntity * EntityMap.Num();
+}
+
 FString FMassArchetypeData::GetCombinedDebugNamesAsString() const
 {
 	TStringBuilder<256> StringBuilder;
@@ -859,7 +876,6 @@ FString FMassArchetypeData::GetCombinedDebugNamesAsString() const
 	return StringBuilder.ToString();
 }
 
-#if WITH_MASSENTITY_DEBUG
 void FMassArchetypeData::DebugPrintArchetype(FOutputDevice& Ar)
 {
 	Ar.Logf(ELogVerbosity::Log, TEXT("Name: %s"), *GetCombinedDebugNamesAsString());
@@ -901,7 +917,7 @@ void FMassArchetypeData::DebugPrintArchetype(FOutputDevice& Ar)
 	{
 		Ar.Logf(ELogVerbosity::Log, TEXT("\tEntity Occupancy: %.1f%%"), CurrentEntityCapacity > 0 ? ((EntityMap.Num() * 100.0f) / (float)CurrentEntityCapacity) : 0.f);
 	}
-	Ar.Logf(ELogVerbosity::Log, TEXT("\tBytes / Entity  : %d"), TotalBytesPerEntity);
+	Ar.Logf(ELogVerbosity::Log, TEXT("\tBytes / Entity  : %zu"), TotalBytesPerEntity);
 	Ar.Logf(ELogVerbosity::Log, TEXT("\tEntities / Chunk: %d"), NumEntitiesPerChunk);
 
 	Ar.Logf(ELogVerbosity::Log, TEXT("\tOffset 0x%04X: Entity[] (%d bytes each)"), EntityListOffsetWithinChunk, sizeof(FMassEntityHandle));
@@ -914,7 +930,7 @@ void FMassArchetypeData::DebugPrintArchetype(FOutputDevice& Ar)
 
 	//@TODO: Print out padding in between things?
 
-	const int32 UnusuablePaddingOffset = TotalBytesPerEntity * NumEntitiesPerChunk;
+	const SIZE_T UnusuablePaddingOffset = TotalBytesPerEntity * NumEntitiesPerChunk;
 	const int32 UnusuablePaddingAmount = GetChunkAllocSize() - UnusuablePaddingOffset;
 	if (UnusuablePaddingAmount > 0)
 	{
diff --git a/Engine/Plugins/Runtime/MassEntity/Source/MassEntity/Private/MassArchetypeData.h b/Engine/Plugins/Runtime/MassEntity/Source/MassEntity/Private/MassArchetypeData.h
--- a/Engine/Plugins/Runtime/MassEntity/Source/MassEntity/Private/MassArchetypeData.h
+++ b/Engine/Plugins/Runtime/MassEntity/Source/MassEntity/Private/MassArchetypeData.h
@@ -21,14 +21,14 @@ struct FMassArchetypeChunk
 {
 private:
 	uint8* RawMemory = nullptr;
-	int32 AllocSize = 0;
+	SIZE_T AllocSize = 0;
 	int32 NumInstances = 0;
 	int32 SerialModificationNumber = 0;
 	TArray<FInstancedStruct> ChunkFragmentData;
 	FMassArchetypeSharedFragmentValues SharedFragmentValues;
 
 public:
-	explicit FMassArchetypeChunk(const int32 InAllocSize, TConstArrayView<FInstancedStruct> InChunkFragmentTemplates, FMassArchetypeSharedFragmentValues InSharedFragmentValues)
+	explicit FMassArchetypeChunk(const SIZE_T InAllocSize, TConstArrayView<FInstancedStruct> InChunkFragmentTemplates, FMassArchetypeSharedFragmentValues InSharedFragmentValues)
 		: AllocSize(InAllocSize)
 		, ChunkFragmentData(InChunkFragmentTemplates)
 		, SharedFragmentValues(InSharedFragmentValues)
@@ -170,17 +170,19 @@ private:
 	TMap<const UScriptStruct*, int32> FragmentIndexMap;
 
 	int32 NumEntitiesPerChunk;
-	int32 TotalBytesPerEntity;
+	SIZE_T TotalBytesPerEntity = 0;
 	int32 EntityListOffsetWithinChunk;
 
 	// Archetype version at which this archetype was created, useful for query to do incremental archetype matching
-	uint32 CreatedArchetypeDataVersion;
+	uint32 CreatedArchetypeDataVersion = 0;
 
+#if WITH_MASSENTITY_DEBUG
 	// Arrays of names the archetype is referred as.
 	TArray<FName> DebugNames;
+#endif // WITH_MASSENTITY_DEBUG
 
 	// Defaults to UMassEntitySettings.ChunkMemorySize. In near future will support being set via constructor.
-	const int32 ChunkMemorySize;
+	const SIZE_T ChunkMemorySize = 0;
 	
 	friend FMassEntityQuery;
 	friend FMassArchetypeEntityCollection;
@@ -236,11 +238,11 @@ public:
 	FORCEINLINE const int32* GetInternalIndexForEntity(const int32 EntityIndex) const { return EntityMap.Find(EntityIndex); }
 	FORCEINLINE int32 GetInternalIndexForEntityChecked(const int32 EntityIndex) const { return EntityMap.FindChecked(EntityIndex); }
 	int32 GetNumEntitiesPerChunk() const { return NumEntitiesPerChunk; }
-	int32 GetBytesPerEntity() const { return TotalBytesPerEntity; }
+	SIZE_T GetBytesPerEntity() const { return TotalBytesPerEntity; }
 
 	int32 GetNumEntities() const { return EntityMap.Num(); }
 
-	int32 GetChunkAllocSize() const { return ChunkMemorySize; }
+	SIZE_T GetChunkAllocSize() const { return ChunkMemorySize; }
 
 	int32 GetChunkCount() const { return Chunks.Num(); }
 
@@ -292,19 +294,27 @@ public:
 	// Converts the list of fragments into a user-readable debug string
 	FString DebugGetDescription() const;
 
+	/** Copies debug names from another archetype data. */
+	void CopyDebugNamesFrom(const FMassArchetypeData& Other)
+	{ 
+#if WITH_MASSENTITY_DEBUG
+		DebugNames = Other.DebugNames; 
+#endif // WITH_MASSENTITY_DEBUG
+	}
+
+#if WITH_MASSENTITY_DEBUG
+	/** Fetches how much memory is allocated for active chunks, and how much of that memory is actually occupied */
+	void DebugGetEntityMemoryNumbers(SIZE_T& OutActiveChunksMemorySize, SIZE_T& OutActiveEntitiesMemorySize) const;
+
 	/** Adds new debug name associated with the archetype. */
 	void AddUniqueDebugName(const FName& Name) { DebugNames.AddUnique(Name); }
 	
 	/** @return array of debug names associated with this archetype. */
 	const TConstArrayView<FName> GetDebugNames() const { return DebugNames; }
 	
-	/** Copies debug names from another archetype data. */
-	void CopyDebugNamesFrom(const FMassArchetypeData& Other) { DebugNames = Other.DebugNames; }
-	
 	/** @return string of all debug names combined */
 	FString GetCombinedDebugNamesAsString() const;
 
-#if WITH_MASSENTITY_DEBUG
 	/**
 	 * Prints out debug information about the archetype
 	 */
diff --git a/Engine/Plugins/Runtime/MassEntity/Source/MassEntity/Private/MassDebugger.cpp b/Engine/Plugins/Runtime/MassEntity/Source/MassEntity/Private/MassDebugger.cpp
--- a/Engine/Plugins/Runtime/MassEntity/Source/MassEntity/Private/MassDebugger.cpp
+++ b/Engine/Plugins/Runtime/MassEntity/Source/MassEntity/Private/MassDebugger.cpp
@@ -436,6 +436,11 @@ void FMassDebugger::GetArchetypeEntityStats(const FMassArchetypeHandle& Archetyp
 	OutStats.ChunksCount = ArchetypeData.GetChunkCount();
 	OutStats.AllocatedSize = ArchetypeData.GetAllocatedSize();
 	OutStats.BytesPerEntity = ArchetypeData.GetBytesPerEntity();
+
+	SIZE_T ActiveChunksMemorySize = 0;
+	SIZE_T ActiveEntitiesMemorySize = 0;
+	ArchetypeData.DebugGetEntityMemoryNumbers(ActiveChunksMemorySize, ActiveEntitiesMemorySize);
+	OutStats.WastedEntityMemory = ActiveChunksMemorySize - ActiveEntitiesMemorySize;
 }
 
 const TConstArrayView<FName> FMassDebugger::GetArchetypeDebugNames(const FMassArchetypeHandle& ArchetypeHandle)
diff --git a/Engine/Plugins/Runtime/MassEntity/Source/MassEntity/Private/MassEntityManager.cpp b/Engine/Plugins/Runtime/MassEntity/Source/MassEntity/Private/MassEntityManager.cpp
--- a/Engine/Plugins/Runtime/MassEntity/Source/MassEntity/Private/MassEntityManager.cpp
+++ b/Engine/Plugins/Runtime/MassEntity/Source/MassEntity/Private/MassEntityManager.cpp
@@ -217,11 +217,13 @@ FMassArchetypeHandle FMassEntityManager::CreateArchetype(const FMassArchetypeCom
 	{
 		if (Ptr->IsEquivalent(Composition))
 		{
+#if WITH_MASSENTITY_DEBUG
 			// Keep track of all names for this archetype.
 			if (!CreationParams.DebugName.IsNone())
 			{
 				Ptr->AddUniqueDebugName(CreationParams.DebugName);
 			}
+#endif // WITH_MASSENTITY_DEBUG
 			if (CreationParams.ChunkMemorySize > 0 && CreationParams.ChunkMemorySize != Ptr->GetChunkAllocSize())
 			{
 				UE_LOG(LogMass, Warning, TEXT("Reusing existing Archetype, but the requested ChunkMemorySize is different. Requested %d, existing: %d")
diff --git a/Engine/Plugins/Runtime/MassEntity/Source/MassEntity/Public/MassDebugger.h b/Engine/Plugins/Runtime/MassEntity/Source/MassEntity/Public/MassDebugger.h
--- a/Engine/Plugins/Runtime/MassEntity/Source/MassEntity/Public/MassDebugger.h
+++ b/Engine/Plugins/Runtime/MassEntity/Source/MassEntity/Public/MassDebugger.h
@@ -33,6 +33,8 @@ namespace UE::Mass::Debug
 		int32 ChunksCount = 0;
 		/** Total amount of memory taken by this archetype */
 		SIZE_T AllocatedSize = 0;
+		/** How much memory allocated for entities is being unused */
+		SIZE_T WastedEntityMemory = 0;
 		/** Total amount of memory needed by a single entity */
 		int32 BytesPerEntity = 0;
 	};
diff --git a/Engine/Plugins/Runtime/MassEntity/Source/MassEntityDebugger/Private/SMassArchetype.cpp b/Engine/Plugins/Runtime/MassEntity/Source/MassEntityDebugger/Private/SMassArchetype.cpp
--- a/Engine/Plugins/Runtime/MassEntity/Source/MassEntityDebugger/Private/SMassArchetype.cpp
+++ b/Engine/Plugins/Runtime/MassEntity/Source/MassEntityDebugger/Private/SMassArchetype.cpp
@@ -55,12 +55,20 @@ void SMassArchetype::Construct(const FArguments& InArgs, TSharedPtr<FMassDebugge
 		.Font(FCoreStyle::GetDefaultFontStyle("Mono", 9))
 	];
 	
-	FText ArchetypeDescription = FText::Format(LOCTEXT("ArchetypeDescrption", "EntitiesCount: {0}\nBytesPerEntity: {1}\nEntitiesCountPerChunk: {2}\nChunksCount: {3}\nAllocated memory: {4}")
+	FText ArchetypeDescription = FText::Format(LOCTEXT("ArchetypeDescrption", "EntitiesCount: {0}\
+		\nBytesPerEntity: {1}\
+		\nEntitiesCountPerChunk: {2}\
+		\nChunksCount: {3}\
+		\nAllocated memory: {4}\
+		\nWasted memory : {5} ({6}%)")
 		, FText::AsNumber(ArchetypeDebugData.ArchetypeStats.EntitiesCount)
 		, FText::AsMemory(ArchetypeDebugData.ArchetypeStats.BytesPerEntity)
 		, FText::AsNumber(ArchetypeDebugData.ArchetypeStats.EntitiesCountPerChunk)
 		, FText::AsNumber(ArchetypeDebugData.ArchetypeStats.ChunksCount)
-		, FText::AsMemory(ArchetypeDebugData.ArchetypeStats.AllocatedSize));
+		, FText::AsMemory(ArchetypeDebugData.ArchetypeStats.AllocatedSize)
+		, FText::AsMemory(ArchetypeDebugData.ArchetypeStats.WastedEntityMemory)
+		, FText::AsNumber(float(ArchetypeDebugData.ArchetypeStats.WastedEntityMemory) * 100.f / ArchetypeDebugData.ArchetypeStats.AllocatedSize)
+		);
 	
 	Box->AddSlot()
 	.AutoHeight()
@@ -70,6 +78,22 @@ void SMassArchetype::Construct(const FArguments& InArgs, TSharedPtr<FMassDebugge
 		.Text(ArchetypeDescription)
 	];
 
+	if (ArchetypeDebugData.ArchetypeStats.EntitiesCount != 0 && ArchetypeDebugData.ArchetypeStats.ChunksCount != 0)
+	{
+		const float AverageEntitiesPerChunkActual = float(ArchetypeDebugData.ArchetypeStats.EntitiesCount) / ArchetypeDebugData.ArchetypeStats.ChunksCount;
+		FText DerivedArchetypeDescription = FText::Format(LOCTEXT("ArchetypeDescrptionAux", "Actual average Entities per Chunk: {0}\nChunk occupancy: {1}")
+			, FText::AsNumber(AverageEntitiesPerChunkActual)
+			, FText::AsNumber(AverageEntitiesPerChunkActual / ArchetypeDebugData.ArchetypeStats.EntitiesCountPerChunk));
+
+		Box->AddSlot()
+			.AutoHeight()
+			.Padding(0, 4)
+			[
+				SNew(STextBlock)
+					.Text(DerivedArchetypeDescription)
+			];
+	}
+
 	const FMassArchetypeCompositionDescriptor& Composition = ArchetypeData->Composition;
 	const FSlateBrush* Brush = FMassDebuggerStyle::GetBrush("MassDebug.Fragment");
 
diff --git a/Engine/Plugins/Runtime/MassEntity/Source/MassEntityDebugger/Private/SMassArchetypesView.cpp b/Engine/Plugins/Runtime/MassEntity/Source/MassEntityDebugger/Private/SMassArchetypesView.cpp
--- a/Engine/Plugins/Runtime/MassEntity/Source/MassEntityDebugger/Private/SMassArchetypesView.cpp
+++ b/Engine/Plugins/Runtime/MassEntity/Source/MassEntityDebugger/Private/SMassArchetypesView.cpp
@@ -241,15 +241,39 @@ void SMassArchetypesView::RebuildSelectedView()
 	{
 		// print out entity count and memory totals
 		int32 EntitiesCount = 0;
+		int32 ChunksCount = 0;
 		SIZE_T AllocatedSize = 0;
+		SIZE_T WastedEntityMemory = 0;
+
 		for (const TSharedPtr<FMassDebuggerArchetypeData>& ArchetypeData : DebuggerModel->SelectedArchetypes)
 		{
 			EntitiesCount += ArchetypeData->ArchetypeStats.EntitiesCount;
+			ChunksCount += ArchetypeData->ArchetypeStats.ChunksCount;
 			AllocatedSize += ArchetypeData->ArchetypeStats.AllocatedSize;
+			WastedEntityMemory += ArchetypeData->ArchetypeStats.WastedEntityMemory;
+		}
+
+		float ChunkOccupancy = 0.f;
+		for (const TSharedPtr<FMassDebuggerArchetypeData>& ArchetypeData : DebuggerModel->SelectedArchetypes)
+		{
+			if (ArchetypeData->ArchetypeStats.ChunksCount)
+			{
+				const float AverageEntitiesPerChunkActual = float(ArchetypeData->ArchetypeStats.EntitiesCount) / ArchetypeData->ArchetypeStats.ChunksCount;
+				const float Occupancy = AverageEntitiesPerChunkActual / ArchetypeData->ArchetypeStats.EntitiesCountPerChunk;
+				// @todo we're doing a simple proportional accumulation here. Percentiles would work better to avoid
+				// outliers skewing the results. 
+				ChunkOccupancy += Occupancy * ArchetypeData->ArchetypeStats.EntitiesCount / EntitiesCount;
+			}
 		}
 
-		const FText TotalDescription = FText::Format(LOCTEXT("ArchetypesSelectionTotalsContent", "EntitiesCount: {0}\nAllocated memory: {1}")
-			, FText::AsNumber(EntitiesCount), FText::AsMemory(AllocatedSize));
+		const FText TotalDescription = FText::Format(LOCTEXT("ArchetypesSelectionTotalsContent", "EntitiesCount: {0}\
+				\nAllocated memory: {1}\
+				\nWasted memory: {2} ({3}%)\
+				\nTotal chunks: {4}\
+				\nChunk occupancy: {5}")
+			, FText::AsNumber(EntitiesCount), FText::AsMemory(AllocatedSize)
+			, FText::AsMemory(WastedEntityMemory), FText::AsNumber(float(WastedEntityMemory) * 100.f / AllocatedSize)
+			, FText::AsNumber(ChunksCount), FText::AsNumber(ChunkOccupancy));
 
 		TSharedRef<SVerticalBox> Box = SNew(SVerticalBox);
 		
