Fixed MassDebugger's chunk memory waste display.

The previous implementation was wrong since it didn't take into account that an archetype might have "empty" chunks, i.e. existing chunks that have no entities and their chunk memory being released.

diff --git a/Engine/Plugins/Runtime/MassEntity/Source/MassEntity/Private/MassArchetypeData.cpp b/Engine/Plugins/Runtime/MassEntity/Source/MassEntity/Private/MassArchetypeData.cpp
--- a/Engine/Plugins/Runtime/MassEntity/Source/MassEntity/Private/MassArchetypeData.cpp
+++ b/Engine/Plugins/Runtime/MassEntity/Source/MassEntity/Private/MassArchetypeData.cpp
@@ -811,16 +811,22 @@ void FMassArchetypeData::BindSharedFragmentRequirements(FMassExecutionContext& R
 	}
 }
 
-SIZE_T FMassArchetypeData::GetAllocatedSize() const
+int32 FMassArchetypeData::GetNonEmptyChunkCount() const
 {
-	int32 NumAllocatedChunkBuffers = 0;
+	int32 NumAllocatedChunks = 0;
 	for (const FMassArchetypeChunk& Chunk : Chunks)
 	{
 		if (Chunk.GetRawMemory() != nullptr)
 		{
-			++NumAllocatedChunkBuffers;
+			++NumAllocatedChunks;
 		}
 	}
+	return NumAllocatedChunks;
+}
+
+SIZE_T FMassArchetypeData::GetAllocatedSize() const
+{
+	const int32 NumAllocatedChunkBuffers = GetNonEmptyChunkCount();
 
 	return sizeof(FMassArchetypeData) +
 		ChunkFragmentsTemplate.GetAllocatedSize() +
@@ -858,7 +864,7 @@ FString FMassArchetypeData::DebugGetDescription() const
 #if WITH_MASSENTITY_DEBUG
 void FMassArchetypeData::DebugGetEntityMemoryNumbers(SIZE_T& OutActiveChunksMemorySize, SIZE_T& OutActiveEntitiesMemorySize) const
 {
-	OutActiveChunksMemorySize = GetChunkAllocSize() *  Chunks.Num();
+	OutActiveChunksMemorySize = GetChunkAllocSize() * GetNonEmptyChunkCount();
 	OutActiveEntitiesMemorySize = TotalBytesPerEntity * EntityMap.Num();
 }
 
diff --git a/Engine/Plugins/Runtime/MassEntity/Source/MassEntity/Private/MassArchetypeData.h b/Engine/Plugins/Runtime/MassEntity/Source/MassEntity/Private/MassArchetypeData.h
--- a/Engine/Plugins/Runtime/MassEntity/Source/MassEntity/Private/MassArchetypeData.h
+++ b/Engine/Plugins/Runtime/MassEntity/Source/MassEntity/Private/MassArchetypeData.h
@@ -245,6 +245,7 @@ public:
 	SIZE_T GetChunkAllocSize() const { return ChunkMemorySize; }
 
 	int32 GetChunkCount() const { return Chunks.Num(); }
+	int32 GetNonEmptyChunkCount() const;
 
 	uint32 GetCreatedArchetypeDataVersion() const { return CreatedArchetypeDataVersion; }
 
