Changes to make editor-available Mass processors actually execute in editor mode.

diff --git a/Engine/Plugins/Runtime/MassEntity/Source/MassEntity/Private/MassEntitySettings.cpp b/Engine/Plugins/Runtime/MassEntity/Source/MassEntity/Private/MassEntitySettings.cpp
--- a/Engine/Plugins/Runtime/MassEntity/Source/MassEntity/Private/MassEntitySettings.cpp
+++ b/Engine/Plugins/Runtime/MassEntity/Source/MassEntity/Private/MassEntitySettings.cpp
@@ -56,6 +56,8 @@ void UMassEntitySettings::BuildProcessorListAndPhases()
 	BuildProcessorList();
 	BuildPhases();
 	bInitialized = true;
+
+	OnInitializedEvent.Broadcast();
 }
 
 void UMassEntitySettings::BuildPhases()
@@ -77,7 +79,7 @@ void UMassEntitySettings::BuildPhases()
 			FMassPhaseProcessorConfigurationHelper Configurator(*PhaseConfig.PhaseProcessor, PhaseConfig, *this, EMassProcessingPhase(i));
 			Configurator.bInitializeCreatedProcessors = false;
 			Configurator.bIsGameRuntime = false;
-			Configurator.Configure({}, /*EntityManager=*/nullptr, &Result);
+			Configurator.Configure({}, EProcessorExecutionFlags::All, /*EntityManager=*/nullptr, &Result);
 
 			if (Result.PrunedProcessorClasses.Num() > 0)
 			{
diff --git a/Engine/Plugins/Runtime/MassEntity/Source/MassEntity/Private/MassEntityUtils.cpp b/Engine/Plugins/Runtime/MassEntity/Source/MassEntity/Private/MassEntityUtils.cpp
--- a/Engine/Plugins/Runtime/MassEntity/Source/MassEntity/Private/MassEntityUtils.cpp
+++ b/Engine/Plugins/Runtime/MassEntity/Source/MassEntity/Private/MassEntityUtils.cpp
@@ -36,6 +36,19 @@ EProcessorExecutionFlags GetProcessorExecutionFlagsForWorld(const UWorld& World)
 	}
 }
 
+EProcessorExecutionFlags DetermineProcessorExecutionFlags(const UWorld* World, EProcessorExecutionFlags ExecutionFlagsOverride)
+{
+	if (ExecutionFlagsOverride != EProcessorExecutionFlags::None)
+	{
+		return ExecutionFlagsOverride;
+	}
+	if (World)
+	{
+		return UE::Mass::Utils::GetProcessorExecutionFlagsForWorld(*World);
+	}
+	return EProcessorExecutionFlags::All;
+}
+
 void CreateEntityCollections(const FMassEntityManager& EntityManager, const TConstArrayView<FMassEntityHandle> Entities
 	, const FMassArchetypeEntityCollection::EDuplicatesHandling DuplicatesHandling, TArray<FMassArchetypeEntityCollection>& OutEntityCollections)
 {
diff --git a/Engine/Plugins/Runtime/MassEntity/Source/MassEntity/Private/MassObserverManager.cpp b/Engine/Plugins/Runtime/MassEntity/Source/MassEntity/Private/MassObserverManager.cpp
--- a/Engine/Plugins/Runtime/MassEntity/Source/MassEntity/Private/MassObserverManager.cpp
+++ b/Engine/Plugins/Runtime/MassEntity/Source/MassEntity/Private/MassObserverManager.cpp
@@ -21,13 +21,10 @@ namespace UE::Mass::ObserverManager
 	{
 	// a helper function to reduce code duplication in FMassObserverManager::Initialize
 	template<typename TBitSet, typename TPointerType>
-	void SetUpObservers(FMassEntityManager& EntityManager, const TMap<TPointerType, FMassProcessorClassCollection>& RegisteredObserverTypes, TBitSet& ObservedBitSet, FMassObserversMap& Observers)
+	void SetUpObservers(FMassEntityManager& EntityManager, const EProcessorExecutionFlags WorldExecutionFlags, UObject& Owner
+		, const TMap<TPointerType, FMassProcessorClassCollection>& RegisteredObserverTypes, TBitSet& ObservedBitSet, FMassObserversMap& Observers)
 	{
 		ObservedBitSet.Reset();
-		UObject* Owner = EntityManager.GetOwner();
-		check(Owner);
-		const UWorld* World = Owner->GetWorld();
-		const EProcessorExecutionFlags WorldExecutionFlags = World ? UE::Mass::Utils::GetProcessorExecutionFlagsForWorld(*World) : EProcessorExecutionFlags::All;
 
 		for (auto It : RegisteredObserverTypes)
 		{
@@ -43,10 +40,10 @@ namespace UE::Mass::ObserverManager
 			{
 				if (ProcessorClass->GetDefaultObject<UMassProcessor>()->ShouldExecute(WorldExecutionFlags))
 				{
-					Pipeline.AppendProcessor(ProcessorClass, *Owner);
+					Pipeline.AppendProcessor(ProcessorClass, Owner);
 				}
 			}
-			Pipeline.Initialize(*Owner);
+			Pipeline.Initialize(Owner);
 		}
 	};
 	} // Private
@@ -72,11 +69,16 @@ void FMassObserverManager::Initialize()
 	// instantiate initializers
 	const UMassObserverRegistry& Registry = UMassObserverRegistry::Get();
 
+	UObject* Owner = EntityManager.GetOwner();
+	check(Owner);
+	const UWorld* World = Owner->GetWorld();
+	const EProcessorExecutionFlags WorldExecutionFlags = UE::Mass::Utils::DetermineProcessorExecutionFlags(World);
+
 	using UE::Mass::ObserverManager::Private::SetUpObservers;
 	for (int i = 0; i < (int)EMassObservedOperation::MAX; ++i)
 	{
-		SetUpObservers(EntityManager, *Registry.FragmentObservers[i], ObservedFragments[i], FragmentObservers[i]);
-		SetUpObservers(EntityManager, *Registry.TagObservers[i], ObservedTags[i], TagObservers[i]);
+		SetUpObservers(EntityManager, WorldExecutionFlags, *Owner, *Registry.FragmentObservers[i], ObservedFragments[i], FragmentObservers[i]);
+		SetUpObservers(EntityManager, WorldExecutionFlags, *Owner, *Registry.TagObservers[i], ObservedTags[i], TagObservers[i]);
 	}
 }
 
diff --git a/Engine/Plugins/Runtime/MassEntity/Source/MassEntity/Private/MassProcessingPhaseManager.cpp b/Engine/Plugins/Runtime/MassEntity/Source/MassEntity/Private/MassProcessingPhaseManager.cpp
--- a/Engine/Plugins/Runtime/MassEntity/Source/MassEntity/Private/MassProcessingPhaseManager.cpp
+++ b/Engine/Plugins/Runtime/MassEntity/Source/MassEntity/Private/MassProcessingPhaseManager.cpp
@@ -144,10 +144,11 @@ void FMassProcessingPhase::Initialize(FMassProcessingPhaseManager& InPhaseManage
 //----------------------------------------------------------------------//
 // FPhaseProcessorConfigurator
 //----------------------------------------------------------------------//
-void FMassPhaseProcessorConfigurationHelper::Configure(TArrayView<UMassProcessor* const> DynamicProcessors,
-	const TSharedPtr<FMassEntityManager>& EntityManager, FMassProcessorDependencySolver::FResult* OutOptionalResult)
+void FMassPhaseProcessorConfigurationHelper::Configure(TArrayView<UMassProcessor* const> DynamicProcessors
+	, EProcessorExecutionFlags InWorldExecutionFlags, const TSharedPtr<FMassEntityManager>& EntityManager
+	, FMassProcessorDependencySolver::FResult* OutOptionalResult)
 {
-	FMassRuntimePipeline TmpPipeline;
+	FMassRuntimePipeline TmpPipeline(InWorldExecutionFlags);
 	TmpPipeline.CreateFromArray(PhaseConfig.ProcessorCDOs, ProcessorOuter);
 	for (UMassProcessor* Processor : DynamicProcessors)
 	{
@@ -163,7 +164,7 @@ void FMassPhaseProcessorConfigurationHelper::Configure(TArrayView<UMassProcessor
 
 	Solver.ResolveDependencies(SortedProcessors, EntityManager, OutOptionalResult);
 
-	PhaseProcessor.UpdateProcessorsCollection(SortedProcessors);
+	PhaseProcessor.UpdateProcessorsCollection(SortedProcessors, InWorldExecutionFlags);
 
 #if WITH_MASSENTITY_DEBUG
 	for (const FMassProcessorOrderInfo& ProcessorOrderInfo : SortedProcessors)
@@ -197,13 +198,15 @@ void FMassPhaseProcessorConfigurationHelper::Configure(TArrayView<UMassProcessor
 //----------------------------------------------------------------------//
 void FMassProcessingPhaseManager::Initialize(UObject& InOwner, TConstArrayView<FMassProcessingPhaseConfig> InProcessingPhasesConfig, const FString& DependencyGraphFileName)
 {
-#if WITH_EDITOR
 	UWorld* World = InOwner.GetWorld();
+#if WITH_EDITOR
 	const bool bCreateProcessorGraphPreview = (World != nullptr) && (World->IsEditorWorld() && !World->IsGameWorld());
 #endif // WITH_EDITOR
 	Owner = &InOwner;
 	ProcessingPhasesConfig = InProcessingPhasesConfig;
 
+	ProcessorExecutionFlags = UE::Mass::Utils::DetermineProcessorExecutionFlags(World, ProcessorExecutionFlags);
+
 	for (int PhaseAsInt = 0; PhaseAsInt < int(EMassProcessingPhase::MAX); ++PhaseAsInt)
 	{		
 		const EMassProcessingPhase Phase = EMassProcessingPhase(PhaseAsInt);
@@ -234,7 +237,9 @@ void FMassProcessingPhaseManager::Initialize(UObject& InOwner, TConstArrayView<F
 			Result.DependencyGraphFileName = DependencyGraphFileName;
 			FMassPhaseProcessorConfigurationHelper Configurator(*PhaseProcessor, ProcessingPhasesConfig[PhaseAsInt], InOwner, EMassProcessingPhase(PhaseAsInt));
 			Configurator.bIsGameRuntime = false;
-			Configurator.Configure({}, /*EntityManager=*/nullptr, &Result);
+			// passing EProcessorExecutionFlags::All here to gather all available processors since bCreateProcessorGraphPreview 
+			// is true when we want to preview processors that might be available at runtime.
+			Configurator.Configure({}, EProcessorExecutionFlags::All, /*EntityManager=*/nullptr, &Result);
 		}
 #endif // WITH_EDITOR
 	}
@@ -397,7 +402,7 @@ void FMassProcessingPhaseManager::OnPhaseStart(const FMassProcessingPhase& Phase
 			GraphBuildState.LastResult.Reset();
 
 			FMassPhaseProcessorConfigurationHelper Configurator(*PhaseProcessor, ProcessingPhasesConfig[PhaseAsInt], *Owner.Get(), Phase.Phase);
-			Configurator.Configure(DynamicProcessors, EntityManager, &GraphBuildState.LastResult);
+			Configurator.Configure(DynamicProcessors, ProcessorExecutionFlags, EntityManager, &GraphBuildState.LastResult);
 
 			GraphBuildState.bInitialized = true;
 
diff --git a/Engine/Plugins/Runtime/MassEntity/Source/MassEntity/Private/MassProcessingTypes.cpp b/Engine/Plugins/Runtime/MassEntity/Source/MassEntity/Private/MassProcessingTypes.cpp
--- a/Engine/Plugins/Runtime/MassEntity/Source/MassEntity/Private/MassProcessingTypes.cpp
+++ b/Engine/Plugins/Runtime/MassEntity/Source/MassEntity/Private/MassProcessingTypes.cpp
@@ -100,7 +100,7 @@ void FMassRuntimePipeline::InitializeFromClassArray(TConstArrayView<TSubclassOf<
 	Reset();
 
 	const UWorld* World = InOwner.GetWorld();
-	const EProcessorExecutionFlags WorldExecutionFlags = World ? UE::Mass::Utils::GetProcessorExecutionFlagsForWorld(*World) : EProcessorExecutionFlags::All;
+	const EProcessorExecutionFlags WorldExecutionFlags = UE::Mass::Utils::DetermineProcessorExecutionFlags(World, ExecutionFlags);
 
 	for (const TSubclassOf<UMassProcessor>& ProcessorClass : InProcessorClasses)
 	{
@@ -131,7 +131,7 @@ bool FMassRuntimePipeline::HasProcessorOfExactClass(TSubclassOf<UMassProcessor>
 void FMassRuntimePipeline::AppendUniqueRuntimeProcessorCopies(TConstArrayView<const UMassProcessor*> InProcessors, UObject& InOwner)
 {
 	const UWorld* World = InOwner.GetWorld();
-	const EProcessorExecutionFlags WorldExecutionFlags = World ? UE::Mass::Utils::GetProcessorExecutionFlagsForWorld(*World) : EProcessorExecutionFlags::All;
+	const EProcessorExecutionFlags WorldExecutionFlags = UE::Mass::Utils::DetermineProcessorExecutionFlags(World, ExecutionFlags);
 	const int32 StartingCount = Processors.Num();
 		
 	for (const UMassProcessor* Proc : InProcessors)
@@ -175,8 +175,8 @@ void FMassRuntimePipeline::AppendUniqueRuntimeProcessorCopies(TConstArrayView<co
 void FMassRuntimePipeline::AppendOrOverrideRuntimeProcessorCopies(TConstArrayView<const UMassProcessor*> InProcessors, UObject& InOwner)
 {
 	const UWorld* World = InOwner.GetWorld();
-	const EProcessorExecutionFlags WorldExecutionFlags = World ? UE::Mass::Utils::GetProcessorExecutionFlagsForWorld(*World) : EProcessorExecutionFlags::All;
-	
+	const EProcessorExecutionFlags WorldExecutionFlags = UE::Mass::Utils::DetermineProcessorExecutionFlags(World, ExecutionFlags);
+
 	for (const UMassProcessor* Proc : InProcessors)
 	{
 		if (Proc && Proc->ShouldExecute(WorldExecutionFlags))
diff --git a/Engine/Plugins/Runtime/MassEntity/Source/MassEntity/Private/MassProcessor.cpp b/Engine/Plugins/Runtime/MassEntity/Source/MassEntity/Private/MassProcessor.cpp
--- a/Engine/Plugins/Runtime/MassEntity/Source/MassEntity/Private/MassProcessor.cpp
+++ b/Engine/Plugins/Runtime/MassEntity/Source/MassEntity/Private/MassProcessor.cpp
@@ -424,13 +424,13 @@ void UMassCompositeProcessor::BuildFlatProcessingGraph(TConstArrayView<FMassProc
 #endif // WITH_MASSENTITY_DEBUG
 }
 
-void UMassCompositeProcessor::UpdateProcessorsCollection(TArrayView<FMassProcessorOrderInfo> InOutOrderedProcessors)
+void UMassCompositeProcessor::UpdateProcessorsCollection(TArrayView<FMassProcessorOrderInfo> InOutOrderedProcessors, EProcessorExecutionFlags InWorldExecutionFlags)
 {
 	TArray<TObjectPtr<UMassProcessor>> ExistingProcessors(ChildPipeline.GetMutableProcessors());
 	ChildPipeline.Reset();
 
 	const UWorld* World = GetWorld();
-	const EProcessorExecutionFlags WorldExecutionFlags = World ? UE::Mass::Utils::GetProcessorExecutionFlagsForWorld(*World) : EProcessorExecutionFlags::All;
+	const EProcessorExecutionFlags WorldExecutionFlags = UE::Mass::Utils::DetermineProcessorExecutionFlags(World, InWorldExecutionFlags);
 	const FMassProcessingPhaseConfig& PhaseConfig = GET_MASS_CONFIG_VALUE(GetProcessingPhaseConfig(ProcessingPhase));
 
 	for (FMassProcessorOrderInfo& ProcessorInfo : InOutOrderedProcessors)
@@ -545,5 +545,5 @@ UMassCompositeProcessor* UMassCompositeProcessor::FindOrAddGroupProcessor(FName
 void UMassCompositeProcessor::Populate(TConstArrayView<FMassProcessorOrderInfo> OrderedProcessors)
 {
 	TArray<FMassProcessorOrderInfo> OrderedProcessorsCopy(OrderedProcessors);
-	UpdateProcessorsCollection(OrderedProcessorsCopy);
+	UpdateProcessorsCollection(OrderedProcessorsCopy, EProcessorExecutionFlags::None);
 }
diff --git a/Engine/Plugins/Runtime/MassEntity/Source/MassEntity/Public/MassEntitySettings.h b/Engine/Plugins/Runtime/MassEntity/Source/MassEntity/Public/MassEntitySettings.h
--- a/Engine/Plugins/Runtime/MassEntity/Source/MassEntity/Public/MassEntitySettings.h
+++ b/Engine/Plugins/Runtime/MassEntity/Source/MassEntity/Public/MassEntitySettings.h
@@ -25,6 +25,7 @@ public:
 #if WITH_EDITORONLY_DATA
 	DECLARE_MULTICAST_DELEGATE_OneParam(FOnSettingsChange, const FPropertyChangedEvent& /*PropertyChangedEvent*/);
 #endif // WITH_EDITORONLY_DATA
+	DECLARE_MULTICAST_DELEGATE(FOnInitialized);
 
 	UMassEntitySettings(const FObjectInitializer& ObjectInitializer = FObjectInitializer::Get());
 
@@ -34,8 +35,11 @@ public:
 	TConstArrayView<FMassProcessingPhaseConfig> GetProcessingPhasesConfig();
 	const FMassProcessingPhaseConfig& GetProcessingPhaseConfig(const EMassProcessingPhase ProcessingPhase) const { check(ProcessingPhase != EMassProcessingPhase::MAX); return ProcessingPhasesConfig[int(ProcessingPhase)]; }
 
+	static FOnInitialized& GetOnInitializedEvent() { return GET_MASS_CONFIG_VALUE(OnInitializedEvent); }
 #if WITH_EDITOR
-	FOnSettingsChange& GetOnSettingsChange() { return OnSettingsChange; }	
+	FOnSettingsChange& GetOnSettingsChange() { return OnSettingsChange; }
+
+	static bool IsInitialized() { return GET_MASS_CONFIG_VALUE(bInitialized); }
 
 protected:
 	virtual void PostEditChangeProperty(struct FPropertyChangedEvent& PropertyChangedEvent) override;
@@ -75,4 +79,6 @@ protected:
 #endif // WITH_EDITORONLY_DATA
 	bool bInitialized = false;
 	bool bEngineInitialized = false;
+
+	FOnInitialized OnInitializedEvent;
 };
diff --git a/Engine/Plugins/Runtime/MassEntity/Source/MassEntity/Public/MassEntityUtils.h b/Engine/Plugins/Runtime/MassEntity/Source/MassEntity/Public/MassEntityUtils.h
--- a/Engine/Plugins/Runtime/MassEntity/Source/MassEntity/Public/MassEntityUtils.h
+++ b/Engine/Plugins/Runtime/MassEntity/Source/MassEntity/Public/MassEntityUtils.h
@@ -15,6 +15,9 @@ namespace UE::Mass::Utils
 /** returns the current execution mode for the processors calculated from the world network mode */
 MASSENTITY_API extern EProcessorExecutionFlags GetProcessorExecutionFlagsForWorld(const UWorld& World);
 
+/** based on given World (which can be null) and ExecutionFlagsOverride the function determines the execution flags to use */
+MASSENTITY_API extern EProcessorExecutionFlags DetermineProcessorExecutionFlags(const UWorld* World, EProcessorExecutionFlags ExecutionFlagsOverride = EProcessorExecutionFlags::None);
+
 UE_DEPRECATED(5.3, "This function is deprecated. Please use the properly spelled GetProcessorExecutionFlagsForWorld instead.")
 FORCEINLINE EProcessorExecutionFlags GetProcessorExecutionFlagsForWold(const UWorld& World)
 {
diff --git a/Engine/Plugins/Runtime/MassEntity/Source/MassEntity/Public/MassProcessingPhaseManager.h b/Engine/Plugins/Runtime/MassEntity/Source/MassEntity/Public/MassProcessingPhaseManager.h
--- a/Engine/Plugins/Runtime/MassEntity/Source/MassEntity/Public/MassProcessingPhaseManager.h
+++ b/Engine/Plugins/Runtime/MassEntity/Source/MassEntity/Public/MassProcessingPhaseManager.h
@@ -100,8 +100,12 @@ struct MASSENTITY_API FMassPhaseProcessorConfigurationHelper
 	{
 	}
 
-	void Configure(TArrayView<UMassProcessor* const> DynamicProcessors, const TSharedPtr<FMassEntityManager>& EntityManager = TSharedPtr<FMassEntityManager>(), 
-		FMassProcessorDependencySolver::FResult* OutOptionalResult = nullptr);
+	/** 
+	 * @param InWorldExecutionFlags - provide EProcessorExecutionFlags::None to let underlying code decide
+	 */
+	void Configure(TArrayView<UMassProcessor* const> DynamicProcessors, EProcessorExecutionFlags InWorldExecutionFlags
+		, const TSharedPtr<FMassEntityManager>& EntityManager = TSharedPtr<FMassEntityManager>()
+		, FMassProcessorDependencySolver::FResult* OutOptionalResult = nullptr);
 
 	UMassCompositeProcessor& PhaseProcessor;
 	const FMassProcessingPhaseConfig& PhaseConfig;
@@ -113,7 +117,14 @@ struct MASSENTITY_API FMassPhaseProcessorConfigurationHelper
 	UE_DEPRECATED(5.2, "This flavor of Configure has been deprecated. Use the one requiring the first parameter to be an array view of additional processors")
 	void Configure(const TSharedPtr<FMassEntityManager>& EntityManager = TSharedPtr<FMassEntityManager>(), FMassProcessorDependencySolver::FResult* OutOptionalResult = nullptr)
 	{
-		Configure({}, EntityManager, OutOptionalResult);
+		Configure({}, EProcessorExecutionFlags::None, EntityManager, OutOptionalResult);
+	}
+
+	UE_DEPRECATED(5.3, "This flavor of Configure has been deprecated. Use the one requiring the first parameter to be an array view of additional processors and EProcessorExecutionFlags to be provided")
+	void Configure(TArrayView<UMassProcessor* const> DynamicProcessors, const TSharedPtr<FMassEntityManager>& EntityManager = TSharedPtr<FMassEntityManager>(),
+		FMassProcessorDependencySolver::FResult* OutOptionalResult = nullptr)
+	{
+		Configure(DynamicProcessors, EProcessorExecutionFlags::None, EntityManager, OutOptionalResult);
 	}
 };
 
@@ -128,7 +139,9 @@ struct MASSENTITY_API FMassPhaseProcessorConfigurationHelper
 struct MASSENTITY_API FMassProcessingPhaseManager : public FGCObject
 {
 public:
-	FMassProcessingPhaseManager() = default;
+	explicit FMassProcessingPhaseManager(EProcessorExecutionFlags InProcessorExecutionFlags = EProcessorExecutionFlags::None) 
+		: ProcessorExecutionFlags(InProcessorExecutionFlags)
+	{}
 	FMassProcessingPhaseManager(const FMassProcessingPhaseManager& Other) = delete;
 	FMassProcessingPhaseManager& operator=(const FMassProcessingPhaseManager& Other) = delete;
 
@@ -222,5 +235,6 @@ protected:
 
 	FDelegateHandle OnNewArchetypeHandle;
 
+	EProcessorExecutionFlags ProcessorExecutionFlags = EProcessorExecutionFlags::None;
 	bool bIsAllowedToTick = false;
 };
diff --git a/Engine/Plugins/Runtime/MassEntity/Source/MassEntity/Public/MassProcessingTypes.h b/Engine/Plugins/Runtime/MassEntity/Source/MassEntity/Public/MassProcessingTypes.h
--- a/Engine/Plugins/Runtime/MassEntity/Source/MassEntity/Public/MassProcessingTypes.h
+++ b/Engine/Plugins/Runtime/MassEntity/Source/MassEntity/Public/MassProcessingTypes.h
@@ -76,7 +76,13 @@ private:
 	UPROPERTY()
 	TArray<TObjectPtr<UMassProcessor>> Processors;
 
+	EProcessorExecutionFlags ExecutionFlags = EProcessorExecutionFlags::None;
+
 public:
+	explicit FMassRuntimePipeline(EProcessorExecutionFlags WorldExecutionFlags = EProcessorExecutionFlags::None)
+		: ExecutionFlags(WorldExecutionFlags)
+	{}
+
 	void Reset();
 	void Initialize(UObject& Owner);
 	
diff --git a/Engine/Plugins/Runtime/MassEntity/Source/MassEntity/Public/MassProcessor.h b/Engine/Plugins/Runtime/MassEntity/Source/MassEntity/Public/MassProcessor.h
--- a/Engine/Plugins/Runtime/MassEntity/Source/MassEntity/Public/MassProcessor.h
+++ b/Engine/Plugins/Runtime/MassEntity/Source/MassEntity/Public/MassProcessor.h
@@ -216,7 +216,7 @@ public:
 	 * will be updated to reflect the reuse.
 	 * The described behavior however is available only for processors with bAllowMultipleInstances == false.
 	 */
-	void UpdateProcessorsCollection(TArrayView<FMassProcessorOrderInfo> InOutOrderedProcessors);
+	void UpdateProcessorsCollection(TArrayView<FMassProcessorOrderInfo> InOutOrderedProcessors, EProcessorExecutionFlags InWorldExecutionFlags = EProcessorExecutionFlags::None);
 
 	/** adds SubProcessor to an appropriately named group. If RequestedGroupName == None then SubProcessor
 	 *  will be added directly to ChildPipeline. If not then the indicated group will be searched for in ChildPipeline 
diff --git a/Engine/Plugins/Runtime/MassEntity/Source/MassEntityEditor/Private/MassEntityEditorSubsystem.cpp b/Engine/Plugins/Runtime/MassEntity/Source/MassEntityEditor/Private/MassEntityEditorSubsystem.cpp
--- a/Engine/Plugins/Runtime/MassEntity/Source/MassEntityEditor/Private/MassEntityEditorSubsystem.cpp
+++ b/Engine/Plugins/Runtime/MassEntity/Source/MassEntityEditor/Private/MassEntityEditorSubsystem.cpp
@@ -69,7 +69,9 @@ struct FMassEditorTickFunction : FTickableEditorObject
 
 UMassEntityEditorSubsystem::UMassEntityEditorSubsystem()
 	: EntityManager(MakeShareable(new FMassEntityManager(this)))
-	, PhaseManager(new FMassProcessingPhaseManager())
+	// we explicitly configure FMassProcessingPhaseManager to be in Editor mode since in this context we will not
+	// get a valid UWorld to deduce that - this subsystem runs for the editor itself, not for worlds loaded in the editor. 
+	, PhaseManager(new FMassProcessingPhaseManager(EProcessorExecutionFlags::Editor))
 {
 
 }
@@ -84,6 +86,26 @@ UMassEntityEditorSubsystem::~UMassEntityEditorSubsystem()
 }
 
 void UMassEntityEditorSubsystem::Initialize(FSubsystemCollectionBase& Collection)
+{
+	if (UMassEntitySettings::IsInitialized())
+	{
+		InitializeMassInternals();
+	}
+	else
+	{
+		OnSettingsInitializedHandle = UMassEntitySettings::GetOnInitializedEvent().AddUObject(this, &UMassEntityEditorSubsystem::InitializeMassInternals);
+	}
+
+	Super::Initialize(Collection);
+}
+
+void UMassEntityEditorSubsystem::Deinitialize()
+{
+	StopAndCleanUp();
+	Super::Deinitialize();
+}
+
+void UMassEntityEditorSubsystem::InitializeMassInternals()
 {
 	// Note that since this is a raw pointer we need to remember to delete it once we're done. It's now done in both 
 	// Deinitialize and the destructor (in case Deinitialize call is missed in an unexpected circumstances).
@@ -103,14 +125,6 @@ void UMassEntityEditorSubsystem::Initialize(FSubsystemCollectionBase& Collection
 	PhaseManager->Initialize(*this, ProcessingPhasesConfig);
 
 	PhaseManager->Start(EntityManager);
-
-	Super::Initialize(Collection);
-}
-
-void UMassEntityEditorSubsystem::Deinitialize()
-{
-	StopAndCleanUp();
-	Super::Deinitialize();
 }
 
 void UMassEntityEditorSubsystem::StopAndCleanUp()
@@ -122,6 +136,11 @@ void UMassEntityEditorSubsystem::StopAndCleanUp()
 	}
 
 	PhaseManager->Stop();
+
+	if (OnSettingsInitializedHandle.IsValid())
+	{
+		UMassEntitySettings::GetOnInitializedEvent().Remove(OnSettingsInitializedHandle);
+	}
 }
 
 void UMassEntityEditorSubsystem::Tick(float DeltaTime)
diff --git a/Engine/Plugins/Runtime/MassEntity/Source/MassEntityEditor/Public/MassEntityEditorSubsystem.h b/Engine/Plugins/Runtime/MassEntity/Source/MassEntityEditor/Public/MassEntityEditorSubsystem.h
--- a/Engine/Plugins/Runtime/MassEntity/Source/MassEntityEditor/Public/MassEntityEditorSubsystem.h
+++ b/Engine/Plugins/Runtime/MassEntity/Source/MassEntityEditor/Public/MassEntityEditorSubsystem.h
@@ -41,6 +41,8 @@ protected:
 	virtual void Initialize(FSubsystemCollectionBase& Collection) override;
 	virtual void Deinitialize() override;
 	// USubsystem interface end
+
+	void InitializeMassInternals();
 	void StopAndCleanUp();
 
 	friend UE::Mass::FMassEditorTickFunction;
@@ -60,4 +62,5 @@ protected:
 	FOnPostTick OnPostTickDelegate;
 	
 	bool IsProcessing{ false };
+	FDelegateHandle OnSettingsInitializedHandle;
 };
