[Mass] fixed typos

diff --git a/Engine/Plugins/Runtime/MassEntity/Source/MassEntity/Private/MassArchetypeData.cpp b/Engine/Plugins/Runtime/MassEntity/Source/MassEntity/Private/MassArchetypeData.cpp
--- a/Engine/Plugins/Runtime/MassEntity/Source/MassEntity/Private/MassArchetypeData.cpp
+++ b/Engine/Plugins/Runtime/MassEntity/Source/MassEntity/Private/MassArchetypeData.cpp
@@ -1004,14 +1004,14 @@ void FMassArchetypeData::BatchAddEntities(TConstArrayView<FMassEntityHandle> Ent
 	} while (NumberMoved < Entities.Num());
 }
 
-void FMassArchetypeData::BatchMoveEntitiesToAnotherArchetype(const FMassArchetypeEntityCollection& EntityCollection, FMassArchetypeData& NewArchetype, TArray<FMassEntityHandle>& OutEntitesBeingMoved, TArray<FMassArchetypeEntityCollection::FArchetypeEntityRange>* OutNewRanges)
+void FMassArchetypeData::BatchMoveEntitiesToAnotherArchetype(const FMassArchetypeEntityCollection& EntityCollection, FMassArchetypeData& NewArchetype, TArray<FMassEntityHandle>& OutEntitiesBeingMoved, TArray<FMassArchetypeEntityCollection::FArchetypeEntityRange>* OutNewRanges)
 {
 	check(&NewArchetype != this);
 
 
 	TArray<FMassArchetypeEntityCollection::FArchetypeEntityRange> Subchunks(EntityCollection.GetRanges());
 
-	const int32 InitialOutEntitiesCount = OutEntitesBeingMoved.Num();
+	const int32 InitialOutEntitiesCount = OutEntitiesBeingMoved.Num();
 
 	for (const FMassArchetypeEntityCollection::FArchetypeEntityRange EntityRange : Subchunks)
 	{
@@ -1028,8 +1028,7 @@ void FMassArchetypeData::BatchMoveEntitiesToAnotherArchetype(const FMassArchetyp
 
 		// gather entities we're about to remove
 		FMassEntityHandle* DyingEntityPtr = &Chunk.GetEntityArrayElementRef(EntityListOffsetWithinChunk, EntityRange.SubchunkStart);
-		const int32 EntitesBeingMovedStartIndex = OutEntitesBeingMoved.Num();
-		OutEntitesBeingMoved.Append(DyingEntityPtr, EntityRange.Length);
+		OutEntitiesBeingMoved.Append(DyingEntityPtr, EntityRange.Length);
 
 		FMassArchetypeEntityCollection::FArchetypeEntityRange ResultSubChunk;
 		ResultSubChunk.ChunkIndex = 0;
@@ -1068,7 +1067,7 @@ void FMassArchetypeData::BatchMoveEntitiesToAnotherArchetype(const FMassArchetyp
 
 	// Sorting the subchunks info so that subchunks of a given chunk are processed "from the back". Otherwise removing 
 	// a subchunk from the front of the chunk would inevitably invalidate following subchunks' information.
-	// Note that we do this after already having added the entities to the new archetype to preserve the order of entites 
+	// Note that we do this after already having added the entities to the new archetype to preserve the order of entities 
 	// as given by the input data.
 	Subchunks.Sort([](const FMassArchetypeEntityCollection::FArchetypeEntityRange& A, const FMassArchetypeEntityCollection::FArchetypeEntityRange& B)
 	{
@@ -1080,9 +1079,9 @@ void FMassArchetypeData::BatchMoveEntitiesToAnotherArchetype(const FMassArchetyp
 		BatchRemoveEntitiesInternal(Subchunk.ChunkIndex, Subchunk.SubchunkStart, Subchunk.Length);
 	}
 
-	for (int i = InitialOutEntitiesCount; i < OutEntitesBeingMoved.Num(); ++i)
+	for (int i = InitialOutEntitiesCount; i < OutEntitiesBeingMoved.Num(); ++i)
 	{
-		EntityMap.FindAndRemoveChecked(OutEntitesBeingMoved[i].Index);
+		EntityMap.FindAndRemoveChecked(OutEntitiesBeingMoved[i].Index);
 	}
 }
 
diff --git a/Engine/Plugins/Runtime/MassEntity/Source/MassEntity/Private/MassArchetypeData.h b/Engine/Plugins/Runtime/MassEntity/Source/MassEntity/Private/MassArchetypeData.h
--- a/Engine/Plugins/Runtime/MassEntity/Source/MassEntity/Private/MassArchetypeData.h
+++ b/Engine/Plugins/Runtime/MassEntity/Source/MassEntity/Private/MassArchetypeData.h
@@ -361,7 +361,7 @@ public:
 	// batched api
 	void BatchDestroyEntityChunks(FMassArchetypeEntityCollection::FConstEntityRangeArrayView EntityRangeContainer, TArray<FMassEntityHandle>& OutEntitiesRemoved);
 	void BatchAddEntities(TConstArrayView<FMassEntityHandle> Entities, const FMassArchetypeSharedFragmentValues& SharedFragmentValues, TArray<FMassArchetypeEntityCollection::FArchetypeEntityRange>& OutNewRanges);
-	void BatchMoveEntitiesToAnotherArchetype(const FMassArchetypeEntityCollection& EntityCollection, FMassArchetypeData& NewArchetype, TArray<FMassEntityHandle>& OutEntitesBeingMoved, TArray<FMassArchetypeEntityCollection::FArchetypeEntityRange>* OutNewChunks = nullptr);
+	void BatchMoveEntitiesToAnotherArchetype(const FMassArchetypeEntityCollection& EntityCollection, FMassArchetypeData& NewArchetype, TArray<FMassEntityHandle>& OutEntitiesBeingMoved, TArray<FMassArchetypeEntityCollection::FArchetypeEntityRange>* OutNewChunks = nullptr);
 	void BatchSetFragmentValues(TConstArrayView<FMassArchetypeEntityCollection::FArchetypeEntityRange> EntityCollection, const FMassGenericPayloadViewSlice& Payload);
 
 protected:
diff --git a/Engine/Plugins/Runtime/MassEntity/Source/MassEntityTestSuite/Private/MassQueryTest.cpp b/Engine/Plugins/Runtime/MassEntity/Source/MassEntityTestSuite/Private/MassQueryTest.cpp
--- a/Engine/Plugins/Runtime/MassEntity/Source/MassEntityTestSuite/Private/MassQueryTest.cpp
+++ b/Engine/Plugins/Runtime/MassEntity/Source/MassEntityTestSuite/Private/MassQueryTest.cpp
@@ -232,7 +232,7 @@ struct FQueryTest_ExecuteSparse : FEntityTestBase
 		for (FMassEntityHandle& Entity : EntitiesToIgnore)
 		{
 			const FTestFragment_Float& TestedFragment = EntityManager->GetFragmentDataChecked<FTestFragment_Float>(Entity);
-			AITEST_EQUAL("Untouched entites should retain default fragment value ", TestedFragment.Value, 0.f);
+			AITEST_EQUAL("Untouched entities should retain default fragment value ", TestedFragment.Value, 0.f);
 		}
 
 		return true;
diff --git a/Engine/Plugins/Runtime/MassGameplay/Source/MassActors/Private/MassAgentComponent.cpp b/Engine/Plugins/Runtime/MassGameplay/Source/MassActors/Private/MassAgentComponent.cpp
--- a/Engine/Plugins/Runtime/MassGameplay/Source/MassActors/Private/MassAgentComponent.cpp
+++ b/Engine/Plugins/Runtime/MassGameplay/Source/MassActors/Private/MassAgentComponent.cpp
@@ -227,7 +227,7 @@ void UMassAgentComponent::SetPuppetHandle(const FMassEntityHandle NewHandle)
 		TEXT("%s is expecting to be in state[EntityPendingCreation] but is in %s"), ANSI_TO_TCHAR(__FUNCTION__), *UEnum::GetValueAsString(State));
 	DebugCheckStateConsistency();
 
-	checkf(AgentHandle.IsValid() == false, TEXT("Can't set a new puppet handle of top of a regular agent handle. Entites would end up dangling."));
+	checkf(AgentHandle.IsValid() == false, TEXT("Can't set a new puppet handle of top of a regular agent handle. Entities would end up dangling."));
 	SetEntityHandleInternal(NewHandle);
 
 	if (UMassAgentSubsystem* AgentSubsystem = UWorld::GetSubsystem<UMassAgentSubsystem>(GetWorld()))
diff --git a/Engine/Plugins/Runtime/MassGameplay/Source/MassRepresentation/Private/MassRepresentationProcessor.cpp b/Engine/Plugins/Runtime/MassGameplay/Source/MassRepresentation/Private/MassRepresentationProcessor.cpp
--- a/Engine/Plugins/Runtime/MassGameplay/Source/MassRepresentation/Private/MassRepresentationProcessor.cpp
+++ b/Engine/Plugins/Runtime/MassGameplay/Source/MassRepresentation/Private/MassRepresentationProcessor.cpp
@@ -410,7 +410,7 @@ void UMassVisualizationProcessor::Execute(FMassEntityManager& InEntityManager, F
 		UpdateVisualization(Context);
 	});
 
-	UE_VLOG(this, LogMassRepresentation, Verbose, TEXT("UMassVisualizationProcessor::Execute processed %d entites"), TotalEntitiesProcessed);
+	UE_VLOG(this, LogMassRepresentation, Verbose, TEXT("UMassVisualizationProcessor::Execute processed %d entities"), TotalEntitiesProcessed);
 }
 
 //----------------------------------------------------------------------//
