[MassEntity]
- fixed iterations limit for first command flush
- added comment to help user to understand a specific compilation error and hint to how to fix it
- fixed typos

diff --git a/Engine/Plugins/Runtime/MassEntity/Source/MassEntity/Private/MassEntityManager.cpp b/Engine/Plugins/Runtime/MassEntity/Source/MassEntity/Private/MassEntityManager.cpp
--- a/Engine/Plugins/Runtime/MassEntity/Source/MassEntity/Private/MassEntityManager.cpp
+++ b/Engine/Plugins/Runtime/MassEntity/Source/MassEntity/Private/MassEntityManager.cpp
@@ -226,7 +226,7 @@ FMassArchetypeHandle FMassEntityManager::CreateArchetype(const FMassArchetypeCom
 #endif // WITH_MASSENTITY_DEBUG
 			if (CreationParams.ChunkMemorySize > 0 && CreationParams.ChunkMemorySize != Ptr->GetChunkAllocSize())
 			{
-				UE_LOG(LogMass, Warning, TEXT("Reusing existing Archetype, but the requested ChunkMemorySize is different. Requested %d, existing: %d")
+				UE_LOG(LogMass, Warning, TEXT("Reusing existing Archetype, but the requested ChunkMemorySize is different. Requested %d, existing: %llu")
 					, CreationParams.ChunkMemorySize, Ptr->GetChunkAllocSize());
 			}
 			ArchetypeDataPtr = Ptr;
@@ -1022,11 +1022,11 @@ void FMassEntityManager::BatchChangeTagsForEntities(TConstArrayView<FMassArchety
 
 			// Move the entity over
 			FMassArchetypeEntityCollection::FEntityRangeArray NewArchetypeEntityRanges;
-			TArray<FMassEntityHandle> EntitesBeingMoved;
-			CurrentArchetype->BatchMoveEntitiesToAnotherArchetype(Collection, *NewArchetypeHandle.DataPtr.Get(), EntitesBeingMoved
+			TArray<FMassEntityHandle> EntitiesBeingMoved;
+			CurrentArchetype->BatchMoveEntitiesToAnotherArchetype(Collection, *NewArchetypeHandle.DataPtr.Get(), EntitiesBeingMoved
 				, bTagsAddedAreObserved ? &NewArchetypeEntityRanges : nullptr);
 
-			for (const FMassEntityHandle& Entity : EntitesBeingMoved)
+			for (const FMassEntityHandle& Entity : EntitiesBeingMoved)
 			{
 				check(Entities.IsValidIndex(Entity.Index));
 
@@ -1073,11 +1073,11 @@ void FMassEntityManager::BatchChangeFragmentCompositionForEntities(TConstArrayVi
 
 				// Move the entity over
 				FMassArchetypeEntityCollection::FEntityRangeArray NewArchetypeEntityRanges;
-				TArray<FMassEntityHandle> EntitesBeingMoved;
-				CurrentArchetype->BatchMoveEntitiesToAnotherArchetype(Collection, *NewArchetypeHandle.DataPtr.Get(), EntitesBeingMoved
+				TArray<FMassEntityHandle> EntitiesBeingMoved;
+				CurrentArchetype->BatchMoveEntitiesToAnotherArchetype(Collection, *NewArchetypeHandle.DataPtr.Get(), EntitiesBeingMoved
 					, bFragmentsAddedAreObserved ? &NewArchetypeEntityRanges : nullptr);
 
-				for (const FMassEntityHandle& Entity : EntitesBeingMoved)
+				for (const FMassEntityHandle& Entity : EntitiesBeingMoved)
 				{
 					check(Entities.IsValidIndex(Entity.Index));
 
@@ -1139,11 +1139,11 @@ void FMassEntityManager::BatchAddFragmentInstancesForEntities(TConstArrayView<FM
 				checkSlow(NewArchetypeHandle.IsValid());
 
 				// Move the entity over
-				TArray<FMassEntityHandle> EntitesBeingMoved;
+				TArray<FMassEntityHandle> EntitiesBeingMoved;
 				CurrentArchetype->BatchMoveEntitiesToAnotherArchetype(EntityRangesWithPayload.GetEntityCollection(), *NewArchetypeHandle.DataPtr.Get()
-					, EntitesBeingMoved, &TargetArchetypeEntityRanges);
+					, EntitiesBeingMoved, &TargetArchetypeEntityRanges);
 
-				for (const FMassEntityHandle& Entity : EntitesBeingMoved)
+				for (const FMassEntityHandle& Entity : EntitiesBeingMoved)
 				{
 					check(Entities.IsValidIndex(Entity.Index));
 
@@ -1506,14 +1506,14 @@ void FMassEntityManager::FlushCommands(const TSharedPtr<FMassCommandBuffer>& InC
 		int32 IterationsCounter = 0;
 		TOptional<TSharedPtr<FMassCommandBuffer>> CurrentCommandBuffer = FlushedCommandBufferQueue.Dequeue();
 
-		while (IterationsCounter < MaxIterations && CurrentCommandBuffer.IsSet())
+		while (IterationsCounter < IterationsLimit && CurrentCommandBuffer.IsSet())
 		{
 			IterationsCounter++;
 			(*CurrentCommandBuffer)->Flush(*this);
 			CurrentCommandBuffer = FlushedCommandBufferQueue.Dequeue();
 		}
-		ensure(IterationsCounter >= MaxIterations || CurrentCommandBuffer.IsSet() == false);
-		UE_CVLOG_UELOG(IterationsCounter >= MaxIterations, GetOwner(), LogMass, Error, TEXT("Reached loop count limit while flushing commands"));
+		ensure(IterationsCounter >= IterationsLimit || CurrentCommandBuffer.IsSet() == false);
+		UE_CVLOG_UELOG(IterationsCounter >= IterationsLimit, GetOwner(), LogMass, Error, TEXT("Reached loop count limit while flushing commands"));
 
 		bCommandBufferFlushingInProgress = false;
 	}
diff --git a/Engine/Plugins/Runtime/MassEntity/Source/MassEntity/Public/MassRequirements.h b/Engine/Plugins/Runtime/MassEntity/Source/MassEntity/Public/MassRequirements.h
--- a/Engine/Plugins/Runtime/MassEntity/Source/MassEntity/Public/MassRequirements.h
+++ b/Engine/Plugins/Runtime/MassEntity/Source/MassEntity/Public/MassRequirements.h
@@ -86,12 +86,15 @@ struct MASSENTITY_API FMassSubsystemRequirements
 	friend FMassDebugger;
 	friend FMassRequirementAccessDetector;
 
-public:
 	template<typename T>
 	FMassSubsystemRequirements& AddSubsystemRequirement(const EMassFragmentAccess AccessMode)
 	{
 		check(AccessMode != EMassFragmentAccess::None && AccessMode != EMassFragmentAccess::MAX);
 
+		// Compilation errors here like: 'GameThreadOnly': is not a member of 'TMassExternalSubsystemTraits<USmartObjectSubsystem>
+		// indicate that there is a missing header that defines the subsystem's trait or that you need to define one for that subsystem type.
+		// @see "MassExternalSubsystemTraits.h" for details
+
 		switch (AccessMode)
 		{
 		case EMassFragmentAccess::ReadOnly:
