[MassEntity] moved observer notification to the public API to make sure fragment data is applied before notifying observers when adding new fragments instances (batched version was fine).

diff --git a/Engine/Plugins/Runtime/MassEntity/Source/MassEntity/Private/MassEntityManager.cpp b/Engine/Plugins/Runtime/MassEntity/Source/MassEntity/Private/MassEntityManager.cpp
--- a/Engine/Plugins/Runtime/MassEntity/Source/MassEntity/Private/MassEntityManager.cpp
+++ b/Engine/Plugins/Runtime/MassEntity/Source/MassEntity/Private/MassEntityManager.cpp
@@ -714,14 +714,16 @@ void FMassEntityManager::AddFragmentToEntity(FMassEntityHandle Entity, const USc
 
 	CheckIfEntityIsActive(Entity);
 
-	InternalAddFragmentListToEntityChecked(Entity, FMassFragmentBitSet(*FragmentType));
+	const FMassArchetypeCompositionDescriptor Descriptor(InternalAddFragmentListToEntityChecked(Entity, FMassFragmentBitSet(*FragmentType)));
+	ObserverManager.OnPostCompositionAdded(Entity, Descriptor);
 }
 
 void FMassEntityManager::AddFragmentListToEntity(FMassEntityHandle Entity, TConstArrayView<const UScriptStruct*> FragmentList)
 {
 	CheckIfEntityIsActive(Entity);
 
-	InternalAddFragmentListToEntityChecked(Entity, FMassFragmentBitSet(FragmentList));
+	const FMassArchetypeCompositionDescriptor Descriptor(InternalAddFragmentListToEntityChecked(Entity, FMassFragmentBitSet(FragmentList)));
+	ObserverManager.OnPostCompositionAdded(Entity, Descriptor);
 }
 
 void FMassEntityManager::AddCompositionToEntity_GetDelta(FMassEntityHandle Entity, FMassArchetypeCompositionDescriptor& InDescriptor)
@@ -816,7 +818,7 @@ void FMassEntityManager::InternalReleaseEntity(FMassEntityHandle Entity)
 	EntityFreeIndexList.Add(Entity.Index);
 }
 
-void FMassEntityManager::InternalAddFragmentListToEntityChecked(FMassEntityHandle Entity, const FMassFragmentBitSet& InFragments)
+FMassFragmentBitSet FMassEntityManager::InternalAddFragmentListToEntityChecked(FMassEntityHandle Entity, const FMassFragmentBitSet& InFragments)
 {
 	const FEntityData& EntityData = Entities[Entity.Index];
 	FMassArchetypeData* OldArchetype = EntityData.CurrentArchetype.Get();
@@ -826,11 +828,12 @@ void FMassEntityManager::InternalAddFragmentListToEntityChecked(FMassEntityHandl
 		, TEXT("Trying to add a new fragment type to an entity, but it already has some of them. (%s)")
 		, *InFragments.GetOverlap(OldArchetype->GetFragmentBitSet()).DebugGetStringDesc());
 
-	const FMassFragmentBitSet NewFragments = InFragments - OldArchetype->GetFragmentBitSet();
+	FMassFragmentBitSet NewFragments = InFragments - OldArchetype->GetFragmentBitSet();
 	if (NewFragments.IsEmpty() == false)
 	{
 		InternalAddFragmentListToEntity(Entity, NewFragments);
 	}
+	return MoveTemp(NewFragments);
 }
 
 void FMassEntityManager::InternalAddFragmentListToEntity(FMassEntityHandle Entity, const FMassFragmentBitSet& InFragments)
@@ -843,20 +846,13 @@ void FMassEntityManager::InternalAddFragmentListToEntity(FMassEntityHandle Entit
 
 	// fetch or create the new archetype
 	const FMassArchetypeHandle NewArchetypeHandle = CreateArchetype(EntityData.CurrentArchetype, InFragments);
+	checkf(NewArchetypeHandle.DataPtr != EntityData.CurrentArchetype, TEXT("%hs is intended for internal calls with non overlapping fragment list."), __FUNCTION__);
 
-	if (NewArchetypeHandle.DataPtr != EntityData.CurrentArchetype)
-	{
-		// Move the entity over
-		FMassArchetypeData& NewArchetype = FMassArchetypeHelper::ArchetypeDataFromHandleChecked(NewArchetypeHandle);
-		NewArchetype.CopyDebugNamesFrom(*OldArchetype);
-		EntityData.CurrentArchetype->MoveEntityToAnotherArchetype(Entity, NewArchetype);
-		EntityData.CurrentArchetype = NewArchetypeHandle.DataPtr;
-
-		FMassArchetypeCompositionDescriptor CompositionChangeDescriptor;
-		CompositionChangeDescriptor.Fragments += InFragments;
-
-		ObserverManager.OnPostCompositionAdded(Entity, CompositionChangeDescriptor);
-	}
+	// Move the entity over
+	FMassArchetypeData& NewArchetype = FMassArchetypeHelper::ArchetypeDataFromHandleChecked(NewArchetypeHandle);
+	NewArchetype.CopyDebugNamesFrom(*OldArchetype);
+	EntityData.CurrentArchetype->MoveEntityToAnotherArchetype(Entity, NewArchetype);
+	EntityData.CurrentArchetype = NewArchetypeHandle.DataPtr;
 }
 
 void FMassEntityManager::AddFragmentInstanceListToEntity(FMassEntityHandle Entity, TConstArrayView<FInstancedStruct> FragmentInstanceList)
@@ -866,10 +862,12 @@ void FMassEntityManager::AddFragmentInstanceListToEntity(FMassEntityHandle Entit
 	CheckIfEntityIsActive(Entity);
 	checkf(FragmentInstanceList.Num() > 0, TEXT("Need to specify at least one fragment instances for this operation"));
 
-	InternalAddFragmentListToEntityChecked(Entity, FMassFragmentBitSet(FragmentInstanceList));
+	const FMassArchetypeCompositionDescriptor Descriptor(InternalAddFragmentListToEntityChecked(Entity, FMassFragmentBitSet(FragmentInstanceList)));
 
 	const FEntityData& EntityData = Entities[Entity.Index];
 	EntityData.CurrentArchetype->SetFragmentsData(Entity, FragmentInstanceList);
+
+	ObserverManager.OnPostCompositionAdded(Entity, Descriptor);
 }
 
 void FMassEntityManager::RemoveFragmentFromEntity(FMassEntityHandle Entity, const UScriptStruct* FragmentType)
diff --git a/Engine/Plugins/Runtime/MassEntity/Source/MassEntity/Public/MassEntityManager.h b/Engine/Plugins/Runtime/MassEntity/Source/MassEntity/Public/MassEntityManager.h
--- a/Engine/Plugins/Runtime/MassEntity/Source/MassEntity/Public/MassEntityManager.h
+++ b/Engine/Plugins/Runtime/MassEntity/Source/MassEntity/Public/MassEntityManager.h
@@ -544,8 +544,9 @@ private:
 
 	/** 
 	 *  Adds fragments in FragmentList to Entity. Only the unique fragments will be added.
+	 *  @return Bitset for the added fragments (might be empty or a subset of `InFragments` depending on the current archetype fragments)
 	 */
-	void InternalAddFragmentListToEntityChecked(FMassEntityHandle Entity, const FMassFragmentBitSet& InFragments);
+	FMassFragmentBitSet InternalAddFragmentListToEntityChecked(FMassEntityHandle Entity, const FMassFragmentBitSet& InFragments);
 
 	/** 
 	 *  Similar to InternalAddFragmentListToEntity but expects NewFragmentList not overlapping with current entity's
