Changes to dynamic processor dependency solving to support dependencies between instances of the same processor class

diff --git a/Engine/Plugins/Runtime/MassEntity/Source/MassEntity/Private/MassProcessorDependencySolver.cpp b/Engine/Plugins/Runtime/MassEntity/Source/MassEntity/Private/MassProcessorDependencySolver.cpp
--- a/Engine/Plugins/Runtime/MassEntity/Source/MassEntity/Private/MassProcessorDependencySolver.cpp
+++ b/Engine/Plugins/Runtime/MassEntity/Source/MassEntity/Private/MassProcessorDependencySolver.cpp
@@ -357,25 +357,25 @@ void FMassProcessorDependencySolver::CreateSubGroupNames(FName InGroupName, TArr
 int32 FMassProcessorDependencySolver::CreateNodes(UMassProcessor& Processor)
 {
 	check(Processor.GetClass());
-	FName ProcName = Processor.GetClass()->GetFName();
+	// for processors supporting multiple instances we use processor name rather than processor's class name for
+	// dependency calculations. This makes the user responsible for fine-tuning per-processor dependencies. 
+	const FName ProcName = Processor.ShouldAllowMultipleInstances() 
+		? Processor.GetFName()
+		: Processor.GetClass()->GetFName();
 
 	if (const int32* NodeIndexPtr = NodeIndexMap.Find(ProcName))
 	{
-		// we can accept another instance of this processor class if the processor itself supports that.
-		// Note that the first instance is added with the class while the subsequent instances are added with 
-		// the instance name. This is done on purpose. The class name is used as dependency, so if at least the first 
-		// processor is placed with the class name, like other processors, it can still be used to influence execution
-		// order via ExecuteBefore and ExecuteAfter.
 		if (Processor.ShouldAllowMultipleInstances())
 		{
-			ProcName = Processor.GetFName();
+			UE_LOG(LogMass, Warning, TEXT("%hs Processor %s, name %s, already registered. This processor class does suport duplicates, but individual instances need to have a unique name.")
+				, __FUNCTION__, *Processor.GetFullName(), *ProcName.ToString());
 		}
 		else
 		{
-			UE_LOG(LogMass, Warning, TEXT("%s Processor %s already registered. Duplicates are not supported.")
-				, ANSI_TO_TCHAR(__FUNCTION__), *ProcName.ToString());
-			return *NodeIndexPtr;
+			UE_LOG(LogMass, Warning, TEXT("%hs Processor %s already registered. Duplicates are not supported by this processor class.")
+				, __FUNCTION__, *ProcName.ToString());
 		}
+		return *NodeIndexPtr;
 	}
 
 	const FMassProcessorExecutionOrder& ExecutionOrder = Processor.GetExecutionOrder();
