Added a function to FMassEntityManager for streamlining external command buffers appending.

diff --git a/Engine/Plugins/Runtime/MassEntity/Source/MassEntity/Private/MassEntityManager.cpp b/Engine/Plugins/Runtime/MassEntity/Source/MassEntity/Private/MassEntityManager.cpp
--- a/Engine/Plugins/Runtime/MassEntity/Source/MassEntity/Private/MassEntityManager.cpp
+++ b/Engine/Plugins/Runtime/MassEntity/Source/MassEntity/Private/MassEntityManager.cpp
@@ -1517,6 +1517,24 @@ void FMassEntityManager::FlushCommands(const TSharedPtr<FMassCommandBuffer>& InC
 	}
 }
 
+void FMassEntityManager::AppendCommands(TSharedPtr<FMassCommandBuffer>& InOutCommandBuffer)
+{
+	if (!ensureMsgf(InOutCommandBuffer != DeferredCommandBuffer, TEXT("We don't expect AppendCommands to be called with EntityManager's command buffer as the input parameter")))
+	{
+		return;
+	}
+	else if (DeferredCommandBuffer->IsFlushing())
+	{
+		// in this case we'll add InOutCommandBuffer to FlushedCommandBufferQueue
+		FlushCommands(InOutCommandBuffer);
+	}
+	else
+	{
+		// otherwise we just move all the commands out of InOutCommandBuffer and into the main buffer
+		DeferredCommandBuffer->MoveAppend(*InOutCommandBuffer.Get());
+	}
+}
+
 void FMassEntityManager::SetDebugName(const FString& NewDebugGame) 
 { 
 #if WITH_MASSENTITY_DEBUG
diff --git a/Engine/Plugins/Runtime/MassEntity/Source/MassEntity/Private/MassExecutor.cpp b/Engine/Plugins/Runtime/MassEntity/Source/MassEntity/Private/MassExecutor.cpp
--- a/Engine/Plugins/Runtime/MassEntity/Source/MassEntity/Private/MassExecutor.cpp
+++ b/Engine/Plugins/Runtime/MassEntity/Source/MassEntity/Private/MassExecutor.cpp
@@ -133,7 +133,7 @@ void RunProcessorsView(TArrayView<UMassProcessor* const> Processors, FMassProces
 		}
 		else if (CommandBuffer.Get() != &EntityManager.Defer())
 		{
-			EntityManager.Defer().MoveAppend(*CommandBuffer.Get());
+			EntityManager.AppendCommands(CommandBuffer);
 		}
 	}
 }
diff --git a/Engine/Plugins/Runtime/MassEntity/Source/MassEntity/Public/MassEntityManager.h b/Engine/Plugins/Runtime/MassEntity/Source/MassEntity/Public/MassEntityManager.h
--- a/Engine/Plugins/Runtime/MassEntity/Source/MassEntity/Public/MassEntityManager.h
+++ b/Engine/Plugins/Runtime/MassEntity/Source/MassEntity/Public/MassEntityManager.h
@@ -427,6 +427,13 @@ public:
 	 */
 	void FlushCommands(const TSharedPtr<FMassCommandBuffer>& InCommandBuffer = TSharedPtr<FMassCommandBuffer>());
 
+	/** 
+	 * Depending on the current state of Manager's command buffer the function will either move all the commands out of 
+	 * InOutCommandBuffer into the main command buffer or append it to the list of command buffers waiting to be flushed.
+	 * @note as a consequence of the call InOutCommandBuffer can get its contents emptied due some of the undelying code using Move semantics
+	 */
+	void AppendCommands(TSharedPtr<FMassCommandBuffer>& InOutCommandBuffer);
+
 	/**
 	 * Shared fragment creation methods
 	 */
