Added another macro for local Mass entity debugging

diff --git a/Engine/Plugins/Runtime/MassEntity/Source/MassEntity/Public/MassDebugger.h b/Engine/Plugins/Runtime/MassEntity/Source/MassEntity/Public/MassDebugger.h
--- a/Engine/Plugins/Runtime/MassEntity/Source/MassEntity/Public/MassDebugger.h
+++ b/Engine/Plugins/Runtime/MassEntity/Source/MassEntity/Public/MassDebugger.h
@@ -50,6 +50,7 @@ namespace UE::Mass::Debug
 
 #define MASS_IF_ENTITY_DEBUGGED(Manager, EntityHandle) (FMassDebugger::GetSelectedEntity(Manager) == EntityHandle)
 #define MASS_BREAK_IF_ENTITY_DEBUGGED(Manager, EntityHandle) { if (UE::Mass::Debug::bAllowBreakOnDebuggedEntity && MASS_IF_ENTITY_DEBUGGED(Manager, EntityHandle)) { PLATFORM_BREAK();} }
+#define MASS_BREAK_IF_ENTITY_INDEX(EntityHandle, InIndex) { if (UE::Mass::Debug::bAllowBreakOnDebuggedEntity && EntityHandle.Index == InIndex) { PLATFORM_BREAK();} }
 #define MASS_SET_ENTITY_DEBUGGED(Manager, EntityHandle) { if (UE::Mass::Debug::bAllowProceduralDebuggedEntitySelection) {FMassDebugger::SelectEntity(Manager, EntityHandle); }}
 
 namespace UE::Mass::Debug
@@ -158,6 +159,7 @@ struct MASSENTITY_API FMassDebugger
 
 #define MASS_IF_ENTITY_DEBUGGED(a, b) false
 #define MASS_BREAK_IF_ENTITY_DEBUGGED(a, b)
+#define MASS_BREAK_IF_ENTITY_INDEX(a, b)
 #define MASS_SET_ENTITY_DEBUGGED(a, b)
 
 #endif // WITH_MASSENTITY_DEBUG
