Improved FMassExecutionContext's checks whether subsystems and fragments are bound and configured for read-write if accessing mutable version.

diff --git a/Engine/Plugins/Runtime/MassEntity/Source/MassEntity/Public/MassExecutionContext.h b/Engine/Plugins/Runtime/MassEntity/Source/MassEntity/Public/MassExecutionContext.h
--- a/Engine/Plugins/Runtime/MassEntity/Source/MassEntity/Public/MassExecutionContext.h
+++ b/Engine/Plugins/Runtime/MassEntity/Source/MassEntity/Public/MassExecutionContext.h
@@ -8,6 +8,16 @@
 #include "MassSubsystemAccess.h"
 
 
+#define CHECK_IF_BOUND(View, Type) \
+	checkf(View \
+		, TEXT("Requested fragment type not bound, type %s. Make sure it has been listed as required."), *GetNameSafe(Type))
+
+#define CHECK_IF_READWRITE(View) \
+	checkf(View && View->Requirement.AccessMode == EMassFragmentAccess::ReadWrite \
+		, TEXT("Requested fragment type not bound for writing, type %s. Make sure it has been listed as required in ReadWrite mode.") \
+		, View ? TEXT("[Not found]") : *GetNameSafe(View->Requirement.StructType))
+
+
 struct MASSENTITY_API FMassExecutionContext
 {
 private:
@@ -154,6 +164,7 @@ public:
 
 		const UScriptStruct* Type = T::StaticStruct();
 		FChunkFragmentView* FoundChunkFragmentData = ChunkFragmentViews.FindByPredicate([Type](const FChunkFragmentView& Element) { return Element.Requirement.StructType == Type; } );
+		CHECK_IF_READWRITE(FoundChunkFragmentData);
 		return FoundChunkFragmentData ? FoundChunkFragmentData->FragmentView.GetPtr<T>() : static_cast<T*>(nullptr);
 	}
 	
@@ -161,7 +172,7 @@ public:
 	T& GetMutableChunkFragment()
 	{
 		T* ChunkFragment = GetMutableChunkFragmentPtr<T>();
-		checkf(ChunkFragment, TEXT("Chunk Fragment requirement not found: %s"), *T::StaticStruct()->GetName());
+		CHECK_IF_BOUND(ChunkFragment, T::StaticStruct());
 		return *ChunkFragment;
 	}
 
@@ -179,7 +190,7 @@ public:
 	const T& GetChunkFragment() const
 	{
 		const T* ChunkFragment = GetChunkFragmentPtr<T>();
-		checkf(ChunkFragment, TEXT("Chunk Fragment requirement not found: %s"), *T::StaticStruct()->GetName());
+		CHECK_IF_BOUND(ChunkFragment, T::StaticStruct());
 		return *ChunkFragment;
 	}
 
@@ -197,17 +208,17 @@ public:
 	const T& GetConstSharedFragment() const
 	{
 		const T* SharedFragment = GetConstSharedFragmentPtr<const T>();
-		checkf(SharedFragment, TEXT("Shared Fragment requirement not found: %s"), *T::StaticStruct()->GetName());
+		CHECK_IF_BOUND(SharedFragment, T::StaticStruct());
 		return *SharedFragment;
 	}
 
-
 	template<typename T>
 	T* GetMutableSharedFragmentPtr()
 	{
 		static_assert(TIsDerivedFrom<T, FMassSharedFragment>::IsDerived, "Given struct doesn't represent a valid shared fragment type. Make sure to inherit from FMassSharedFragment or one of its child-types.");
 
 		FSharedFragmentView* FoundSharedFragmentData = SharedFragmentViews.FindByPredicate([](const FSharedFragmentView& Element) { return Element.Requirement.StructType == T::StaticStruct(); });
+		CHECK_IF_READWRITE(FoundSharedFragmentData);
 		return FoundSharedFragmentData ? FoundSharedFragmentData->FragmentView.GetPtr<T>() : static_cast<T*>(nullptr);
 	}
 
@@ -224,7 +235,7 @@ public:
 	T& GetMutableSharedFragment()
 	{
 		T* SharedFragment = GetMutableSharedFragmentPtr<T>();
-		checkf(SharedFragment, TEXT("Shared Fragment requirement not found: %s"), *T::StaticStruct()->GetName());
+		CHECK_IF_BOUND(SharedFragment, T::StaticStruct());
 		return *SharedFragment;
 	}
 
@@ -232,7 +243,7 @@ public:
 	const T& GetSharedFragment() const
 	{
 		const T* SharedFragment = GetSharedFragmentPtr<T>();
-		checkf(SharedFragment, TEXT("Shared Fragment requirement not found: %s"), *T::StaticStruct()->GetName());
+		CHECK_IF_BOUND(SharedFragment, T::StaticStruct());
 		return *SharedFragment;
 	}
 
@@ -242,8 +253,8 @@ public:
 	{
 		const UScriptStruct* FragmentType = TFragment::StaticStruct();
 		const FFragmentView* View = FragmentViews.FindByPredicate([FragmentType](const FFragmentView& Element) { return Element.Requirement.StructType == FragmentType; });
-		checkfSlow(View != nullptr, TEXT("Requested fragment type not bound"));
-		checkfSlow(View->Requirement.AccessMode == EMassFragmentAccess::ReadWrite, TEXT("Requested fragment has not been bound for writing"));
+		CHECK_IF_BOUND(View, FragmentType);
+		CHECK_IF_READWRITE(View);
 		return MakeArrayView<TFragment>((TFragment*)View->FragmentView.GetData(), View->FragmentView.Num());
 	}
 
@@ -252,21 +263,22 @@ public:
 	{
 		const UScriptStruct* FragmentType = TFragment::StaticStruct();
 		const FFragmentView* View = FragmentViews.FindByPredicate([FragmentType](const FFragmentView& Element) { return Element.Requirement.StructType == FragmentType; });
-		checkfSlow(View != nullptr, TEXT("Requested fragment type not bound. Make sure the Frament was requested in ReadOnly or ReadWrite mode."));
+		CHECK_IF_BOUND(View, TFragment::StaticStruct());
 		return TConstArrayView<TFragment>((const TFragment*)View->FragmentView.GetData(), View->FragmentView.Num());
 	}
 
 	TConstArrayView<FMassFragment> GetFragmentView(const UScriptStruct* FragmentType) const
 	{
 		const FFragmentView* View = FragmentViews.FindByPredicate([FragmentType](const FFragmentView& Element) { return Element.Requirement.StructType == FragmentType; });
-		checkSlow(View);
+		CHECK_IF_BOUND(View, FragmentType);
 		return TConstArrayView<FMassFragment>((const FMassFragment*)View->FragmentView.GetData(), View->FragmentView.Num());;
 	}
 
 	TArrayView<FMassFragment> GetMutableFragmentView(const UScriptStruct* FragmentType) 
 	{
 		const FFragmentView* View = FragmentViews.FindByPredicate([FragmentType](const FFragmentView& Element) { return Element.Requirement.StructType == FragmentType; });
-		checkSlow(View);
+		CHECK_IF_BOUND(View, FragmentType);
+		CHECK_IF_READWRITE(View);
 		return View->FragmentView;
 	}
 
@@ -447,3 +459,6 @@ public:
 	UE_DEPRECATED(5.4, "Deprecated in favor of 'SetCurrentArchetypeCompositionDescriptor' as this provides information on the entire archetype.")
 	void SetCurrentArchetypesTagBitSet(const FMassTagBitSet&) {}
 };
+
+#undef CHECK_IF_BOUND
+#undef CHECK_IF_READWRITE
