Added MassEntityQuery.ParallelForEachEntityChunk along with basic unit tests.

diff --git a/Engine/Plugins/Runtime/MassEntity/Source/MassEntity/Private/MassArchetypeData.cpp b/Engine/Plugins/Runtime/MassEntity/Source/MassEntity/Private/MassArchetypeData.cpp
--- a/Engine/Plugins/Runtime/MassEntity/Source/MassEntity/Private/MassArchetypeData.cpp
+++ b/Engine/Plugins/Runtime/MassEntity/Source/MassEntity/Private/MassArchetypeData.cpp
@@ -493,7 +493,7 @@ void FMassArchetypeData::ExecuteFunction(FMassExecutionContext& RunContext, cons
 	}
 }
 
-void FMassArchetypeData::ExecutionFunctionForChunk(FMassExecutionContext RunContext, const FMassExecuteFunction& Function, const FMassQueryRequirementIndicesMapping& RequirementMapping, const FMassArchetypeEntityCollection::FArchetypeEntityRange& EntityRange, const FMassChunkConditionFunction& ChunkCondition)
+void FMassArchetypeData::ExecutionFunctionForChunk(FMassExecutionContext& RunContext, const FMassExecuteFunction& Function, const FMassQueryRequirementIndicesMapping& RequirementMapping, const FMassArchetypeEntityCollection::FArchetypeEntityRange& EntityRange, const FMassChunkConditionFunction& ChunkCondition)
 {
 	FMassArchetypeChunk& Chunk = Chunks[EntityRange.ChunkIndex];
 	const int32 ChunkLength = EntityRange.Length > 0 ? EntityRange.Length : (Chunk.GetNumInstances() - EntityRange.SubchunkStart);
diff --git a/Engine/Plugins/Runtime/MassEntity/Source/MassEntity/Private/MassArchetypeData.h b/Engine/Plugins/Runtime/MassEntity/Source/MassEntity/Private/MassArchetypeData.h
--- a/Engine/Plugins/Runtime/MassEntity/Source/MassEntity/Private/MassArchetypeData.h
+++ b/Engine/Plugins/Runtime/MassEntity/Source/MassEntity/Private/MassArchetypeData.h
@@ -249,10 +249,13 @@ public:
 
 	uint32 GetCreatedArchetypeDataVersion() const { return CreatedArchetypeDataVersion; }
 
-	void ExecuteFunction(FMassExecutionContext& RunContext, const FMassExecuteFunction& Function, const FMassQueryRequirementIndicesMapping& RequirementMapping, FMassArchetypeEntityCollection::FConstEntityRangeArrayView EntityRangeContainer, const FMassChunkConditionFunction& ChunkCondition);
-	void ExecuteFunction(FMassExecutionContext& RunContext, const FMassExecuteFunction& Function, const FMassQueryRequirementIndicesMapping& RequirementMapping, const FMassChunkConditionFunction& ChunkCondition);
+	void ExecuteFunction(FMassExecutionContext& RunContext, const FMassExecuteFunction& Function, const FMassQueryRequirementIndicesMapping& RequirementMapping
+		, FMassArchetypeEntityCollection::FConstEntityRangeArrayView EntityRangeContainer, const FMassChunkConditionFunction& ChunkCondition);
+	void ExecuteFunction(FMassExecutionContext& RunContext, const FMassExecuteFunction& Function, const FMassQueryRequirementIndicesMapping& RequirementMapping
+		, const FMassChunkConditionFunction& ChunkCondition);
 
-	void ExecutionFunctionForChunk(FMassExecutionContext RunContext, const FMassExecuteFunction& Function, const FMassQueryRequirementIndicesMapping& RequirementMapping, const FMassArchetypeEntityCollection::FArchetypeEntityRange& EntityRange, const FMassChunkConditionFunction& ChunkCondition = FMassChunkConditionFunction());
+	void ExecutionFunctionForChunk(FMassExecutionContext& RunContext, const FMassExecuteFunction& Function, const FMassQueryRequirementIndicesMapping& RequirementMapping
+		, const FMassArchetypeEntityCollection::FArchetypeEntityRange& EntityRange, const FMassChunkConditionFunction& ChunkCondition = FMassChunkConditionFunction());
 
 	/**
 	 * Compacts entities to fill up chunks as much as possible
diff --git a/Engine/Plugins/Runtime/MassEntity/Source/MassEntity/Private/MassEntityQuery.cpp b/Engine/Plugins/Runtime/MassEntity/Source/MassEntity/Private/MassEntityQuery.cpp
--- a/Engine/Plugins/Runtime/MassEntity/Source/MassEntity/Private/MassEntityQuery.cpp
+++ b/Engine/Plugins/Runtime/MassEntity/Source/MassEntity/Private/MassEntityQuery.cpp
@@ -19,16 +19,41 @@
 #endif // WITH_MASSENTITY_DEBUG
 
 
-//////////////////////////////////////////////////////////////////////
-// FMassEntityQuery
+namespace UE::Mass::Tweakables
+{
+	/**
+	 * Controls whether ParallelForEachEntityChunk actually performs ParallelFor operations. If `false` the call is passed
+	 * the the regular ForEachEntityChunk call.
+	 */
+	bool bAllowParallelExecution = true;
 
-FMassEntityQuery::FMassEntityQuery()
+	namespace
+	{
+		static FAutoConsoleVariableRef AnonymousCVars[] = {
+			{	TEXT("mass.AllowQueryParallelFor"), bAllowParallelExecution, TEXT("Controls whether EntityQueries are allowed to utilize ParallelFor construct"), ECVF_Cheat }
+		};
+	}
+}
+
+//-----------------------------------------------------------------------------
+// FScopedSubsystemRequirementsRestore
+//-----------------------------------------------------------------------------
+FMassEntityQuery::FScopedSubsystemRequirementsRestore::FScopedSubsystemRequirementsRestore(FMassExecutionContext& ExecutionContext)
+	: CachedExecutionContext(ExecutionContext)
 {
-	bAllowParallelExecution = false;
-	bRequiresGameThreadExecution = false;
-	bRequiresMutatingWorldAccess = false;
+	CachedExecutionContext.GetSubsystemRequirementBits(ConstSubsystemsBitSet, MutableSubsystemsBitSet);
+}
 
-	ReadCommandlineParams();
+FMassEntityQuery::FScopedSubsystemRequirementsRestore::~FScopedSubsystemRequirementsRestore()
+{
+	CachedExecutionContext.SetSubsystemRequirementBits(ConstSubsystemsBitSet, MutableSubsystemsBitSet);
+}
+
+//-----------------------------------------------------------------------------
+// FMassEntityQuery
+//-----------------------------------------------------------------------------
+FMassEntityQuery::FMassEntityQuery()
+{
 }
 
 FMassEntityQuery::FMassEntityQuery(std::initializer_list<UScriptStruct*> InitList)
@@ -63,15 +88,6 @@ void FMassEntityQuery::RegisterWithProcessor(UMassProcessor& Owner)
 #endif // WITH_MASSENTITY_DEBUG
 }
 
-void FMassEntityQuery::ReadCommandlineParams()
-{
-	int AllowParallelQueries = -1;
-	if (FParse::Value(FCommandLine::Get(), TEXT("ParallelMassQueries="), AllowParallelQueries))
-	{
-		bAllowParallelExecution = (AllowParallelQueries != 0);
-	}
-}
-
 void FMassEntityQuery::CacheArchetypes(const FMassEntityManager& InEntityManager)
 {
 	const uint32 InEntityManagerHash = PointerHash(&InEntityManager);
@@ -157,31 +173,12 @@ void FMassEntityQuery::ForEachEntityChunk(const FMassArchetypeEntityCollection&
 void FMassEntityQuery::ForEachEntityChunk(FMassEntityManager& EntityManager, FMassExecutionContext& ExecutionContext, const FMassExecuteFunction& ExecuteFunction)
 {
 #if WITH_MASSENTITY_DEBUG
-	int32 NumEntitiesToProcess = 0;
-
 	checkf(ExecutionContext.ExecutionType == ExpectedContextType && (ExpectedContextType == EMassExecutionContextType::Local || bRegistered)
 		, TEXT("ExecutionContextType mismatch, make sure all the queries run as part of processor execution are registered with some processor with a FMassEntityQuery::RegisterWithProcessor call"));
 
 	EntityManager.GetRequirementAccessDetector().RequireAccess(*this);
 #endif
 
-	struct FScopedSubsystemRequirementsRestore
-	{
-		FScopedSubsystemRequirementsRestore(FMassExecutionContext& ExecutionContext)
-			: CachedExecutionContext(ExecutionContext)
-		{
-			CachedExecutionContext.GetSubsystemRequirementBits(ConstSubsystemsBitSet, MutableSubsystemsBitSet);
-		}
-
-		~FScopedSubsystemRequirementsRestore()
-		{
-			CachedExecutionContext.SetSubsystemRequirementBits(ConstSubsystemsBitSet, MutableSubsystemsBitSet);
-		}
-
-		FMassExecutionContext& CachedExecutionContext;
-		FMassExternalSubsystemBitSet ConstSubsystemsBitSet;
-		FMassExternalSubsystemBitSet MutableSubsystemsBitSet;
-	};
 	FScopedSubsystemRequirementsRestore SubsystemRestore(ExecutionContext);
 
 	if (ExecutionContext.CacheSubsystemRequirements(*this) == false)
@@ -190,12 +187,19 @@ void FMassEntityQuery::ForEachEntityChunk(FMassEntityManager& EntityManager, FMa
 		return;
 	}
 
+	// note that the following function will usualy only resort to verifying that the data is up to date by
+	// checking the version number. In rare cases when it would result in non trivial cost we actually
+	// do need those calculations.
+	CacheArchetypes(EntityManager);
+
 	// if there's a chunk collection set by the external code - use that
 	if (ExecutionContext.GetEntityCollection().IsSet())
 	{
 		const FMassArchetypeHandle& ArchetypeHandle = ExecutionContext.GetEntityCollection().GetArchetype();
-		// verify the archetype matches requirements
-		if (DoesArchetypeMatchRequirements(ArchetypeHandle) == false)
+		const int32 ArchetypeIndex = ValidArchetypes.Find(ArchetypeHandle);
+		
+		// if given ArchetypeHandle cannot be found in ValidArchetypes then it doesn't match the query's requirements
+		if (ArchetypeIndex == INDEX_NONE)
 		{
 			UE_VLOG_UELOG(EntityManager.GetOwner(), LogMass, Log, TEXT("Attempted to execute FMassEntityQuery with an incompatible Archetype: %s")
 				, *FMassDebugger::GetArchetypeRequirementCompatibilityDescription(*this, ArchetypeHandle));
@@ -205,11 +209,6 @@ void FMassEntityQuery::ForEachEntityChunk(FMassEntityManager& EntityManager, FMa
 #endif // WITH_MASSENTITY_DEBUG
 			return;
 		}
-		
-		// note that the following function will only resort to verifying that the data is up to date by
-		// checking the version number. In rare cases when it would result in non trivial cost we actually
-		// do need those calculations.
-		CacheArchetypes(EntityManager);
 
 		ExecutionContext.SetFragmentRequirements(*this);
 		
@@ -218,13 +217,9 @@ void FMassEntityQuery::ForEachEntityChunk(FMassEntityManager& EntityManager, FMa
 			, GetRequirementsMappingForArchetype(ArchetypeHandle)
 			, ExecutionContext.GetEntityCollection().GetRanges()
 			, ChunkCondition);
-#if WITH_MASSENTITY_DEBUG
-		NumEntitiesToProcess = ExecutionContext.GetNumEntities();
-#endif
 	}
 	else
 	{
-		CacheArchetypes(EntityManager);
 		// it's important to set requirements after caching archetypes due to that call potentially sorting the requirements and the order is relevant here.
 		ExecutionContext.SetFragmentRequirements(*this);
 
@@ -234,17 +229,160 @@ void FMassEntityQuery::ForEachEntityChunk(FMassEntityManager& EntityManager, FMa
 			FMassArchetypeData& ArchetypeData = FMassArchetypeHelper::ArchetypeDataFromHandleChecked(ArchetypeHandle);
 			ArchetypeData.ExecuteFunction(ExecutionContext, ExecuteFunction, ArchetypeFragmentMapping[i], ChunkCondition);
 			ExecutionContext.ClearFragmentViews();
+		}
+	}
+
 #if WITH_MASSENTITY_DEBUG
-			NumEntitiesToProcess += ExecutionContext.GetNumEntities();
+	EntityManager.GetRequirementAccessDetector().ReleaseAccess(*this);
 #endif
-		}
+
+	ExecutionContext.ClearExecutionData();
+	ExecutionContext.FlushDeferred();
+}
+
+void FMassEntityQuery::ParallelForEachEntityChunkInCollection(TConstArrayView<FMassArchetypeEntityCollection> EntityCollections
+	, FMassEntityManager& EntityManager, FMassExecutionContext& ExecutionContext, const FMassExecuteFunction& ExecuteFunction
+	, const EParallelForMode ParallelMode)
+{
+	if (UE::Mass::Tweakables::bAllowParallelExecution == false && ParallelMode != ForceParallelExecution)
+	{
+		ForEachEntityChunkInCollections(EntityCollections, EntityManager, ExecutionContext, ExecuteFunction);
+		return;
+	}
+
+	ParallelFor(EntityCollections.Num(), [this, &EntityManager, &ExecutionContext, &ExecuteFunction, &EntityCollections, ParallelMode](const int32 JobIndex)
+	{
+		FMassExecutionContext LocalExecutionContext = ExecutionContext; 
+		LocalExecutionContext.SetEntityCollection(EntityCollections[JobIndex]);
+		ParallelForEachEntityChunk(EntityManager, LocalExecutionContext, ExecuteFunction, ParallelMode);
+	});
+}
+
+void FMassEntityQuery::ParallelForEachEntityChunk(FMassEntityManager& EntityManager, FMassExecutionContext& ExecutionContext
+	, const FMassExecuteFunction& ExecuteFunction, const EParallelForMode ParallelMode)
+{
+	if (UE::Mass::Tweakables::bAllowParallelExecution == false && ParallelMode != ForceParallelExecution)
+	{
+		ForEachEntityChunk(EntityManager, ExecutionContext, ExecuteFunction);
+		return;
+	}
+
+#if WITH_MASSENTITY_DEBUG
+	checkf(ExecutionContext.ExecutionType == ExpectedContextType && (ExpectedContextType == EMassExecutionContextType::Local || bRegistered)
+		, TEXT("ExecutionContextType mismatch, make sure all the queries run as part of processor execution are registered with some processor with a FMassEntityQuery::RegisterWithProcessor call"));
+
+	EntityManager.GetRequirementAccessDetector().RequireAccess(*this);
+#endif
+
+	FScopedSubsystemRequirementsRestore SubsystemRestore(ExecutionContext);
+
+	if (ExecutionContext.CacheSubsystemRequirements(*this) == false)
+	{
+		// required subsystems are not available, bail out.
+		return;
 	}
 
+	struct FChunkJob
+	{
+		FMassArchetypeData& Archetype;
+		const int32 ArchetypeIndex;
+		const FMassArchetypeEntityCollection::FArchetypeEntityRange EntityRange;
+	};
+	TArray<FChunkJob> Jobs;
+
+	// note that the following function will usualy only resort to verifying that the data is up to date by
+	// checking the version number. In rare cases when it would result in non trivial cost we actually
+	// do need those calculations.
+	CacheArchetypes(EntityManager);
+
+	// if there's a chunk collection set by the external code - use that
+	if (ExecutionContext.GetEntityCollection().IsSet())
+	{
+		const FMassArchetypeHandle& ArchetypeHandle = ExecutionContext.GetEntityCollection().GetArchetype();
+		const int32 ArchetypeIndex = ValidArchetypes.Find(ArchetypeHandle);
+
+		// if given ArchetypeHandle cannot be found in ValidArchetypes then it doesn't match the query's requirements
+		if (ArchetypeIndex == INDEX_NONE)
+		{
+			UE_VLOG_UELOG(EntityManager.GetOwner(), LogMass, Log, TEXT("Attempted to execute FMassEntityQuery with an incompatible Archetype: %s")
+				, *FMassDebugger::GetArchetypeRequirementCompatibilityDescription(*this, ExecutionContext.GetEntityCollection().GetArchetype()));
+
 #if WITH_MASSENTITY_DEBUG
-	// Not using VLOG to be thread safe
-	UE_CLOG(!ExecutionContext.DebugGetExecutionDesc().IsEmpty(), LogMass, VeryVerbose,
-		TEXT("%s: %d entities sent for processing"), *ExecutionContext.DebugGetExecutionDesc(), NumEntitiesToProcess);
+			EntityManager.GetRequirementAccessDetector().ReleaseAccess(*this);
+#endif // WITH_MASSENTITY_DEBUG
+			return;
+		}
+
+		ExecutionContext.SetFragmentRequirements(*this);
+
+		FMassArchetypeData& ArchetypeRef = FMassArchetypeHelper::ArchetypeDataFromHandleChecked(ArchetypeHandle);
+		for (const FMassArchetypeEntityCollection::FArchetypeEntityRange& EntityRange : ExecutionContext.GetEntityCollection().GetRanges())
+		{
+			Jobs.Add({ ArchetypeRef, ArchetypeIndex, EntityRange });
+		}
+	}
+	else
+	{
+		ExecutionContext.SetFragmentRequirements(*this);
+		for (int ArchetypeIndex = 0; ArchetypeIndex < ValidArchetypes.Num(); ++ArchetypeIndex)
+		{
+			FMassArchetypeHandle& ArchetypeHandle = ValidArchetypes[ArchetypeIndex];
+			FMassArchetypeData& ArchetypeRef = FMassArchetypeHelper::ArchetypeDataFromHandleChecked(ArchetypeHandle);
+			const FMassArchetypeEntityCollection AsEntityCollection(ArchetypeHandle);
+			for (const FMassArchetypeEntityCollection::FArchetypeEntityRange& EntityRange : AsEntityCollection.GetRanges())
+			{
+				Jobs.Add({ ArchetypeRef, ArchetypeIndex, EntityRange });
+			}
+		}
+	}
+
+	if (Jobs.Num())
+	{
+		if (bAllowParallelCommands)
+		{
+			struct FTaskContext
+			{
+				FTaskContext()
+				{
+					CommandBuffer = MakeShared<FMassCommandBuffer>();
+				}
+				TSharedPtr<FMassCommandBuffer> CommandBuffer;
+			};
 
+			TArray<FTaskContext> TaskContext;
+
+			ParallelForWithTaskContext(TaskContext, Jobs.Num(), [this, &ExecutionContext, &ExecuteFunction, &Jobs](FTaskContext& TaskContext, const int32 JobIndex)
+				{
+					FMassExecutionContext LocalExecutionContext = ExecutionContext;
+
+					LocalExecutionContext.SetDeferredCommandBuffer(TaskContext.CommandBuffer);
+
+					Jobs[JobIndex].Archetype.ExecutionFunctionForChunk(LocalExecutionContext, ExecuteFunction
+						, Jobs[JobIndex].ArchetypeIndex != INDEX_NONE ? ArchetypeFragmentMapping[Jobs[JobIndex].ArchetypeIndex] : FMassQueryRequirementIndicesMapping()
+						, Jobs[JobIndex].EntityRange
+						, ChunkCondition);
+				});
+
+			// merge all command buffers
+			for (FTaskContext& CommandContext : TaskContext)
+			{
+				ExecutionContext.Defer().MoveAppend(*CommandContext.CommandBuffer);
+			}
+		}
+		else
+		{
+			ParallelFor(Jobs.Num(), [this, &ExecutionContext, &ExecuteFunction, &Jobs](const int32 JobIndex)
+				{
+					FMassExecutionContext LocalExecutionContext = ExecutionContext;
+					Jobs[JobIndex].Archetype.ExecutionFunctionForChunk(LocalExecutionContext, ExecuteFunction
+						, Jobs[JobIndex].ArchetypeIndex != INDEX_NONE ? ArchetypeFragmentMapping[Jobs[JobIndex].ArchetypeIndex] : FMassQueryRequirementIndicesMapping()
+						, Jobs[JobIndex].EntityRange
+						, ChunkCondition);
+				});
+		}
+	}
+
+#if WITH_MASSENTITY_DEBUG
 	EntityManager.GetRequirementAccessDetector().ReleaseAccess(*this);
 #endif
 
@@ -310,4 +448,3 @@ void FMassEntityQuery::ExportRequirements(FMassExecutionRequirements& OutRequire
 	FMassSubsystemRequirements::ExportRequirements(OutRequirements);
 	FMassFragmentRequirements::ExportRequirements(OutRequirements);
 }
-
diff --git a/Engine/Plugins/Runtime/MassEntity/Source/MassEntity/Public/MassEntityQuery.h b/Engine/Plugins/Runtime/MassEntity/Source/MassEntity/Public/MassEntityQuery.h
--- a/Engine/Plugins/Runtime/MassEntity/Source/MassEntity/Public/MassEntityQuery.h
+++ b/Engine/Plugins/Runtime/MassEntity/Source/MassEntity/Public/MassEntityQuery.h
@@ -28,6 +28,12 @@ struct MASSENTITY_API FMassEntityQuery : public FMassFragmentRequirements, publi
 	friend struct FMassDebugger;
 
 public:
+	enum EParallelForMode
+	{
+		Default, // use whatever the whole system has been configured for
+		ForceParallelExecution // force
+	};
+
 	FMassEntityQuery();
 	FMassEntityQuery(std::initializer_list<UScriptStruct*> InitList);
 	FMassEntityQuery(TConstArrayView<const UScriptStruct*> InitList);
@@ -39,9 +45,21 @@ public:
 	void ForEachEntityChunk(FMassEntityManager& EntityManager, FMassExecutionContext& ExecutionContext, const FMassExecuteFunction& ExecuteFunction);
 	
 	/** Will first verify that the archetype given with Collection matches the query's requirements, and if so will run the other, more generic ForEachEntityChunk implementation */
-	void ForEachEntityChunk(const FMassArchetypeEntityCollection& EntityCollection, FMassEntityManager& EntitySubsystem, FMassExecutionContext& ExecutionContext, const FMassExecuteFunction& ExecuteFunction);
+	void ForEachEntityChunk(const FMassArchetypeEntityCollection& EntityCollection, FMassEntityManager& EntitySubsystem
+		, FMassExecutionContext& ExecutionContext, const FMassExecuteFunction& ExecuteFunction);
 
-	void ForEachEntityChunkInCollections(TConstArrayView<FMassArchetypeEntityCollection> EntityCollections, FMassEntityManager& EntityManager, FMassExecutionContext& ExecutionContext, const FMassExecuteFunction& ExecuteFunction);
+	/**
+	 * Attempts to process every chunk of every affected archetype in parallel.
+	 */
+	void ParallelForEachEntityChunk(FMassEntityManager& EntityManager, FMassExecutionContext& ExecutionContext
+		, const FMassExecuteFunction& ExecuteFunction, const EParallelForMode ParallelMode = Default);
+
+	void ForEachEntityChunkInCollections(TConstArrayView<FMassArchetypeEntityCollection> EntityCollections, FMassEntityManager& EntityManager
+		, FMassExecutionContext& ExecutionContext, const FMassExecuteFunction& ExecuteFunction);
+
+	void ParallelForEachEntityChunkInCollection(TConstArrayView<FMassArchetypeEntityCollection> EntityCollections
+		, FMassEntityManager& EntityManager, FMassExecutionContext& ExecutionContext, const FMassExecuteFunction& ExecuteFunction
+		, const EParallelForMode ParallelMode);
 
 	/** Will gather all archetypes from InEntityManager matching this->Requirements.
 	 *  Note that no work will be done if the cached data is up to date (as tracked by EntitySubsystemHash and 
@@ -115,10 +133,23 @@ public:
 
 	void ExportRequirements(FMassExecutionRequirements& OutRequirements) const;
 
-protected:
-	void ReadCommandlineParams();
+	/** 
+	 * Controls whether ParallelForEachEntityChunk creates separate command buffers for each job.
+	 * @see bAllowParallelCommands for more details
+	 */
+	void SetParallelCommandBufferEnabled(const bool bInAllowParallelCommands) { bAllowParallelCommands = bInAllowParallelCommands; }
 
 private:
+	struct FScopedSubsystemRequirementsRestore
+	{
+		FScopedSubsystemRequirementsRestore(FMassExecutionContext& ExecutionContext);
+		~FScopedSubsystemRequirementsRestore();
+
+		FMassExecutionContext& CachedExecutionContext;
+		FMassExternalSubsystemBitSet ConstSubsystemsBitSet;
+		FMassExternalSubsystemBitSet MutableSubsystemsBitSet;
+	};
+
 	/** 
 	 * This function represents a condition that will be called for every chunk to be processed before the actual 
 	 * execution function is called. The chunk fragment requirements are already bound and ready to be used by the time 
@@ -132,9 +163,17 @@ private:
 	TArray<FMassArchetypeHandle> ValidArchetypes;
 	TArray<FMassQueryRequirementIndicesMapping> ArchetypeFragmentMapping;
 
-	uint8 bAllowParallelExecution : 1;
-	uint8 bRequiresGameThreadExecution : 1;
-	uint8 bRequiresMutatingWorldAccess : 1;
+	/** 
+	 * Controls whether ParallelForEachEntityChunk created dedicated command buffer for each job. This is required 
+	 * to ensure thread safety. Disable by calling SetParallelCommandBufferEnabled(false) if execution function doesn't 
+	 * issue commands. Disabling will save some performance since it will avoid dynamic allocation of command buffers.
+	 * 
+	 * @Note that disabling parallel commands will result in no command buffer getting passed to execution which in turn
+	 *	will cause crashes if the underlying code does try to issue commands. 
+	 */
+	uint8 bAllowParallelCommands : 1 = true;
+	uint8 bRequiresGameThreadExecution : 1 = false;
+	uint8 bRequiresMutatingWorldAccess : 1 = false;
 
 	EMassExecutionContextType ExpectedContextType = EMassExecutionContextType::Local;
 
diff --git a/Engine/Plugins/Runtime/MassEntity/Source/MassEntityTestSuite/MassEntityTestSuite.Build.cs b/Engine/Plugins/Runtime/MassEntity/Source/MassEntityTestSuite/MassEntityTestSuite.Build.cs
--- a/Engine/Plugins/Runtime/MassEntity/Source/MassEntityTestSuite/MassEntityTestSuite.Build.cs
+++ b/Engine/Plugins/Runtime/MassEntity/Source/MassEntityTestSuite/MassEntityTestSuite.Build.cs
@@ -1,6 +1,7 @@
 // Copyright Epic Games, Inc. All Rights Reserved.
 
 using UnrealBuildTool;
+using System.IO;
 
 namespace UnrealBuildTool.Rules
 {
@@ -10,10 +11,7 @@ namespace UnrealBuildTool.Rules
 		{
 			PCHUsage = ModuleRules.PCHUsageMode.UseExplicitOrSharedPCHs;
 
-			PublicIncludePaths.AddRange(
-				new string[] {
-				}
-			);
+			PrivateIncludePaths.Add(Path.Combine(GetModuleDirectory("MassEntity"), "Private")); 
 
 			PublicDependencyModuleNames.AddRange(
 				new string[] {
diff --git a/Engine/Plugins/Runtime/MassEntity/Source/MassEntityTestSuite/Private/MassParallelQueryTest.cpp b/Engine/Plugins/Runtime/MassEntity/Source/MassEntityTestSuite/Private/MassParallelQueryTest.cpp
new file mode 100644
--- /dev/null
+++ b/Engine/Plugins/Runtime/MassEntity/Source/MassEntityTestSuite/Private/MassParallelQueryTest.cpp
@@ -0,0 +1,158 @@
+// Copyright Epic Games, Inc. All Rights Reserved.
+
+#include "AITestsCommon.h"
+
+#include "Engine/World.h"
+#include "MassEntityManager.h"
+#include "MassProcessingTypes.h"
+#include "MassEntityTestTypes.h"
+#include "MassExecutor.h"
+#include "MassExecutionContext.h"
+#include "MassArchetypeData.h"
+
+#define LOCTEXT_NAMESPACE "MassTest"
+
+namespace FMassQueryTest
+{
+
+struct FQueryTest_ParallelForBasic : FEntityTestBase
+{
+	virtual bool InstantTest() override
+	{
+		CA_ASSUME(EntityManager);
+
+		const int32 EntitiesPerChunk = FMassArchetypeHelper::ArchetypeDataFromHandleChecked(FloatsArchetype).GetNumEntitiesPerChunk();
+
+		const int32 NumEntitiesToCreate = 32 * EntitiesPerChunk;
+		TArray<FMassEntityHandle> CreatedEntities;
+		EntityManager->BatchCreateEntities(FloatsArchetype, NumEntitiesToCreate, CreatedEntities);
+
+		FMassEntityQuery Query;
+		Query.AddRequirement<FTestFragment_Float>(EMassFragmentAccess::ReadWrite);
+
+		std::atomic<int32> EntitiesProcessed = 0;
+		std::atomic<int32> ConcurrentAcces = 0;
+		FCriticalSection ParallelDetectionCS;
+
+		FMassExecutionContext Context = EntityManager->CreateExecutionContext(/*DeltaSeconds=*/0.f);
+		Query.ParallelForEachEntityChunk(*EntityManager, Context, [&ParallelDetectionCS, &EntitiesProcessed, &ConcurrentAcces](FMassExecutionContext& Context)
+			{
+				EntitiesProcessed += Context.GetNumEntities();
+
+				if (ParallelDetectionCS.TryLock() == false)
+				{
+					++ConcurrentAcces;
+					return;
+				}
+				FPlatformProcess::Sleep(0.01f);
+			}, FMassEntityQuery::ForceParallelExecution);
+		
+		AITEST_EQUAL("All created entites should have been processed at this point", EntitiesProcessed.load(), NumEntitiesToCreate);
+		AITEST_TRUE("Some work is expected to be executed in parallel", ConcurrentAcces.load() > 0);
+
+		return true;
+	}
+};
+IMPLEMENT_AI_INSTANT_TEST(FQueryTest_ParallelForBasic, "System.Mass.Query.ParallelForBasic");
+
+struct FQueryTest_ParallelForCollection : FEntityTestBase
+{
+	virtual bool InstantTest() override
+	{
+		CA_ASSUME(EntityManager);
+
+		const int32 NumEntitiesToCreate = 1000;
+		TArray<FMassEntityHandle> CreatedEntities;
+		EntityManager->BatchCreateEntities(FloatsArchetype, NumEntitiesToCreate, CreatedEntities);
+
+		TArray<FMassEntityHandle> EntitiesToProcess;
+		EntitiesToProcess.Reserve(NumEntitiesToCreate);
+
+		FRandomStream RandomStream(/*Seed=*/1);
+		for (FMassEntityHandle& EntityHandle : CreatedEntities)
+		{
+			if (RandomStream.FRand() < 0.5)
+			{
+				EntitiesToProcess.Add(EntityHandle);
+			}
+		}
+		ensure(EntitiesToProcess.Num() > 0);
+		const FMassArchetypeEntityCollection EntityCollection(FloatsArchetype, EntitiesToProcess, FMassArchetypeEntityCollection::NoDuplicates);
+		const int32 NumJobs = EntityCollection.GetRanges().Num();
+
+		FMassEntityQuery Query;
+		Query.AddRequirement<FTestFragment_Float>(EMassFragmentAccess::ReadWrite);
+
+		std::atomic<int32> EntitiesProcessed = 0;
+		std::atomic<int32> ConcurrentAcces = 0;
+		std::atomic<int32> JobsExecuted = 0;
+		FCriticalSection ParallelDetectionCS;
+
+		FMassExecutionContext Context = EntityManager->CreateExecutionContext(/*DeltaSeconds=*/0.f);
+		Context.SetEntityCollection(EntityCollection);
+
+		Query.ParallelForEachEntityChunk(*EntityManager, Context, [&JobsExecuted, &ParallelDetectionCS, &EntitiesProcessed, &ConcurrentAcces](FMassExecutionContext& Context)
+			{
+				++JobsExecuted;
+				EntitiesProcessed += Context.GetNumEntities();
+
+				if (ParallelDetectionCS.TryLock() == false)
+				{
+					++ConcurrentAcces;
+					return;
+				}
+				FPlatformProcess::Sleep(0.01f);
+			}, FMassEntityQuery::ForceParallelExecution);
+
+		AITEST_EQUAL("All entities passed for processing should have been counted", EntitiesProcessed.load(), EntitiesToProcess.Num());
+		AITEST_EQUAL("The number of jobs should match expectations", NumJobs, JobsExecuted.load());
+		AITEST_TRUE("Some work is expected to be executed in parallel", ConcurrentAcces.load() > 0);
+
+		return true;
+	}
+};
+IMPLEMENT_AI_INSTANT_TEST(FQueryTest_ParallelForCollection, "System.Mass.Query.ParallelForCollection");
+
+#if WITH_MASSENTITY_DEBUG
+struct FQueryTest_ParallelCommands : FEntityTestBase
+{
+	virtual bool InstantTest() override
+	{
+		CA_ASSUME(EntityManager);
+
+		const FMassArchetypeHandle TargetArchetype = EntityManager->CreateArchetype(FloatsArchetype, { FTestFragment_Tag::StaticStruct() });
+
+		const int32 EntitiesPerChunk = FMassArchetypeHelper::ArchetypeDataFromHandleChecked(FloatsArchetype).GetNumEntitiesPerChunk();
+
+		const int32 NumEntitiesToCreate = 32 * EntitiesPerChunk;
+		TArray<FMassEntityHandle> CreatedEntities;
+		EntityManager->BatchCreateEntities(FloatsArchetype, NumEntitiesToCreate, CreatedEntities);
+
+		FMassEntityQuery Query;
+		Query.AddRequirement<FTestFragment_Float>(EMassFragmentAccess::ReadWrite);
+		Query.SetParallelCommandBufferEnabled(true);
+
+		const int32 OriginalEntityCount = EntityManager->DebugGetArchetypeEntitiesCount(FloatsArchetype);
+		AITEST_EQUAL("Target archetype should only contain freshly created entities", OriginalEntityCount, NumEntitiesToCreate);
+
+		FMassExecutionContext Context = EntityManager->CreateExecutionContext(/*DeltaSeconds=*/0.f);
+		Query.ParallelForEachEntityChunk(*EntityManager, Context, [](FMassExecutionContext& Context)
+			{
+				Context.Defer().PushCommand<FMassCommandAddTags<FTestFragment_Tag>>(Context.GetEntities());
+			}, FMassEntityQuery::ForceParallelExecution);
+
+		const int32 OriginalArchetypeCountAfterMove = EntityManager->DebugGetArchetypeEntitiesCount(FloatsArchetype);
+		AITEST_EQUAL("Target archetype should be empty after the move", OriginalArchetypeCountAfterMove, 0);
+		const int32 TargetArchetypeCountAfterMove = EntityManager->DebugGetArchetypeEntitiesCount(TargetArchetype);
+		AITEST_EQUAL("All entities are expected to be moved over to the target archetype", OriginalEntityCount, TargetArchetypeCountAfterMove);
+
+		return true;
+	}
+};
+IMPLEMENT_AI_INSTANT_TEST(FQueryTest_ParallelCommands, "System.Mass.Query.ParallelForCommands");
+#endif // WITH_MASSENTITY_DEBUG
+
+} // FMassQueryTest
+
+#undef LOCTEXT_NAMESPACE
+
