Improved FMassExecutionContext's checks whether subsystems and fragments are bound and configured for read-write if accessing mutable version.
The CL also contains fixes to existing Mass code that used misconfigured EntityQueries

diff --git a/Engine/Plugins/AI/MassAI/Source/MassAIBehavior/Private/MassZoneGraphAnnotationProcessors.cpp b/Engine/Plugins/AI/MassAI/Source/MassAIBehavior/Private/MassZoneGraphAnnotationProcessors.cpp
--- a/Engine/Plugins/AI/MassAI/Source/MassAIBehavior/Private/MassZoneGraphAnnotationProcessors.cpp
+++ b/Engine/Plugins/AI/MassAI/Source/MassAIBehavior/Private/MassZoneGraphAnnotationProcessors.cpp
@@ -83,7 +83,7 @@ void UMassZoneGraphAnnotationTagUpdateProcessor::ConfigureQueries()
 	Super::ConfigureQueries();
 	EntityQuery.AddRequirement<FMassZoneGraphAnnotationFragment>(EMassFragmentAccess::ReadWrite);
 	EntityQuery.AddRequirement<FMassZoneGraphLaneLocationFragment>(EMassFragmentAccess::ReadOnly);
-	EntityQuery.AddChunkRequirement<FMassZoneGraphAnnotationVariableTickChunkFragment>(EMassFragmentAccess::ReadOnly);
+	EntityQuery.AddChunkRequirement<FMassZoneGraphAnnotationVariableTickChunkFragment>(EMassFragmentAccess::ReadWrite);
 	EntityQuery.AddChunkRequirement<FMassSimulationVariableTickChunkFragment>(EMassFragmentAccess::ReadOnly, EMassFragmentPresence::Optional);
 	EntityQuery.AddSubsystemRequirement<UZoneGraphAnnotationSubsystem>(EMassFragmentAccess::ReadWrite);
 
diff --git a/Engine/Plugins/AI/MassAI/Source/MassAIReplication/Private/MassReplicationPathHandlers.cpp b/Engine/Plugins/AI/MassAI/Source/MassAIReplication/Private/MassReplicationPathHandlers.cpp
--- a/Engine/Plugins/AI/MassAI/Source/MassAIReplication/Private/MassReplicationPathHandlers.cpp
+++ b/Engine/Plugins/AI/MassAI/Source/MassAIReplication/Private/MassReplicationPathHandlers.cpp
@@ -12,9 +12,9 @@
 
 void FMassReplicationProcessorPathHandler::AddRequirements(FMassEntityQuery& InQuery)
 {
-	InQuery.AddRequirement<FMassZoneGraphPathRequestFragment>(EMassFragmentAccess::ReadOnly);
-	InQuery.AddRequirement<FMassMoveTargetFragment>(EMassFragmentAccess::ReadOnly);
-	InQuery.AddRequirement<FMassZoneGraphLaneLocationFragment>(EMassFragmentAccess::ReadOnly);
+	InQuery.AddRequirement<FMassZoneGraphPathRequestFragment>(EMassFragmentAccess::ReadWrite);
+	InQuery.AddRequirement<FMassMoveTargetFragment>(EMassFragmentAccess::ReadWrite);
+	InQuery.AddRequirement<FMassZoneGraphLaneLocationFragment>(EMassFragmentAccess::ReadWrite);
 }
 
 void FMassReplicationProcessorPathHandler::CacheFragmentViews(FMassExecutionContext& ExecContext)
diff --git a/Engine/Plugins/Runtime/MassEntity/Source/MassEntity/Public/MassExecutionContext.h b/Engine/Plugins/Runtime/MassEntity/Source/MassEntity/Public/MassExecutionContext.h
--- a/Engine/Plugins/Runtime/MassEntity/Source/MassEntity/Public/MassExecutionContext.h
+++ b/Engine/Plugins/Runtime/MassEntity/Source/MassEntity/Public/MassExecutionContext.h
@@ -8,6 +8,16 @@
 #include "MassSubsystemAccess.h"
 
 
+#define CHECK_IF_VALID(View, Type) \
+	checkf(View \
+		, TEXT("Requested fragment type not bound, type %s. Make sure it has been listed as required."), *GetNameSafe(Type))
+
+#define CHECK_IF_READWRITE(View) \
+	checkf(View == nullptr || View->Requirement.AccessMode == EMassFragmentAccess::ReadWrite \
+		, TEXT("Requested fragment type not bound for writing, type %s. Make sure it has been listed as required in ReadWrite mode.") \
+		, View ? *GetNameSafe(View->Requirement.StructType) : TEXT("[Not found]"))
+
+
 struct MASSENTITY_API FMassExecutionContext
 {
 private:
@@ -154,6 +164,7 @@ public:
 
 		const UScriptStruct* Type = T::StaticStruct();
 		FChunkFragmentView* FoundChunkFragmentData = ChunkFragmentViews.FindByPredicate([Type](const FChunkFragmentView& Element) { return Element.Requirement.StructType == Type; } );
+		CHECK_IF_READWRITE(FoundChunkFragmentData);
 		return FoundChunkFragmentData ? FoundChunkFragmentData->FragmentView.GetPtr<T>() : static_cast<T*>(nullptr);
 	}
 	
@@ -161,7 +172,7 @@ public:
 	T& GetMutableChunkFragment()
 	{
 		T* ChunkFragment = GetMutableChunkFragmentPtr<T>();
-		checkf(ChunkFragment, TEXT("Chunk Fragment requirement not found: %s"), *T::StaticStruct()->GetName());
+		CHECK_IF_VALID(ChunkFragment, T::StaticStruct());
 		return *ChunkFragment;
 	}
 
@@ -179,7 +190,7 @@ public:
 	const T& GetChunkFragment() const
 	{
 		const T* ChunkFragment = GetChunkFragmentPtr<T>();
-		checkf(ChunkFragment, TEXT("Chunk Fragment requirement not found: %s"), *T::StaticStruct()->GetName());
+		CHECK_IF_VALID(ChunkFragment, T::StaticStruct());
 		return *ChunkFragment;
 	}
 
@@ -197,17 +208,17 @@ public:
 	const T& GetConstSharedFragment() const
 	{
 		const T* SharedFragment = GetConstSharedFragmentPtr<const T>();
-		checkf(SharedFragment, TEXT("Shared Fragment requirement not found: %s"), *T::StaticStruct()->GetName());
+		CHECK_IF_VALID(SharedFragment, T::StaticStruct());
 		return *SharedFragment;
 	}
 
-
 	template<typename T>
 	T* GetMutableSharedFragmentPtr()
 	{
 		static_assert(TIsDerivedFrom<T, FMassSharedFragment>::IsDerived, "Given struct doesn't represent a valid shared fragment type. Make sure to inherit from FMassSharedFragment or one of its child-types.");
 
 		FSharedFragmentView* FoundSharedFragmentData = SharedFragmentViews.FindByPredicate([](const FSharedFragmentView& Element) { return Element.Requirement.StructType == T::StaticStruct(); });
+		CHECK_IF_READWRITE(FoundSharedFragmentData);
 		return FoundSharedFragmentData ? FoundSharedFragmentData->FragmentView.GetPtr<T>() : static_cast<T*>(nullptr);
 	}
 
@@ -224,7 +235,7 @@ public:
 	T& GetMutableSharedFragment()
 	{
 		T* SharedFragment = GetMutableSharedFragmentPtr<T>();
-		checkf(SharedFragment, TEXT("Shared Fragment requirement not found: %s"), *T::StaticStruct()->GetName());
+		CHECK_IF_VALID(SharedFragment, T::StaticStruct());
 		return *SharedFragment;
 	}
 
@@ -232,7 +243,7 @@ public:
 	const T& GetSharedFragment() const
 	{
 		const T* SharedFragment = GetSharedFragmentPtr<T>();
-		checkf(SharedFragment, TEXT("Shared Fragment requirement not found: %s"), *T::StaticStruct()->GetName());
+		CHECK_IF_VALID(SharedFragment, T::StaticStruct());
 		return *SharedFragment;
 	}
 
@@ -242,8 +253,8 @@ public:
 	{
 		const UScriptStruct* FragmentType = TFragment::StaticStruct();
 		const FFragmentView* View = FragmentViews.FindByPredicate([FragmentType](const FFragmentView& Element) { return Element.Requirement.StructType == FragmentType; });
-		checkfSlow(View != nullptr, TEXT("Requested fragment type not bound"));
-		checkfSlow(View->Requirement.AccessMode == EMassFragmentAccess::ReadWrite, TEXT("Requested fragment has not been bound for writing"));
+		CHECK_IF_VALID(View, FragmentType);
+		CHECK_IF_READWRITE(View);
 		return MakeArrayView<TFragment>((TFragment*)View->FragmentView.GetData(), View->FragmentView.Num());
 	}
 
@@ -252,21 +263,22 @@ public:
 	{
 		const UScriptStruct* FragmentType = TFragment::StaticStruct();
 		const FFragmentView* View = FragmentViews.FindByPredicate([FragmentType](const FFragmentView& Element) { return Element.Requirement.StructType == FragmentType; });
-		checkfSlow(View != nullptr, TEXT("Requested fragment type not bound. Make sure the Frament was requested in ReadOnly or ReadWrite mode."));
+		CHECK_IF_VALID(View, TFragment::StaticStruct());
 		return TConstArrayView<TFragment>((const TFragment*)View->FragmentView.GetData(), View->FragmentView.Num());
 	}
 
 	TConstArrayView<FMassFragment> GetFragmentView(const UScriptStruct* FragmentType) const
 	{
 		const FFragmentView* View = FragmentViews.FindByPredicate([FragmentType](const FFragmentView& Element) { return Element.Requirement.StructType == FragmentType; });
-		checkSlow(View);
+		CHECK_IF_VALID(View, FragmentType);
 		return TConstArrayView<FMassFragment>((const FMassFragment*)View->FragmentView.GetData(), View->FragmentView.Num());;
 	}
 
 	TArrayView<FMassFragment> GetMutableFragmentView(const UScriptStruct* FragmentType) 
 	{
 		const FFragmentView* View = FragmentViews.FindByPredicate([FragmentType](const FFragmentView& Element) { return Element.Requirement.StructType == FragmentType; });
-		checkSlow(View);
+		CHECK_IF_VALID(View, FragmentType);
+		CHECK_IF_READWRITE(View);
 		return View->FragmentView;
 	}
 
@@ -447,3 +459,6 @@ public:
 	UE_DEPRECATED(5.4, "Deprecated in favor of 'SetCurrentArchetypeCompositionDescriptor' as this provides information on the entire archetype.")
 	void SetCurrentArchetypesTagBitSet(const FMassTagBitSet&) {}
 };
+
+#undef CHECK_IF_VALID
+#undef CHECK_IF_READWRITE
diff --git a/Engine/Plugins/Runtime/MassGameplay/Source/MassLOD/Private/MassSimulationLOD.cpp b/Engine/Plugins/Runtime/MassGameplay/Source/MassLOD/Private/MassSimulationLOD.cpp
--- a/Engine/Plugins/Runtime/MassGameplay/Source/MassLOD/Private/MassSimulationLOD.cpp
+++ b/Engine/Plugins/Runtime/MassGameplay/Source/MassLOD/Private/MassSimulationLOD.cpp
@@ -98,7 +98,7 @@ void UMassSimulationLODProcessor::ConfigureQueries()
 	EntityQueryVariableTick.AddRequirement<FMassSimulationVariableTickFragment>(EMassFragmentAccess::ReadWrite);
 	EntityQueryVariableTick.AddConstSharedRequirement<FMassSimulationVariableTickParameters>();
 	EntityQueryVariableTick.AddChunkRequirement<FMassSimulationVariableTickChunkFragment>(EMassFragmentAccess::ReadWrite);
-	EntityQueryVariableTick.AddSharedRequirement<FMassSimulationVariableTickSharedFragment>(EMassFragmentAccess::ReadOnly);
+	EntityQueryVariableTick.AddSharedRequirement<FMassSimulationVariableTickSharedFragment>(EMassFragmentAccess::ReadWrite);
 
 	// In case where the variableTick isn't enabled, we might need to set LOD tags as if the users still wants them
 	EntityQuerySetLODTag.AddRequirement<FMassSimulationLODFragment>(EMassFragmentAccess::ReadOnly);
diff --git a/Engine/Plugins/Runtime/MassGameplay/Source/MassReplication/Private/MassReplicationProcessor.cpp b/Engine/Plugins/Runtime/MassGameplay/Source/MassReplication/Private/MassReplicationProcessor.cpp
--- a/Engine/Plugins/Runtime/MassGameplay/Source/MassReplication/Private/MassReplicationProcessor.cpp
+++ b/Engine/Plugins/Runtime/MassGameplay/Source/MassReplication/Private/MassReplicationProcessor.cpp
@@ -41,7 +41,7 @@ void UMassReplicationProcessor::ConfigureQueries()
 	SyncClientData.AddRequirement<FMassReplicatedAgentFragment>(EMassFragmentAccess::ReadWrite);
 
 	CollectViewerInfoQuery.AddRequirement<FTransformFragment>(EMassFragmentAccess::ReadOnly);
-	CollectViewerInfoQuery.AddRequirement<FMassReplicationViewerInfoFragment>(EMassFragmentAccess::ReadOnly);
+	CollectViewerInfoQuery.AddRequirement<FMassReplicationViewerInfoFragment>(EMassFragmentAccess::ReadWrite);
 	CollectViewerInfoQuery.AddSharedRequirement<FMassReplicationSharedFragment>(EMassFragmentAccess::ReadWrite);
 
 	CalculateLODQuery.AddRequirement<FMassReplicationViewerInfoFragment>(EMassFragmentAccess::ReadOnly);
