Made it possible to run empty MassEntityQueries provided the collection of entities to process is explicitly provided via FMassExecutionContext.EntityCollection.

Some Mass unit tests unrelated to the specific change have been touch up due to the used processors' queries being misconfigured. That was causing test failures unrelated to the funcitonality being tested.

diff --git a/Engine/Plugins/Runtime/MassEntity/Source/MassEntity/Private/MassEntityQuery.cpp b/Engine/Plugins/Runtime/MassEntity/Source/MassEntity/Private/MassEntityQuery.cpp
--- a/Engine/Plugins/Runtime/MassEntity/Source/MassEntity/Private/MassEntityQuery.cpp
+++ b/Engine/Plugins/Runtime/MassEntity/Source/MassEntity/Private/MassEntityQuery.cpp
@@ -96,7 +96,7 @@ void FMassEntityQuery::CacheArchetypes(const FMassEntityManager& InEntityManager
     bool bUpdateArchetypes = InEntityManager.GetArchetypeDataVersion() != LastUpdatedArchetypeDataVersion;
 
 	// Force a full update if the entity system changed or if the requirements changed
-	if (EntitySubsystemHash != InEntityManagerHash || IncrementalChangesCount)
+	if (EntitySubsystemHash != InEntityManagerHash || HasIncrementalChanges())
 	{
 		bUpdateArchetypes = true;
 		EntitySubsystemHash = InEntityManagerHash;
@@ -104,10 +104,10 @@ void FMassEntityQuery::CacheArchetypes(const FMassEntityManager& InEntityManager
 		LastUpdatedArchetypeDataVersion = 0;
 		ArchetypeFragmentMapping.Reset();
 
-		if (IncrementalChangesCount)
+		if (HasIncrementalChanges())
 		{
-			IncrementalChangesCount = 0;
-			if( CheckValidity() )
+			ConsumeIncrementalChangesCount();
+			if (CheckValidity())
 			{
 				SortRequirements();
 			}
@@ -187,48 +187,65 @@ void FMassEntityQuery::ForEachEntityChunk(FMassEntityManager& EntityManager, FMa
 		return;
 	}
 
-	// note that the following function will usualy only resort to verifying that the data is up to date by
-	// checking the version number. In rare cases when it would result in non trivial cost we actually
-	// do need those calculations.
-	CacheArchetypes(EntityManager);
-
-	// if there's a chunk collection set by the external code - use that
-	if (ExecutionContext.GetEntityCollection().IsSet())
+	if (FMassFragmentRequirements::IsEmpty())
 	{
-		const FMassArchetypeHandle& ArchetypeHandle = ExecutionContext.GetEntityCollection().GetArchetype();
-		const int32 ArchetypeIndex = ValidArchetypes.Find(ArchetypeHandle);
-		
-		// if given ArchetypeHandle cannot be found in ValidArchetypes then it doesn't match the query's requirements
-		if (ArchetypeIndex == INDEX_NONE)
+		if (ensureMsgf(ExecutionContext.GetEntityCollection().IsSet(), TEXT("Using empty queries is only supported in combination with Entity Collections that explicitly indicate entities to process")))
 		{
-			UE_VLOG_UELOG(EntityManager.GetOwner(), LogMass, Log, TEXT("Attempted to execute FMassEntityQuery with an incompatible Archetype: %s")
-				, *FMassDebugger::GetArchetypeRequirementCompatibilityDescription(*this, ArchetypeHandle));
+			static const FMassQueryRequirementIndicesMapping EmptyMapping;
 
-#if WITH_MASSENTITY_DEBUG
-			EntityManager.GetRequirementAccessDetector().ReleaseAccess(*this);
-#endif // WITH_MASSENTITY_DEBUG
-			return;
+			const FMassArchetypeHandle& ArchetypeHandle = ExecutionContext.GetEntityCollection().GetArchetype();
+			FMassArchetypeData& ArchetypeData = FMassArchetypeHelper::ArchetypeDataFromHandleChecked(ArchetypeHandle);
+			ArchetypeData.ExecuteFunction(ExecutionContext, ExecuteFunction
+				, EmptyMapping
+				, ExecutionContext.GetEntityCollection().GetRanges()
+				, ChunkCondition);
 		}
-
-		ExecutionContext.SetFragmentRequirements(*this);
-		
-		FMassArchetypeData& ArchetypeData = FMassArchetypeHelper::ArchetypeDataFromHandleChecked(ArchetypeHandle);
-		ArchetypeData.ExecuteFunction(ExecutionContext, ExecuteFunction
-			, GetRequirementsMappingForArchetype(ArchetypeHandle)
-			, ExecutionContext.GetEntityCollection().GetRanges()
-			, ChunkCondition);
 	}
 	else
 	{
-		// it's important to set requirements after caching archetypes due to that call potentially sorting the requirements and the order is relevant here.
-		ExecutionContext.SetFragmentRequirements(*this);
+		// note that the following function will usually only resort to verifying that the data is up to date by
+			// checking the version number. In rare cases when it would result in non trivial cost we actually
+			// do need those calculations.
+		CacheArchetypes(EntityManager);
 
-		for (int i = 0; i < ValidArchetypes.Num(); ++i)
+		// if there's a chunk collection set by the external code - use that
+		if (ExecutionContext.GetEntityCollection().IsSet())
 		{
-			const FMassArchetypeHandle& ArchetypeHandle = ValidArchetypes[i];
+			const FMassArchetypeHandle& ArchetypeHandle = ExecutionContext.GetEntityCollection().GetArchetype();
+			const int32 ArchetypeIndex = ValidArchetypes.Find(ArchetypeHandle);
+
+			// if given ArchetypeHandle cannot be found in ValidArchetypes then it doesn't match the query's requirements
+			if (ArchetypeIndex == INDEX_NONE)
+			{
+				UE_VLOG_UELOG(EntityManager.GetOwner(), LogMass, Log, TEXT("Attempted to execute FMassEntityQuery with an incompatible Archetype: %s")
+					, *FMassDebugger::GetArchetypeRequirementCompatibilityDescription(*this, ArchetypeHandle));
+
+#if WITH_MASSENTITY_DEBUG
+				EntityManager.GetRequirementAccessDetector().ReleaseAccess(*this);
+#endif // WITH_MASSENTITY_DEBUG
+				return;
+			}
+
+			ExecutionContext.SetFragmentRequirements(*this);
+
 			FMassArchetypeData& ArchetypeData = FMassArchetypeHelper::ArchetypeDataFromHandleChecked(ArchetypeHandle);
-			ArchetypeData.ExecuteFunction(ExecutionContext, ExecuteFunction, ArchetypeFragmentMapping[i], ChunkCondition);
-			ExecutionContext.ClearFragmentViews();
+			ArchetypeData.ExecuteFunction(ExecutionContext, ExecuteFunction
+				, GetRequirementsMappingForArchetype(ArchetypeHandle)
+				, ExecutionContext.GetEntityCollection().GetRanges()
+				, ChunkCondition);
+		}
+		else
+		{
+			// it's important to set requirements after caching archetypes due to that call potentially sorting the requirements and the order is relevant here.
+			ExecutionContext.SetFragmentRequirements(*this);
+
+			for (int i = 0; i < ValidArchetypes.Num(); ++i)
+			{
+				const FMassArchetypeHandle& ArchetypeHandle = ValidArchetypes[i];
+				FMassArchetypeData& ArchetypeData = FMassArchetypeHelper::ArchetypeDataFromHandleChecked(ArchetypeHandle);
+				ArchetypeData.ExecuteFunction(ExecutionContext, ExecuteFunction, ArchetypeFragmentMapping[i], ChunkCondition);
+				ExecutionContext.ClearFragmentViews();
+			}
 		}
 	}
 
@@ -290,50 +307,66 @@ void FMassEntityQuery::ParallelForEachEntityChunk(FMassEntityManager& EntityMana
 	};
 	TArray<FChunkJob> Jobs;
 
-	// note that the following function will usualy only resort to verifying that the data is up to date by
-	// checking the version number. In rare cases when it would result in non trivial cost we actually
-	// do need those calculations.
-	CacheArchetypes(EntityManager);
-
-	// if there's a chunk collection set by the external code - use that
-	if (ExecutionContext.GetEntityCollection().IsSet())
+	if (FMassFragmentRequirements::IsEmpty())
+	{
+		if (ensureMsgf(ExecutionContext.GetEntityCollection().IsSet(), TEXT("Using empty queries is only supported in combination with Entity Collections that explicitly indicate entities to process")))
+		{
+			const FMassArchetypeHandle& ArchetypeHandle = ExecutionContext.GetEntityCollection().GetArchetype();
+			FMassArchetypeData& ArchetypeRef = FMassArchetypeHelper::ArchetypeDataFromHandleChecked(ArchetypeHandle);
+			for (const FMassArchetypeEntityCollection::FArchetypeEntityRange& EntityRange : ExecutionContext.GetEntityCollection().GetRanges())
+			{
+				Jobs.Add({ ArchetypeRef, INDEX_NONE, EntityRange });
+			}
+		}
+	}
+	else
 	{
-		const FMassArchetypeHandle& ArchetypeHandle = ExecutionContext.GetEntityCollection().GetArchetype();
-		const int32 ArchetypeIndex = ValidArchetypes.Find(ArchetypeHandle);
 
-		// if given ArchetypeHandle cannot be found in ValidArchetypes then it doesn't match the query's requirements
-		if (ArchetypeIndex == INDEX_NONE)
+		// note that the following function will usualy only resort to verifying that the data is up to date by
+		// checking the version number. In rare cases when it would result in non trivial cost we actually
+		// do need those calculations.
+		CacheArchetypes(EntityManager);
+
+		// if there's a chunk collection set by the external code - use that
+		if (ExecutionContext.GetEntityCollection().IsSet())
 		{
-			UE_VLOG_UELOG(EntityManager.GetOwner(), LogMass, Log, TEXT("Attempted to execute FMassEntityQuery with an incompatible Archetype: %s")
-				, *FMassDebugger::GetArchetypeRequirementCompatibilityDescription(*this, ExecutionContext.GetEntityCollection().GetArchetype()));
+			const FMassArchetypeHandle& ArchetypeHandle = ExecutionContext.GetEntityCollection().GetArchetype();
+			const int32 ArchetypeIndex = ValidArchetypes.Find(ArchetypeHandle);
+
+			// if given ArchetypeHandle cannot be found in ValidArchetypes then it doesn't match the query's requirements
+			if (ArchetypeIndex == INDEX_NONE)
+			{
+				UE_VLOG_UELOG(EntityManager.GetOwner(), LogMass, Log, TEXT("Attempted to execute FMassEntityQuery with an incompatible Archetype: %s")
+					, *FMassDebugger::GetArchetypeRequirementCompatibilityDescription(*this, ExecutionContext.GetEntityCollection().GetArchetype()));
 
 #if WITH_MASSENTITY_DEBUG
-			EntityManager.GetRequirementAccessDetector().ReleaseAccess(*this);
+				EntityManager.GetRequirementAccessDetector().ReleaseAccess(*this);
 #endif // WITH_MASSENTITY_DEBUG
-			return;
-		}
+				return;
+			}
 
-		ExecutionContext.SetFragmentRequirements(*this);
+			ExecutionContext.SetFragmentRequirements(*this);
 
-		FMassArchetypeData& ArchetypeRef = FMassArchetypeHelper::ArchetypeDataFromHandleChecked(ArchetypeHandle);
-		for (const FMassArchetypeEntityCollection::FArchetypeEntityRange& EntityRange : ExecutionContext.GetEntityCollection().GetRanges())
-		{
-			Jobs.Add({ ArchetypeRef, ArchetypeIndex, EntityRange });
-		}
-	}
-	else
-	{
-		ExecutionContext.SetFragmentRequirements(*this);
-		for (int ArchetypeIndex = 0; ArchetypeIndex < ValidArchetypes.Num(); ++ArchetypeIndex)
-		{
-			FMassArchetypeHandle& ArchetypeHandle = ValidArchetypes[ArchetypeIndex];
 			FMassArchetypeData& ArchetypeRef = FMassArchetypeHelper::ArchetypeDataFromHandleChecked(ArchetypeHandle);
-			const FMassArchetypeEntityCollection AsEntityCollection(ArchetypeHandle);
-			for (const FMassArchetypeEntityCollection::FArchetypeEntityRange& EntityRange : AsEntityCollection.GetRanges())
+			for (const FMassArchetypeEntityCollection::FArchetypeEntityRange& EntityRange : ExecutionContext.GetEntityCollection().GetRanges())
 			{
 				Jobs.Add({ ArchetypeRef, ArchetypeIndex, EntityRange });
 			}
 		}
+		else
+		{
+			ExecutionContext.SetFragmentRequirements(*this);
+			for (int ArchetypeIndex = 0; ArchetypeIndex < ValidArchetypes.Num(); ++ArchetypeIndex)
+			{
+				FMassArchetypeHandle& ArchetypeHandle = ValidArchetypes[ArchetypeIndex];
+				FMassArchetypeData& ArchetypeRef = FMassArchetypeHelper::ArchetypeDataFromHandleChecked(ArchetypeHandle);
+				const FMassArchetypeEntityCollection AsEntityCollection(ArchetypeHandle);
+				for (const FMassArchetypeEntityCollection::FArchetypeEntityRange& EntityRange : AsEntityCollection.GetRanges())
+				{
+					Jobs.Add({ ArchetypeRef, ArchetypeIndex, EntityRange });
+				}
+			}
+		}
 	}
 
 	if (Jobs.Num())
@@ -438,7 +471,7 @@ bool FMassEntityQuery::HasMatchingEntities(FMassEntityManager& InEntityManager)
 const FMassQueryRequirementIndicesMapping& FMassEntityQuery::GetRequirementsMappingForArchetype(const FMassArchetypeHandle ArchetypeHandle) const
 {
 	static const FMassQueryRequirementIndicesMapping FallbackEmptyMapping;
-	checkf(IncrementalChangesCount == 0, TEXT("Fetching cached fragments mapping while the query's cached data is out of sync!"));
+	checkf(HasIncrementalChanges() == false, TEXT("Fetching cached fragments mapping while the query's cached data is out of sync!"));
 	const int32 ArchetypeIndex = ValidArchetypes.Find(ArchetypeHandle);
 	return ArchetypeIndex != INDEX_NONE ? ArchetypeFragmentMapping[ArchetypeIndex] : FallbackEmptyMapping;
 }
diff --git a/Engine/Plugins/Runtime/MassEntity/Source/MassEntity/Private/MassRequirements.cpp b/Engine/Plugins/Runtime/MassEntity/Source/MassEntity/Private/MassRequirements.cpp
--- a/Engine/Plugins/Runtime/MassEntity/Source/MassEntity/Private/MassRequirements.cpp
+++ b/Engine/Plugins/Runtime/MassEntity/Source/MassEntity/Private/MassRequirements.cpp
@@ -83,15 +83,25 @@ void FMassFragmentRequirements::SortRequirements()
 
 bool FMassFragmentRequirements::CheckValidity() const
 {
-	return
-		RequiredAllFragments.IsEmpty() == false || RequiredAnyFragments.IsEmpty() == false || RequiredOptionalFragments.IsEmpty() == false ||
-		RequiredAllTags.IsEmpty() == false || RequiredAnyTags.IsEmpty() == false;
+	if (bValidityIsCached == false)
+	{
+		bAreRequirementsValid = RequiredAllFragments.IsEmpty() == false || RequiredAnyFragments.IsEmpty() == false
+			|| RequiredOptionalFragments.IsEmpty() == false	|| RequiredAllTags.IsEmpty() == false || RequiredAnyTags.IsEmpty() == false;
+		bValidityIsCached = true;
+	}
+	return bAreRequirementsValid;
 }
 
 bool FMassFragmentRequirements::IsEmpty() const
 {
-	return FragmentRequirements.IsEmpty() && ChunkFragmentRequirements.IsEmpty() && ConstSharedFragmentRequirements.IsEmpty() 
-		&& SharedFragmentRequirements.IsEmpty() && RequiredAllTags.IsEmpty() && RequiredAnyTags.IsEmpty() && RequiredNoneTags.IsEmpty();
+	if (bEmptynessIsCached == false)
+	{
+		bAreRequirementsEmpty = FragmentRequirements.IsEmpty() && ChunkFragmentRequirements.IsEmpty() && ConstSharedFragmentRequirements.IsEmpty()
+			&& SharedFragmentRequirements.IsEmpty() && RequiredAllTags.IsEmpty() && RequiredAnyTags.IsEmpty() && RequiredNoneTags.IsEmpty();
+		bEmptynessIsCached = true;
+	}
+
+	return bAreRequirementsEmpty;
 }
 
 bool FMassFragmentRequirements::DoesArchetypeMatchRequirements(const FMassArchetypeHandle& ArchetypeHandle) const
diff --git a/Engine/Plugins/Runtime/MassEntity/Source/MassEntity/Public/MassRequirements.h b/Engine/Plugins/Runtime/MassEntity/Source/MassEntity/Public/MassRequirements.h
--- a/Engine/Plugins/Runtime/MassEntity/Source/MassEntity/Public/MassRequirements.h
+++ b/Engine/Plugins/Runtime/MassEntity/Source/MassEntity/Public/MassRequirements.h
@@ -366,9 +366,12 @@ public:
 
 	void Reset();
 
-	/** The function validates requirements we make for queries. See the FMassFragmentRequirements struct description for details.
-	 *  Note that this function is non-trivial and end users are not expected to need to use it. 
-	 *  @return whether this query's requirements follow the rules. */
+	/** 
+	 * The function validates requirements we make for queries. See the FMassFragmentRequirements struct description for details.
+	 * Even though the code of the function is non trivial the consecutive calls will be essentially free due to the result 
+	 * being cached (note that the caching gets invalidated if the composition changes).
+	 * @return whether this query's requirements follow the rules.
+	 */
 	bool CheckValidity() const;
 
 	TConstArrayView<FMassFragmentRequirementDescription> GetFragmentRequirements() const { return FragmentRequirements; }
@@ -399,6 +402,8 @@ protected:
 	void SortRequirements();
 
 	FORCEINLINE void IncrementChangeCounter() { ++IncrementalChangesCount; }
+	void ConsumeIncrementalChangesCount() { IncrementalChangesCount = 0; }
+	bool HasIncrementalChanges() const{ return IncrementalChangesCount > 0; }
 
 protected:
 	friend FMassRequirementAccessDetector;
@@ -421,9 +426,14 @@ protected:
 	FMassSharedFragmentBitSet RequiredOptionalSharedFragments;
 	FMassSharedFragmentBitSet RequiredNoneSharedFragments;
 
-	uint32 IncrementalChangesCount = 0;
-
 private:
+	mutable uint16 bValidityIsCached : 1 = false;
+	mutable uint16 bAreRequirementsValid : 1 = false;
+	mutable uint16 bEmptynessIsCached : 1 = false;
+	mutable uint16 bAreRequirementsEmpty: 1 = false;
+
+	uint16 IncrementalChangesCount = 0;
+
 	bool bRequiresGameThreadExecution = false;
 };
 
diff --git a/Engine/Plugins/Runtime/MassEntity/Source/MassEntityTestSuite/Private/MassCompositeProcessorTest.cpp b/Engine/Plugins/Runtime/MassEntity/Source/MassEntityTestSuite/Private/MassCompositeProcessorTest.cpp
--- a/Engine/Plugins/Runtime/MassEntity/Source/MassEntityTestSuite/Private/MassCompositeProcessorTest.cpp
+++ b/Engine/Plugins/Runtime/MassEntity/Source/MassEntityTestSuite/Private/MassCompositeProcessorTest.cpp
@@ -54,6 +54,8 @@ struct FCompositeProcessorTest_NoWork : FEntityTestBase
 			{
 				UMassTestProcessorBase* Processor = NewObject<UMassTestProcessorBase>();
 				Processor->ForEachEntityChunkExecutionFunction = QueryExecFunction;
+				// need to set up some requirements to make EntityQuery valid
+				Processor->EntityQuery.AddRequirement<FTestFragment_Float>(EMassFragmentAccess::ReadOnly);
 				Processors.Add(Processor);
 			}
 
diff --git a/Engine/Plugins/Runtime/MassEntity/Source/MassEntityTestSuite/Private/MassExecutionTest.cpp b/Engine/Plugins/Runtime/MassEntity/Source/MassEntityTestSuite/Private/MassExecutionTest.cpp
--- a/Engine/Plugins/Runtime/MassEntity/Source/MassEntityTestSuite/Private/MassExecutionTest.cpp
+++ b/Engine/Plugins/Runtime/MassEntity/Source/MassEntityTestSuite/Private/MassExecutionTest.cpp
@@ -108,8 +108,10 @@ struct FExecution_SingleValidProcessor : FExecutionTestBase
 		FMassProcessingContext ProcessingContext(EntityManager, DeltaSeconds);
 		UMassTestProcessorBase* Processor = NewObject<UMassTestProcessorBase>();
 		check(Processor);
+		// need to set up some requirements to make EntityQuery valid
+		Processor->EntityQuery.AddRequirement<FTestFragment_Float>(EMassFragmentAccess::ReadOnly);
 
-		// nothing should break. The actual result of processing is getting tested in MassProcessorTests.cpp
+		// nothing should break. The actual result of processing is getting tested in MassProcessorTest.cpp
 		UE::Mass::Executor::Run(*Processor, ProcessingContext);
 		return true;
 	}
@@ -149,6 +151,8 @@ struct FExecution_Sparse : FEntityTestBase
 		FMassProcessingContext ProcessingContext(*EntityManager, DeltaSeconds);
 		UMassTestProcessorBase* Processor = NewObject<UMassTestProcessorBase>();
 		check(Processor);
+		// need to set up some requirements to make EntityQuery valid
+		Processor->EntityQuery.AddRequirement<FTestFragment_Float>(EMassFragmentAccess::ReadOnly);
 
 		FMassRuntimePipeline Pipeline;
 		{
@@ -158,7 +162,7 @@ struct FExecution_Sparse : FEntityTestBase
 		}
 
 		FMassArchetypeEntityCollection EntityCollection(FloatsArchetype);
-		// nothing should break. The actual result of processing is getting tested in MassProcessorTests.cpp
+		// nothing should break. The actual result of processing is getting tested in MassProcessorTest.cpp
 		
 		UE::Mass::Executor::RunSparse(Pipeline, ProcessingContext, EntityCollection);
 
