Added a non-strictly-typed version of FMassFragmentRequirements::AddConstSharedRequirement

diff --git a/Engine/Plugins/Runtime/MassEntity/Source/MassEntity/Public/MassRequirements.h b/Engine/Plugins/Runtime/MassEntity/Source/MassEntity/Public/MassRequirements.h
--- a/Engine/Plugins/Runtime/MassEntity/Source/MassEntity/Public/MassRequirements.h
+++ b/Engine/Plugins/Runtime/MassEntity/Source/MassEntity/Public/MassRequirements.h
@@ -232,7 +232,7 @@ public:
 
 	void AddTagRequirement(const UScriptStruct& TagType, const EMassFragmentPresence Presence)
 	{
-		checkfSlow(int(Presence) < int(EMassFragmentPresence::Optional), TEXT("Optional and MAX presence are not valid calues for AddTagRequirement"));
+		checkf(int(Presence) < int(EMassFragmentPresence::Optional), TEXT("Optional and MAX presence are not valid calues for AddTagRequirement"));
 		switch (Presence)
 		{
 		case EMassFragmentPresence::All:
@@ -251,7 +251,7 @@ public:
 	template<typename T>
 	FMassFragmentRequirements& AddTagRequirement(const EMassFragmentPresence Presence)
 	{
-		checkfSlow(int(Presence) < int(EMassFragmentPresence::Optional), TEXT("Optional and MAX presence are not valid calues for AddTagRequirement"));
+		checkf(int(Presence) < int(EMassFragmentPresence::Optional), TEXT("Optional and MAX presence are not valid calues for AddTagRequirement"));
 		static_assert(TIsDerivedFrom<T, FMassTag>::IsDerived, "Given struct doesn't represent a valid tag type. Make sure to inherit from FMassFragment or one of its child-types.");
 		switch (Presence)
 		{
@@ -284,7 +284,7 @@ public:
 		static_assert(TIsDerivedFrom<T, FMassChunkFragment>::IsDerived, "Given struct doesn't represent a valid chunk fragment type. Make sure to inherit from FMassChunkFragment or one of its child-types.");
 		checkf(ChunkFragmentRequirements.FindByPredicate([](const FMassFragmentRequirementDescription& Item) { return Item.StructType == T::StaticStruct(); }) == nullptr
 			, TEXT("Duplicated requirements are not supported. %s already present"), *T::StaticStruct()->GetName());
-		checkfSlow(Presence != EMassFragmentPresence::Any, TEXT("\'Any\' is not a valid Presence value for AddChunkRequirement."));
+		checkf(Presence != EMassFragmentPresence::Any, TEXT("\'Any\' is not a valid Presence value for AddChunkRequirement."));
 
 		switch (Presence)
 		{
@@ -310,7 +310,7 @@ public:
 		static_assert(TIsDerivedFrom<T, FMassSharedFragment>::IsDerived, "Given struct doesn't represent a valid shared fragment type. Make sure to inherit from FMassSharedFragment or one of its child-types.");
 		checkf(ConstSharedFragmentRequirements.FindByPredicate([](const FMassFragmentRequirementDescription& Item) { return Item.StructType == T::StaticStruct(); }) == nullptr
 			, TEXT("Duplicated requirements are not supported. %s already present"), *T::StaticStruct()->GetName());
-		checkfSlow(Presence != EMassFragmentPresence::Any, TEXT("\'Any\' is not a valid Presence value for AddSharedRequirement."));
+		checkf(Presence != EMassFragmentPresence::Any, TEXT("\'Any\' is not a valid Presence value for AddConstSharedRequirement."));
 
 		switch (Presence)
 		{
@@ -330,13 +330,43 @@ public:
 		return *this;
 	}
 
+	FMassFragmentRequirements& AddConstSharedRequirement(const UScriptStruct* FragmentType, const EMassFragmentPresence Presence = EMassFragmentPresence::All)
+	{
+		if (!ensureMsgf(FragmentType->IsChildOf(FMassSharedFragment::StaticStruct())
+			, TEXT("Given struct doesn't represent a valid shared fragment type. Make sure to inherit from FMassSharedFragment or one of its child-types.")))
+		{
+			return *this;
+		}
+
+		checkf(ConstSharedFragmentRequirements.FindByPredicate([FragmentType](const FMassFragmentRequirementDescription& Item) { return Item.StructType == FragmentType; }) == nullptr
+			, TEXT("Duplicated requirements are not supported. %s already present"), *FragmentType->GetName());
+		checkf(Presence != EMassFragmentPresence::Any, TEXT("\'Any\' is not a valid Presence value for AddConstSharedRequirement."));
+
+		switch (Presence)
+		{
+		case EMassFragmentPresence::All:
+			RequiredAllSharedFragments.Add(*FragmentType);
+			ConstSharedFragmentRequirements.Emplace(FragmentType, EMassFragmentAccess::ReadOnly, Presence);
+			break;
+		case EMassFragmentPresence::Optional:
+			RequiredOptionalSharedFragments.Add(*FragmentType);
+			ConstSharedFragmentRequirements.Emplace(FragmentType, EMassFragmentAccess::ReadOnly, Presence);
+			break;
+		case EMassFragmentPresence::None:
+			RequiredNoneSharedFragments.Add(*FragmentType);
+			break;
+		}
+		IncrementChangeCounter();
+		return *this;
+	}
+
 	template<typename T>
 	FMassFragmentRequirements& AddSharedRequirement(const EMassFragmentAccess AccessMode, const EMassFragmentPresence Presence = EMassFragmentPresence::All)
 	{
 		static_assert(TIsDerivedFrom<T, FMassSharedFragment>::IsDerived, "Given struct doesn't represent a valid shared fragment type. Make sure to inherit from FMassSharedFragment or one of its child-types.");
 		checkf(SharedFragmentRequirements.FindByPredicate([](const FMassFragmentRequirementDescription& Item) { return Item.StructType == T::StaticStruct(); }) == nullptr
 			, TEXT("Duplicated requirements are not supported. %s already present"), *T::StaticStruct()->GetName());
-		checkfSlow(Presence != EMassFragmentPresence::Any, TEXT("\'Any\' is not a valid Presence value for AddSharedRequirement."));
+		checkf(Presence != EMassFragmentPresence::Any, TEXT("\'Any\' is not a valid Presence value for AddSharedRequirement."));
 
 		switch (Presence)
 		{
