Added means of checking if a given shared fragment type is present in a given FMassArchetypeSharedFragmentValues instance

diff --git a/Engine/Plugins/Runtime/MassEntity/Source/MassEntity/Public/MassEntityTypes.h b/Engine/Plugins/Runtime/MassEntity/Source/MassEntity/Public/MassEntityTypes.h
--- a/Engine/Plugins/Runtime/MassEntity/Source/MassEntity/Public/MassEntityTypes.h
+++ b/Engine/Plugins/Runtime/MassEntity/Source/MassEntity/Public/MassEntityTypes.h
@@ -242,8 +242,6 @@ struct FMassArchetypeCompositionDescriptor
 
 struct MASSENTITY_API FMassArchetypeSharedFragmentValues
 {
-	static constexpr uint32 EmptyInstanceHash = 0;
-
 	FMassArchetypeSharedFragmentValues() = default;
 	FMassArchetypeSharedFragmentValues(const FMassArchetypeSharedFragmentValues& OtherFragmentValues) = default;
 	FMassArchetypeSharedFragmentValues(FMassArchetypeSharedFragmentValues&& OtherFragmentValues) = default;
@@ -265,6 +263,17 @@ struct MASSENTITY_API FMassArchetypeSharedFragmentValues
 		return GetTypeHash(*this) == GetTypeHash(OtherSharedFragmentValues);
 	}
 
+	FORCEINLINE bool ContainsType(const UScriptStruct* FragmentType) const
+	{
+		return FragmentType != nullptr && SharedFragmentBitSet.Contains(*FragmentType);
+	}
+
+	template<typename T>
+	FORCEINLINE bool ContainsType() const
+	{
+		return SharedFragmentBitSet.Contains(*T::StaticStruct());
+	}
+
 	FORCEINLINE FConstSharedStruct& AddConstSharedFragment(const FConstSharedStruct& Fragment)
 	{
 		DirtyHashCache();
diff --git a/Engine/Plugins/Runtime/MassEntity/Source/MassEntityTestSuite/Private/MassSharedFragmentsTest.cpp b/Engine/Plugins/Runtime/MassEntity/Source/MassEntityTestSuite/Private/MassSharedFragmentsTest.cpp
new file mode 100644
--- /dev/null
+++ b/Engine/Plugins/Runtime/MassEntity/Source/MassEntityTestSuite/Private/MassSharedFragmentsTest.cpp
@@ -0,0 +1,58 @@
+// Copyright Epic Games, Inc. All Rights Reserved.
+
+#include "AITestsCommon.h"
+#include "MassEntityManager.h"
+#include "MassEntityTypes.h"
+#include "MassEntityTestTypes.h"
+
+#define LOCTEXT_NAMESPACE "MassTest"
+
+UE_DISABLE_OPTIMIZATION_SHIP
+
+namespace FMassEntityTest
+{
+
+struct FSharedFragmentValues_Contains : FExecutionTestBase
+{
+	virtual bool InstantTest() override
+	{
+		FMassArchetypeSharedFragmentValues Values;
+
+		AITEST_FALSE("Empty FMassArchetypeSharedFragmentValues should fail ContainsType tests"
+			, Values.ContainsType(FTestSharedFragment_Int::StaticStruct()));
+		AITEST_FALSE("Empty FMassArchetypeSharedFragmentValues should fail ContainsType tests"
+			, Values.ContainsType<FTestSharedFragment_Int>());
+
+		{
+			FTestSharedFragment_Int FragmentInstance(1);
+			FConstSharedStruct SharedFragmentInstance;
+			SharedFragmentInstance.InitializeAs(FragmentInstance);
+			Values.AddConstSharedFragment(SharedFragmentInstance);
+		}
+
+		AITEST_TRUE("Empty FMassArchetypeSharedFragmentValues should fail ContainsType tests"
+			, Values.ContainsType(FTestSharedFragment_Int::StaticStruct()));
+		AITEST_TRUE("Empty FMassArchetypeSharedFragmentValues should fail ContainsType tests"
+			, Values.ContainsType<FTestSharedFragment_Int>());
+
+		{
+			FTestSharedFragment_Float FragmentInstance(1.f);
+			FConstSharedStruct SharedFragmentInstance;
+			SharedFragmentInstance.InitializeAs(FragmentInstance);
+			Values.AddConstSharedFragment(SharedFragmentInstance);
+		}
+		AITEST_TRUE("Empty FMassArchetypeSharedFragmentValues should fail ContainsType tests"
+			, Values.ContainsType(FTestSharedFragment_Float::StaticStruct()));
+		AITEST_TRUE("Empty FMassArchetypeSharedFragmentValues should fail ContainsType tests"
+			, Values.ContainsType<FTestSharedFragment_Float>());
+
+		return true;
+	}
+};
+IMPLEMENT_AI_INSTANT_TEST(FSharedFragmentValues_Contains, "System.Mass.SharedFragments.Contains");
+
+} // FMassEntityTest
+
+UE_ENABLE_OPTIMIZATION_SHIP
+
+#undef LOCTEXT_NAMESPACE
