Changed how Mass's shared fragments hashing works.
- hidden the API requiring users to provide explicit hash to use
- switched fragment hashing to use 64 bit city hashing.

diff --git a/Engine/Plugins/AI/MassCrowd/Source/MassCrowd/Private/MassCrowdServerRepresentationTrait.cpp b/Engine/Plugins/AI/MassCrowd/Source/MassCrowd/Private/MassCrowdServerRepresentationTrait.cpp
--- a/Engine/Plugins/AI/MassCrowd/Source/MassCrowd/Private/MassCrowdServerRepresentationTrait.cpp
+++ b/Engine/Plugins/AI/MassCrowd/Source/MassCrowd/Private/MassCrowdServerRepresentationTrait.cpp
@@ -42,8 +42,7 @@ void UMassCrowdServerRepresentationTrait::BuildTemplate(FMassEntityTemplateBuild
 
 	FMassRepresentationSubsystemSharedFragment SubsystemSharedFragment;
 	SubsystemSharedFragment.RepresentationSubsystem = RepresentationSubsystem;
-	uint32 SubsystemHash = UE::StructUtils::GetStructCrc32(FConstStructView::Make(SubsystemSharedFragment));
-	FSharedStruct SubsystemFragment = EntityManager.GetOrCreateSharedFragmentByHash<FMassRepresentationSubsystemSharedFragment>(SubsystemHash, SubsystemSharedFragment);
+	FSharedStruct SubsystemFragment = EntityManager.GetOrCreateSharedFragment<FMassRepresentationSubsystemSharedFragment>(SubsystemSharedFragment);
 	BuildContext.AddSharedFragment(SubsystemFragment);
 
 	FConstSharedStruct ParamsFragment = EntityManager.GetOrCreateConstSharedFragment(Params);
diff --git a/Engine/Plugins/Runtime/InstancedActors/Source/InstancedActors/Private/InstancedActorsData.cpp b/Engine/Plugins/Runtime/InstancedActors/Source/InstancedActors/Private/InstancedActorsData.cpp
--- a/Engine/Plugins/Runtime/InstancedActors/Source/InstancedActors/Private/InstancedActorsData.cpp
+++ b/Engine/Plugins/Runtime/InstancedActors/Source/InstancedActors/Private/InstancedActorsData.cpp
@@ -408,8 +408,7 @@ void UInstancedActorsData::DespawnEntities()
 
 	FInstancedActorsDataSharedFragment ManagerFragment;
 	ManagerFragment.InstanceData = this;
-	const uint32 FragmentHash = UE::StructUtils::GetStructCrc32(FConstStructView::Make(ManagerFragment));
-	FSharedStruct SharedFragmentInstance = MassEntityManager.GetOrCreateSharedFragmentByHash<FInstancedActorsDataSharedFragment>(FragmentHash, ManagerFragment);
+	FSharedStruct SharedFragmentInstance = MassEntityManager.GetOrCreateSharedFragment<FInstancedActorsDataSharedFragment>(ManagerFragment);
 	FInstancedActorsDataSharedFragment* AsSharedManagerFragment = SharedFragmentInstance.GetPtr<FInstancedActorsDataSharedFragment>();
 	if (ensure(AsSharedManagerFragment) && AsSharedManagerFragment->InstanceData == this)
 	{
diff --git a/Engine/Plugins/Runtime/InstancedActors/Source/InstancedActors/Private/InstancedActorsVisualizationTrait.cpp b/Engine/Plugins/Runtime/InstancedActors/Source/InstancedActors/Private/InstancedActorsVisualizationTrait.cpp
--- a/Engine/Plugins/Runtime/InstancedActors/Source/InstancedActors/Private/InstancedActorsVisualizationTrait.cpp
+++ b/Engine/Plugins/Runtime/InstancedActors/Source/InstancedActors/Private/InstancedActorsVisualizationTrait.cpp
@@ -87,8 +87,7 @@ void UInstancedActorsVisualizationTrait::BuildTemplate(FMassEntityTemplateBuildC
 
 	FInstancedActorsDataSharedFragment ManagerSharedFragment;
 	ManagerSharedFragment.InstanceData = InstanceData;
-	const uint32 SubsystemHash = UE::StructUtils::GetStructCrc32(FConstStructView::Make(ManagerSharedFragment));
-	FSharedStruct SubsystemFragment = EntityManager.GetOrCreateSharedFragmentByHash<FInstancedActorsDataSharedFragment>(SubsystemHash, ManagerSharedFragment);
+	FSharedStruct SubsystemFragment = EntityManager.GetOrCreateSharedFragment<FInstancedActorsDataSharedFragment>(ManagerSharedFragment);
 
 	FInstancedActorsDataSharedFragment* AsShared = SubsystemFragment.GetPtr<FInstancedActorsDataSharedFragment>();
 	if (ensure(AsShared))
diff --git a/Engine/Plugins/Runtime/MassEntity/Source/MassEntity/Public/MassEntityManager.h b/Engine/Plugins/Runtime/MassEntity/Source/MassEntity/Public/MassEntityManager.h
--- a/Engine/Plugins/Runtime/MassEntity/Source/MassEntity/Public/MassEntityManager.h
+++ b/Engine/Plugins/Runtime/MassEntity/Source/MassEntity/Public/MassEntityManager.h
@@ -437,44 +437,43 @@ public:
 	 */
 	void AppendCommands(TSharedPtr<FMassCommandBuffer>& InOutCommandBuffer);
 
-	/**
-	 * Shared fragment creation methods
+
+	/** 
+	 * @param HashOverride if provided will be used instead of the auto-calculated hash
 	 */
 	template<typename T>
-	const FConstSharedStruct& GetOrCreateConstSharedFragmentByHash(const uint32 Hash, const T& Fragment)
+	const FConstSharedStruct& GetOrCreateConstSharedFragment(const T& Fragment, const uint64 HashOverride = 0)
 	{
-		static_assert(TIsDerivedFrom<T, FMassSharedFragment>::IsDerived, "Given struct doesn't represent a valid shared fragment type. Make sure to inherit from FMassSharedFragment or one of its child-types.");
-		int32& Index = ConstSharedFragmentsMap.FindOrAddByHash(Hash, Hash, INDEX_NONE);
-		if (Index == INDEX_NONE)
-		{
-			Index = ConstSharedFragments.Add(FSharedStruct::Make(Fragment));
-		}
-		return ConstSharedFragments[Index];
+		static_assert(TIsDerivedFrom<T, FMassSharedFragment>::IsDerived, "Given struct doesn't represent a valid const shared fragment type. Make sure to inherit from FMassSharedFragment or one of its child-types.");
+
+		const uint64 Hash = HashOverride ? HashOverride : UE::StructUtils::GetStructHash64(FConstStructView::Make(Fragment));
+		return InternalGetOrCreateConstSharedFragment(Hash, Fragment);
 	}
 
+	/**
+	 * @param HashOverride if provided will be used instead of the auto-calculated hash
+	 */
 	template<typename T>
-	const FConstSharedStruct& GetOrCreateConstSharedFragment(const T& Fragment)
+	const FSharedStruct& GetOrCreateSharedFragment(const T& Fragment, const uint64 HashOverride = 0)
 	{
-		const uint32 Hash = UE::StructUtils::GetStructCrc32(FConstStructView::Make(Fragment));
-		return GetOrCreateConstSharedFragmentByHash(Hash, Fragment);
+		static_assert(TIsDerivedFrom<T, FMassSharedFragment>::IsDerived, "Given struct doesn't represent a valid shared fragment type. Make sure to inherit from FMassSharedFragment or one of its child-types.");
+
+		const uint64 Hash = HashOverride ? HashOverride : UE::StructUtils::GetStructHash64(FConstStructView::Make(Fragment));
+		return InternalGetOrCreateSharedFragment(Hash, Fragment);
 	}
 
+	/**
+	 * Note that this version doesn't let the used override the Hash. Use the other version if you need that functionality
+	 */
 	template<typename T, typename... TArgs>
-	const FSharedStruct& GetOrCreateSharedFragmentByHash(const uint32 Hash, TArgs&&... InArgs)
+	const FSharedStruct& GetOrCreateSharedFragment(TArgs&&... InArgs)
 	{
 		static_assert(TIsDerivedFrom<T, FMassSharedFragment>::IsDerived, "Given struct doesn't represent a valid shared fragment type. Make sure to inherit from FMassSharedFragment or one of its child-types.");
 
-		int32& Index = SharedFragmentsMap.FindOrAddByHash(Hash, Hash, INDEX_NONE);
-		if (Index == INDEX_NONE)
-		{
-			Index = SharedFragments.Add(FSharedStruct::Make<T>(Forward<TArgs>(InArgs)...));
-			// note that even though we're copying the freshly created FSharedStruct instance it's perfectly fine since 
-			// FSharedStruct do guarantee there's not going to be data duplication (via a member shared pointer to hosted data)
-			TArray<FSharedStruct>& InstancesOfType = SharedFragmentsTypeMap.FindOrAdd(T::StaticStruct(), {});
-			InstancesOfType.Add(SharedFragments[Index]);
-		}
+		T Fragment(Forward<TArgs>(InArgs)...);
+		const uint64 Hash = UE::StructUtils::GetStructHash64(FConstStructView::Make(Fragment));
 
-		return SharedFragments[Index];
+		return InternalGetOrCreateSharedFragment(Hash, MoveTemp(Fragment));
 	}
 
 	template<typename T>
@@ -548,6 +547,35 @@ protected:
 	void InternalAppendFragmentsAndTagsToArchetypeCompositionDescriptor(FMassArchetypeCompositionDescriptor& InOutComposition,
 		TConstArrayView<const UScriptStruct*> FragmentsAndTagsList) const;
 
+	template<typename T>
+	const FConstSharedStruct& InternalGetOrCreateConstSharedFragment(const uint64 Hash, const T& Fragment)
+	{
+		static_assert(TIsDerivedFrom<T, FMassSharedFragment>::IsDerived, "Given struct doesn't represent a valid shared fragment type. Make sure to inherit from FMassSharedFragment or one of its child-types.");
+		int32& Index = ConstSharedFragmentsMap.FindOrAdd(Hash, INDEX_NONE);
+		if (Index == INDEX_NONE)
+		{
+			Index = ConstSharedFragments.Add(FSharedStruct::Make(Fragment));
+		}
+		return ConstSharedFragments[Index];
+	}
+
+	template<typename T>
+	const FSharedStruct& InternalGetOrCreateSharedFragment(const uint64 Hash, T&& Fragment)
+	{
+		int32& Index = SharedFragmentsMap.FindOrAdd(Hash, INDEX_NONE);
+		if (Index == INDEX_NONE)
+		{
+			Index = SharedFragments.Num();
+			FSharedStruct& CreatedStruct = SharedFragments.Add_GetRef(FSharedStruct::Make(Fragment));
+			// note that even though we're copying the freshly created FSharedStruct instance it's perfectly fine since 
+			// FSharedStruct do guarantee there's not going to be data duplication (via a member shared pointer to hosted data)
+			TArray<FSharedStruct>& InstancesOfType = SharedFragmentsTypeMap.FindOrAdd(std::remove_const_t<typename TRemoveReference<T>::Type>::StaticStruct(), {});
+			InstancesOfType.Add(CreatedStruct);
+		}
+
+		return SharedFragments[Index];
+	}
+
 private:
 	void InternalBuildEntity(FMassEntityHandle Entity, const FMassArchetypeHandle& ArchetypeHandle, const FMassArchetypeSharedFragmentValues& SharedFragmentValues);
 	void InternalReleaseEntity(FMassEntityHandle Entity);
@@ -589,7 +617,7 @@ private:
 	uint32 ArchetypeDataVersion = 0;
 
 	// Map of hash of sorted fragment list to archetypes with that hash
-	TMap<uint32, TArray<TSharedPtr<FMassArchetypeData>>> FragmentHashToArchetypeMap;
+	TMap<uint64, TArray<TSharedPtr<FMassArchetypeData>>> FragmentHashToArchetypeMap;
 
 	// Map to list of archetypes that contain the specified fragment type
 	TMap<const UScriptStruct*, TArray<TSharedPtr<FMassArchetypeData>>> FragmentTypeToArchetypeMap;
@@ -601,13 +629,13 @@ private:
 	// Shared fragments
 	TArray<FConstSharedStruct> ConstSharedFragments;
 	// Hash/Index in array pair
-	TMap<uint32, int32> ConstSharedFragmentsMap;
+	TMap<uint64, int32> ConstSharedFragmentsMap;
 
 	TArray<FSharedStruct> SharedFragments;
 	// Hash/Index in array pair, indices point at SharedFragments
-	TMap<uint32, int32> SharedFragmentsMap;
+	TMap<uint64, int32> SharedFragmentsMap;
 	// Maps specific struct type to a collection of FSharedStruct instances of that type
-	TMap<UScriptStruct*, TArray<FSharedStruct>> SharedFragmentsTypeMap;
+	TMap<const UScriptStruct*, TArray<FSharedStruct>> SharedFragmentsTypeMap;
 
 	FMassObserverManager ObserverManager;
 
@@ -644,6 +672,18 @@ public:
 	UE_DEPRECATED(5.3, "This Flavor of BatchBuildEntities is deprecated. Use the one with FMassArchetypeCreationParams parameter instead.")
 	void BatchBuildEntities(const FMassArchetypeEntityCollectionWithPayload& EncodedEntitiesWithPayload, FMassArchetypeCompositionDescriptor&& Composition
 		, const FMassArchetypeSharedFragmentValues& SharedFragmentValues, const FName ArchetypeDebugName);
+	template<typename T>
+	UE_DEPRECATED(5.5, "GetOrCreateConstSharedFragmentByHash is deprecated. Use GetOrCreateConstSharedFragment instead.")
+	const FConstSharedStruct& GetOrCreateConstSharedFragmentByHash(const uint64 Hash, const T& Fragment)
+	{
+		return InternalGetOrCreateConstSharedFragment(Hash, Fragment);
+	}
+	template<typename T, typename... TArgs>
+	UE_DEPRECATED(5.5, "GetOrCreateSharedFragmentByHash is deprecated. Use GetOrCreateSharedFragment instead.")
+	const FSharedStruct& GetOrCreateSharedFragmentByHash(const uint32 Hash, TArgs&&... InArgs)
+	{
+		return InternalGetOrCreateSharedFragment(Hash, FSharedStruct::Make<T>(Forward<TArgs>(InArgs)...));
+	}
 };
 
 #if UE_ENABLE_INCLUDE_ORDER_DEPRECATED_IN_5_2
diff --git a/Engine/Plugins/Runtime/MassGameplay/Source/MassGameplayTestSuite/Private/MassEntityTemplateTest.cpp b/Engine/Plugins/Runtime/MassGameplay/Source/MassGameplayTestSuite/Private/MassEntityTemplateTest.cpp
--- a/Engine/Plugins/Runtime/MassGameplay/Source/MassGameplayTestSuite/Private/MassEntityTemplateTest.cpp
+++ b/Engine/Plugins/Runtime/MassGameplay/Source/MassGameplayTestSuite/Private/MassEntityTemplateTest.cpp
@@ -64,11 +64,11 @@ struct FEntityTemplate_Composition : FEntityTemplateBase
 		TemplateData.AddFragment<FTestFragment_Float>();
 		TemplateData.AddChunkFragment<FTestChunkFragment_Int>();
 		{
-			FSharedStruct SharedFragment = EntityManager->GetOrCreateSharedFragmentByHash<FTestSharedFragment_Int>(0);
+			FSharedStruct SharedFragment = EntityManager->GetOrCreateSharedFragment<FTestSharedFragment_Int>();
 			TemplateData.AddSharedFragment(SharedFragment);
 		}
 		{
-			FSharedStruct SharedFragment = EntityManager->GetOrCreateSharedFragmentByHash<FTestSharedFragment_Float>(1);
+			FSharedStruct SharedFragment = EntityManager->GetOrCreateSharedFragment<FTestSharedFragment_Float>();
 			TemplateData.AddConstSharedFragment(SharedFragment);
 		}
 		TemplateData.Sort();
diff --git a/Engine/Plugins/Runtime/MassGameplay/Source/MassLOD/Private/MassLODTrait.cpp b/Engine/Plugins/Runtime/MassGameplay/Source/MassLOD/Private/MassLODTrait.cpp
--- a/Engine/Plugins/Runtime/MassGameplay/Source/MassLOD/Private/MassLODTrait.cpp
+++ b/Engine/Plugins/Runtime/MassGameplay/Source/MassLOD/Private/MassLODTrait.cpp
@@ -51,8 +51,7 @@ void UMassSimulationLODTrait::BuildTemplate(FMassEntityTemplateBuildContext& Bui
 	FConstSharedStruct ParamsFragment = EntityManager.GetOrCreateConstSharedFragment(Params);
 	BuildContext.AddConstSharedFragment(ParamsFragment);
 
-	uint32 ParamsHash = UE::StructUtils::GetStructCrc32(FConstStructView::Make(Params));
-	FSharedStruct SharedFragment = EntityManager.GetOrCreateSharedFragmentByHash<FMassSimulationLODSharedFragment>(ParamsHash, Params);
+	FSharedStruct SharedFragment = EntityManager.GetOrCreateSharedFragment<FMassSimulationLODSharedFragment>(Params);
 	BuildContext.AddSharedFragment(SharedFragment);
 
 	// Variable ticking from simulation LOD
@@ -64,8 +63,7 @@ void UMassSimulationLODTrait::BuildTemplate(FMassEntityTemplateBuildContext& Bui
 		FConstSharedStruct VariableTickParamsFragment = EntityManager.GetOrCreateConstSharedFragment(VariableTickParams);
 		BuildContext.AddConstSharedFragment(VariableTickParamsFragment);
 
-		uint32 VariableTickParamsHash = UE::StructUtils::GetStructCrc32(FConstStructView::Make(VariableTickParams));
-		FSharedStruct VariableTickSharedFragment = EntityManager.GetOrCreateSharedFragmentByHash<FMassSimulationVariableTickSharedFragment>(VariableTickParamsHash, VariableTickParams);
+		FSharedStruct VariableTickSharedFragment = EntityManager.GetOrCreateSharedFragment<FMassSimulationVariableTickSharedFragment>(VariableTickParams);
 		BuildContext.AddSharedFragment(VariableTickSharedFragment);
 	}
 }
diff --git a/Engine/Plugins/Runtime/MassGameplay/Source/MassReplication/Private/MassReplicationTrait.cpp b/Engine/Plugins/Runtime/MassGameplay/Source/MassReplication/Private/MassReplicationTrait.cpp
--- a/Engine/Plugins/Runtime/MassGameplay/Source/MassReplication/Private/MassReplicationTrait.cpp
+++ b/Engine/Plugins/Runtime/MassGameplay/Source/MassReplication/Private/MassReplicationTrait.cpp
@@ -34,7 +34,6 @@ void UMassReplicationTrait::BuildTemplate(FMassEntityTemplateBuildContext& Build
 	FConstSharedStruct ParamsFragment = EntityManager.GetOrCreateConstSharedFragment(Params);
 	BuildContext.AddConstSharedFragment(ParamsFragment);
 
-	uint32 ParamsHash = UE::StructUtils::GetStructCrc32(FConstStructView::Make(Params));
-	FSharedStruct SharedFragment = EntityManager.GetOrCreateSharedFragmentByHash<FMassReplicationSharedFragment>(ParamsHash, *ReplicationSubsystem, Params);
+	FSharedStruct SharedFragment = EntityManager.GetOrCreateSharedFragment<FMassReplicationSharedFragment>(*ReplicationSubsystem, Params);
 	BuildContext.AddSharedFragment(SharedFragment);
 }
\ No newline at end of file
diff --git a/Engine/Plugins/Runtime/MassGameplay/Source/MassRepresentation/Private/MassDistanceVisualizationTrait.cpp b/Engine/Plugins/Runtime/MassGameplay/Source/MassRepresentation/Private/MassDistanceVisualizationTrait.cpp
--- a/Engine/Plugins/Runtime/MassGameplay/Source/MassRepresentation/Private/MassDistanceVisualizationTrait.cpp
+++ b/Engine/Plugins/Runtime/MassGameplay/Source/MassRepresentation/Private/MassDistanceVisualizationTrait.cpp
@@ -58,8 +58,7 @@ void UMassDistanceVisualizationTrait::BuildTemplate(FMassEntityTemplateBuildCont
 
 	FMassRepresentationSubsystemSharedFragment SubsystemSharedFragment;
 	SubsystemSharedFragment.RepresentationSubsystem = RepresentationSubsystem;
-	uint32 SubsystemHash = UE::StructUtils::GetStructCrc32(FConstStructView::Make(SubsystemSharedFragment));
-	FSharedStruct SubsystemFragment = EntityManager.GetOrCreateSharedFragmentByHash<FMassRepresentationSubsystemSharedFragment>(SubsystemHash, SubsystemSharedFragment);
+	FSharedStruct SubsystemFragment = EntityManager.GetOrCreateSharedFragment<FMassRepresentationSubsystemSharedFragment>(SubsystemSharedFragment);
 	BuildContext.AddSharedFragment(SubsystemFragment);
 
 	if (!Params.RepresentationActorManagementClass)
@@ -81,8 +80,7 @@ void UMassDistanceVisualizationTrait::BuildTemplate(FMassEntityTemplateBuildCont
 	FConstSharedStruct LODParamsFragment = EntityManager.GetOrCreateConstSharedFragment(LODParams);
 	BuildContext.AddConstSharedFragment(LODParamsFragment);
 
-	uint32 LODParamsHash = UE::StructUtils::GetStructCrc32(FConstStructView::Make(LODParams));
-	FSharedStruct LODSharedFragment = EntityManager.GetOrCreateSharedFragmentByHash<FMassDistanceLODSharedFragment>(LODParamsHash, LODParams);
+	FSharedStruct LODSharedFragment = EntityManager.GetOrCreateSharedFragment<FMassDistanceLODSharedFragment>(LODParams);
 	BuildContext.AddSharedFragment(LODSharedFragment);
 
 	BuildContext.AddFragment<FMassRepresentationLODFragment>();
diff --git a/Engine/Plugins/Runtime/MassGameplay/Source/MassRepresentation/Private/MassVisualizationTrait.cpp b/Engine/Plugins/Runtime/MassGameplay/Source/MassRepresentation/Private/MassVisualizationTrait.cpp
--- a/Engine/Plugins/Runtime/MassGameplay/Source/MassRepresentation/Private/MassVisualizationTrait.cpp
+++ b/Engine/Plugins/Runtime/MassGameplay/Source/MassRepresentation/Private/MassVisualizationTrait.cpp
@@ -70,8 +70,7 @@ void UMassVisualizationTrait::BuildTemplate(FMassEntityTemplateBuildContext& Bui
 
 	FMassRepresentationSubsystemSharedFragment SubsystemSharedFragment;
 	SubsystemSharedFragment.RepresentationSubsystem = RepresentationSubsystem;
-	uint32 SubsystemHash = UE::StructUtils::GetStructCrc32(FConstStructView::Make(SubsystemSharedFragment));
-	FSharedStruct SubsystemFragment = EntityManager.GetOrCreateSharedFragmentByHash<FMassRepresentationSubsystemSharedFragment>(SubsystemHash, SubsystemSharedFragment);
+	FSharedStruct SubsystemFragment = EntityManager.GetOrCreateSharedFragment<FMassRepresentationSubsystemSharedFragment>(SubsystemSharedFragment);
 	BuildContext.AddSharedFragment(SubsystemFragment);
 
 	if (!Params.RepresentationActorManagementClass)
@@ -90,11 +89,15 @@ void UMassVisualizationTrait::BuildTemplate(FMassEntityTemplateBuildContext& Bui
 	RepresentationFragment.HighResTemplateActorIndex = HighResTemplateActor.Get() ? RepresentationSubsystem->FindOrAddTemplateActor(HighResTemplateActor.Get()) : INDEX_NONE;
 	RepresentationFragment.LowResTemplateActorIndex = LowResTemplateActor.Get() ? RepresentationSubsystem->FindOrAddTemplateActor(LowResTemplateActor.Get()) : INDEX_NONE;
 
-	FConstSharedStruct LODParamsFragment = EntityManager.GetOrCreateConstSharedFragment(LODParams);
+	// we're calculating a hash of our parameters and use it during shared fragment creation to ensure that 
+	// both LODParamsFragment and LODSharedFragment end up paired. 
+	const uint64 LODParamsHash = UE::StructUtils::GetStructHash64(FConstStructView::Make(LODParams));
+
+	FConstSharedStruct LODParamsFragment = EntityManager.GetOrCreateConstSharedFragment(LODParams, LODParamsHash);
 	BuildContext.AddConstSharedFragment(LODParamsFragment);
 
-	uint32 LODParamsHash = UE::StructUtils::GetStructCrc32(FConstStructView::Make(LODParams));
-	FSharedStruct LODSharedFragment = EntityManager.GetOrCreateSharedFragmentByHash<FMassVisualizationLODSharedFragment>(LODParamsHash, LODParams);
+	FMassVisualizationLODSharedFragment LODFragment(LODParams);
+	FSharedStruct LODSharedFragment = EntityManager.GetOrCreateSharedFragment(LODFragment, LODParamsHash);
 	BuildContext.AddSharedFragment(LODSharedFragment);
 
 	BuildContext.AddFragment<FMassRepresentationLODFragment>();
