Added more functionality to FMassArchetypeSharedFragmentValues like adding and removing elements and fetching data by type.
Use cases include testing and upcoming additions to MassEntityManager supporting runtime adding shared fragments to entities.

Also, added a bunch of unit tests for existing and new functionality

diff --git a/Engine/Plugins/Runtime/MassEntity/Source/MassEntity/MassEntity.Build.cs b/Engine/Plugins/Runtime/MassEntity/Source/MassEntity/MassEntity.Build.cs
--- a/Engine/Plugins/Runtime/MassEntity/Source/MassEntity/MassEntity.Build.cs
+++ b/Engine/Plugins/Runtime/MassEntity/Source/MassEntity/MassEntity.Build.cs
@@ -30,6 +30,13 @@ namespace UnrealBuildTool.Rules
 			{
 				PrivateDependencyModuleNames.Add("MessageLog");
 			}
+
+			if (Target.Configuration != UnrealTargetConfiguration.Shipping
+				&& Target.Configuration != UnrealTargetConfiguration.Test)
+			{
+				// pulling this one in for the testableEnsureMsgf
+				PrivateDependencyModuleNames.Add("AITestSuite");
+			}
 		}
 	}
 }
\ No newline at end of file
diff --git a/Engine/Plugins/Runtime/MassEntity/Source/MassEntity/Private/MassEntityTypes.cpp b/Engine/Plugins/Runtime/MassEntity/Source/MassEntity/Private/MassEntityTypes.cpp
--- a/Engine/Plugins/Runtime/MassEntity/Source/MassEntity/Private/MassEntityTypes.cpp
+++ b/Engine/Plugins/Runtime/MassEntity/Source/MassEntity/Private/MassEntityTypes.cpp
@@ -16,9 +16,48 @@ DEFINE_TYPEBITSET(FMassExternalSubsystemBitSet);
 //-----------------------------------------------------------------------------
 // FMassArchetypeSharedFragmentValues
 //-----------------------------------------------------------------------------
+FConstSharedStruct FMassArchetypeSharedFragmentValues::AddConstSharedFragment(const FConstSharedStruct& Fragment)
+{
+	check(Fragment.IsValid());
+	const UScriptStruct* StructType = Fragment.GetScriptStruct();
+	if (!ensureMsgf(ContainsType(StructType) == false, TEXT("Shared Fragment of type %s already added to FMassArchetypeSharedFragmentValues"), *GetNameSafe(StructType)))
+	{
+		FConstSharedStruct ExistingSharedStruct = GetConstSharedFragmentStruct(StructType);
+		ensureMsgf(ExistingSharedStruct.IsValid(), TEXT("Shared Fragment of type %s already added as NON-CONST shared struct"), *GetNameSafe(StructType));
+		return ExistingSharedStruct;
+	}
+
+	check(StructType);
+	SharedFragmentBitSet.Add(*StructType);
+	FConstSharedStruct& StructInstance = ConstSharedFragments.Add_GetRef(Fragment);
+	DirtyHashCache();
+	return StructInstance;
+}
+
+FSharedStruct FMassArchetypeSharedFragmentValues::AddSharedFragment(const FSharedStruct& Fragment)
+{
+	check(Fragment.IsValid());
+	const UScriptStruct* StructType = Fragment.GetScriptStruct();
+	if (!ensureMsgf(ContainsType(StructType) == false, TEXT("Shared Fragment of type %s already added to FMassArchetypeSharedFragmentValues"), *GetNameSafe(StructType)))
+	{
+		FSharedStruct ExistingSharedStruct = GetSharedFragmentStruct(StructType);
+		ensureMsgf(ExistingSharedStruct.IsValid(), TEXT("Shared Fragment of type %s already added as CONST shared struct"), *GetNameSafe(StructType));
+		return ExistingSharedStruct;
+	}
+
+	check(StructType);
+	SharedFragmentBitSet.Add(*StructType);
+	FSharedStruct& StructInstance = SharedFragments.Add_GetRef(Fragment);
+	DirtyHashCache();
+	return StructInstance;
+}
+
 uint32 FMassArchetypeSharedFragmentValues::CalculateHash() const
 {
-	checkf(bSorted, TEXT("Expecting the containers to be sorted for the hash caluclation to be correct"));
+	if (!testableEnsureMsgf(bSorted, TEXT("Expecting the containers to be sorted for the hash caluclation to be consistent")))
+	{
+		return 0;
+	}
 
 	// Fragments are not part of the uniqueness 
 	uint32 Hash = 0;
@@ -35,6 +74,163 @@ uint32 FMassArchetypeSharedFragmentValues::CalculateHash() const
 	return Hash;
 }
 
+namespace UE::Mass::Private
+{
+	template<typename TSharedStruct>
+	int32 CountInvalid(const TArray<TSharedStruct>& View)
+	{
+		int32 Count = 0;
+		for (const TSharedStruct& SharedStruct : View)
+		{
+			const UScriptStruct* StructType = SharedStruct.GetScriptStruct();
+			Count += StructType ? 0 : 1;
+		}
+		return Count;
+	}
+
+	/** Note that this function assumes that both ViewA and ViewB do not contain duplicates */
+	template<typename TSharedStruct, bool bSkipNulls=true>
+	bool ArraysHaveSameContents(const TArray<TSharedStruct>& ViewA, const TArray<TSharedStruct>& ViewB)
+	{
+		if constexpr (bSkipNulls)
+		{
+			const int32 NullstCountA = CountInvalid(ViewA);
+			const int32 NullstCountB = CountInvalid(ViewB);
+			if (ViewA.Num() - NullstCountA != ViewB.Num() - NullstCountB)
+			{
+				return false;
+			}
+		}
+		else if (ViewA.Num() != ViewB.Num())
+		{
+			return false;
+		}
+
+		for (const TSharedStruct& SharedStruct : ViewA)
+		{
+			const UScriptStruct* StructType = SharedStruct.GetScriptStruct();
+			if constexpr (bSkipNulls)
+			{
+				if (StructType == nullptr)
+				{
+					continue;
+				}
+			}
+			const int32 FragmentIndex = ViewB.IndexOfByPredicate(FStructTypeEqualOperator(StructType));
+			if (FragmentIndex == INDEX_NONE)
+			{
+				return false;
+			}
+			if (ViewB[FragmentIndex].CompareStructValues(SharedStruct) == false)
+			{
+				return false;
+			}
+		}
+
+		return true;
+	}
+}
+
+bool FMassArchetypeSharedFragmentValues::HasSameValues(const FMassArchetypeSharedFragmentValues& Other) const
+{
+	if (SharedFragmentBitSet.IsEquivalent(Other.SharedFragmentBitSet) == false)
+	{
+		return false;
+	}
+
+	return UE::Mass::Private::ArraysHaveSameContents(SharedFragments, Other.GetSharedFragments())
+		&& UE::Mass::Private::ArraysHaveSameContents(ConstSharedFragments, Other.GetConstSharedFragments());
+}
+
+int32 FMassArchetypeSharedFragmentValues::Append(const FMassArchetypeSharedFragmentValues& Other)
+{
+	int32 AddedOrModifiedCount = 0;
+
+	for (const FSharedStruct& SharedStruct : Other.GetSharedFragments())
+	{
+		const UScriptStruct* StructType = SharedStruct.GetScriptStruct();
+		check(StructType);
+		if (SharedFragmentBitSet.Contains(*StructType))
+		{
+			const int32 FragmentIndex = SharedFragments.IndexOfByPredicate(FStructTypeEqualOperator(StructType));
+			if (testableEnsureMsgf(FragmentIndex != INDEX_NONE, TEXT("%hs trying to switch const fragment to non-const, type %s")
+				, __FUNCTION__, *StructType->GetName()))
+			{
+				SharedFragments[FragmentIndex] = SharedStruct;
+				++AddedOrModifiedCount;
+			}
+		}
+		else
+		{
+			SharedFragments.Add(SharedStruct);
+			++AddedOrModifiedCount;
+		}
+	}
+
+	for (const FConstSharedStruct& SharedStruct : Other.GetConstSharedFragments())
+	{
+		const UScriptStruct* StructType = SharedStruct.GetScriptStruct();
+		check(StructType);
+		if (SharedFragmentBitSet.Contains(*StructType))
+		{
+			const int32 FragmentIndex = ConstSharedFragments.IndexOfByPredicate(FStructTypeEqualOperator(StructType));
+			if (testableEnsureMsgf(FragmentIndex != INDEX_NONE, TEXT("%hs trying to switch non-const fragment to const, type %s")
+				, __FUNCTION__, *StructType->GetName()))
+			{
+				ConstSharedFragments[FragmentIndex] = SharedStruct;
+				++AddedOrModifiedCount;
+			}
+		}
+		else
+		{
+			ConstSharedFragments.Add(SharedStruct);
+			++AddedOrModifiedCount;
+		}
+	}
+	
+	SharedFragmentBitSet += Other.SharedFragmentBitSet;
+	DirtyHashCache();
+
+	return AddedOrModifiedCount;
+}
+
+int32 FMassArchetypeSharedFragmentValues::Remove(const FMassSharedFragmentBitSet& SharedFragmentToRemoveBitSet)
+{
+	int32 RemovedCount = 0;
+	FMassSharedFragmentBitSet CommonFragments = (SharedFragmentBitSet & SharedFragmentToRemoveBitSet);
+	FMassSharedFragmentBitSet::FIndexIterator It = CommonFragments.GetIndexIterator();
+	while(It)
+	{
+		const UScriptStruct* StructType = CommonFragments.GetTypeAtIndex(*It);
+		check(StructType);
+
+		const int32 RecularFragmentIndex = SharedFragments.IndexOfByPredicate(FStructTypeEqualOperator(StructType));
+		if (RecularFragmentIndex != INDEX_NONE)
+		{
+			SharedFragments[RecularFragmentIndex].Reset();
+			++RemovedCount;
+		}
+		else
+		{
+			const int32 ConstFragmentIndex = ConstSharedFragments.IndexOfByPredicate(FStructTypeEqualOperator(StructType));
+			if (ConstFragmentIndex != INDEX_NONE)
+			{
+				ConstSharedFragments[ConstFragmentIndex].Reset();
+				++RemovedCount;
+			}
+		}
+
+		++It;
+	}
+
+	if (RemovedCount)
+	{
+		SharedFragmentBitSet -= CommonFragments;
+		DirtyHashCache();
+	}
+	return RemovedCount;
+}
+
 //-----------------------------------------------------------------------------
 // FMassGenericPayloadView
 //-----------------------------------------------------------------------------
diff --git a/Engine/Plugins/Runtime/MassEntity/Source/MassEntity/Public/MassEntityTypes.h b/Engine/Plugins/Runtime/MassEntity/Source/MassEntity/Public/MassEntityTypes.h
--- a/Engine/Plugins/Runtime/MassEntity/Source/MassEntity/Public/MassEntityTypes.h
+++ b/Engine/Plugins/Runtime/MassEntity/Source/MassEntity/Public/MassEntityTypes.h
@@ -9,6 +9,13 @@
 #include "MassExternalSubsystemTraits.h"
 #include "SharedStruct.h"
 #include "MassEntityTypes.generated.h"
+#ifdef WITH_AITESTSUITE
+#include "TestableEnsures.h"
+#else
+#define testableEnsureMsgf ensureMsgf
+#define testableCheckf checkf
+#define testableCheckfReturn(InExpression, ReturnValue, InFormat, ... ) checkf(InExpression, InFormat, ##__VA_ARGS__)
+#endif 
 
 
 MASSENTITY_API DECLARE_LOG_CATEGORY_EXTERN(LogMass, Warning, All);
@@ -240,6 +247,11 @@ struct FMassArchetypeCompositionDescriptor
 	FMassSharedFragmentBitSet SharedFragments;
 };
 
+/** 
+ * Wrapper for const and non-const shared fragment containers that tracks which struct types it holds (via a FMassSharedFragmentBitSet).
+ * Note that having multiple instanced of a given struct type is not supported and Add* functions will fetch the previously 
+ * added fragment instead of adding a new one.
+ */
 struct MASSENTITY_API FMassArchetypeSharedFragmentValues
 {
 	FMassArchetypeSharedFragmentValues() = default;
@@ -263,6 +275,14 @@ struct MASSENTITY_API FMassArchetypeSharedFragmentValues
 		return GetTypeHash(*this) == GetTypeHash(OtherSharedFragmentValues);
 	}
 
+	/** 
+	 * Compares contents of `this` and the Other, and allows different order of elements in both containers.
+	 * Note that the function ignores "nulls", i.e. empty FConstSharedStruct and FSharedStruct instances. The function
+	 * does care however about matching "mode", meaning ConstSharedFragments and SharedFragments arrays are compared
+	 * independently.
+	 */
+	bool HasSameValues(const FMassArchetypeSharedFragmentValues& Other) const;
+
 	FORCEINLINE bool ContainsType(const UScriptStruct* FragmentType) const
 	{
 		return FragmentType != nullptr && SharedFragmentBitSet.Contains(*FragmentType);
@@ -274,21 +294,30 @@ struct MASSENTITY_API FMassArchetypeSharedFragmentValues
 		return SharedFragmentBitSet.Contains(*T::StaticStruct());
 	}
 
-	FORCEINLINE FConstSharedStruct& AddConstSharedFragment(const FConstSharedStruct& Fragment)
-	{
-		DirtyHashCache();
-		check(Fragment.GetScriptStruct());
-		SharedFragmentBitSet.Add(*Fragment.GetScriptStruct());
-		return ConstSharedFragments.Add_GetRef(Fragment);
-	}
+	/** 
+	 * Adds Fragment to the collection. If a fragment of the given FMassSharedFragment subclass has already added then 
+	 * the function will return the previously added instance. In that case the function will also assert if the given type 
+	 * has been added as a CONST shared fragment and if not it will return an empty FConstSharedStruct
+	 */
+	FConstSharedStruct AddConstSharedFragment(const FConstSharedStruct& Fragment);
 
-	FORCEINLINE FSharedStruct AddSharedFragment(const FSharedStruct& Fragment)
-	{
-		DirtyHashCache();
-		check(Fragment.GetScriptStruct());
-		SharedFragmentBitSet.Add(*Fragment.GetScriptStruct());
-		return SharedFragments.Add_GetRef(Fragment);
-	}
+	/**
+	 * Adds Fragment to the collection. If a fragment of the given FMassSharedFragment subclass has already added then
+	 * the function will return the previously added instance. In that case the function will also assert if the given type
+	 * has been added as a NON-CONST shared fragment and if not it will return an empty FSharedStruct
+	 */
+	FSharedStruct AddSharedFragment(const FSharedStruct& Fragment);
+
+	/** 
+	 * Appends contents of Other to `this` instance. All common fragments will get overridden with values in Other.
+	 * Note that changing a fragments "role" (being const or non-const) is not supported and the function will fail an
+	 * ensure when that is attempted.
+	 * @return number of fragments added or changed
+	 */
+	int32 Append(const FMassArchetypeSharedFragmentValues& Other);
+
+	/** Note that the function removes the shared fragments by type */
+	int32 Remove(const FMassSharedFragmentBitSet& SharedFragmentToRemoveBitSet);
 
 	FORCEINLINE const TArray<FConstSharedStruct>& GetConstSharedFragments() const
 	{
@@ -304,11 +333,29 @@ struct MASSENTITY_API FMassArchetypeSharedFragmentValues
 	{
 		return SharedFragments;
 	}
+	
+	FConstSharedStruct GetConstSharedFragmentStruct(const UScriptStruct* StructType) const
+	{
+		const int32 FragmentIndex = ConstSharedFragments.IndexOfByPredicate(FStructTypeEqualOperator(StructType));
+		return FragmentIndex != INDEX_NONE ? ConstSharedFragments[FragmentIndex] : FConstSharedStruct();
+	}
+		
+	FSharedStruct GetSharedFragmentStruct(const UScriptStruct* StructType)
+	{
+		const int32 FragmentIndex = SharedFragments.IndexOfByPredicate(FStructTypeEqualOperator(StructType));
+		return FragmentIndex != INDEX_NONE ? SharedFragments[FragmentIndex] : FSharedStruct();
+	}
+
+	const FMassSharedFragmentBitSet& GetSharedFragmentBitSet() const
+	{
+		return SharedFragmentBitSet;
+	}
 
 	FORCEINLINE void DirtyHashCache()
 	{
 		HashCache = UINT32_MAX;
-		bSorted = false;
+		// we consider a single shared fragment as being "sorted"
+		bSorted = (SharedFragments.Num() + ConstSharedFragments.Num() <= 1) ;
 	}
 
 	FORCEINLINE void CacheHash() const
@@ -346,7 +393,11 @@ struct MASSENTITY_API FMassArchetypeSharedFragmentValues
 
 protected:
 	mutable uint32 HashCache = UINT32_MAX;
-	mutable bool bSorted = true; // When no element in the array, consider already sorted
+	/**
+	 * We consider empty FMassArchetypeSharedFragmentValues a sorted containter.Same goes for a container containing
+	 * a single element, @see DirtyHashCache
+	 */ 
+	mutable bool bSorted = true; 
 	
 	FMassSharedFragmentBitSet SharedFragmentBitSet;
 	TArray<FConstSharedStruct> ConstSharedFragments;
diff --git a/Engine/Plugins/Runtime/MassEntity/Source/MassEntityTestSuite/Private/MassSharedFragmentsTest.cpp b/Engine/Plugins/Runtime/MassEntity/Source/MassEntityTestSuite/Private/MassSharedFragmentsTest.cpp
--- a/Engine/Plugins/Runtime/MassEntity/Source/MassEntityTestSuite/Private/MassSharedFragmentsTest.cpp
+++ b/Engine/Plugins/Runtime/MassEntity/Source/MassEntityTestSuite/Private/MassSharedFragmentsTest.cpp
@@ -4,6 +4,9 @@
 #include "MassEntityManager.h"
 #include "MassEntityTypes.h"
 #include "MassEntityTestTypes.h"
+#include "MassArchetypeData.h"
+#include "MassArchetypeTypes.h"
+#include "Algo/RandomShuffle.h"
 
 #define LOCTEXT_NAMESPACE "MassTest"
 
@@ -11,11 +14,81 @@ UE_DISABLE_OPTIMIZATION_SHIP
 
 namespace FMassEntityTest
 {
+//-----------------------------------------------------------------------------
+// Helpers
+//-----------------------------------------------------------------------------
+template<typename TMassSharedFragmentType>
+const TMassSharedFragmentType* GetConstSharedFragmentPtr(const FMassArchetypeSharedFragmentValues& Values)
+{
+	FConstSharedStruct SharedStruct = Values.GetConstSharedFragmentStruct(TMassSharedFragmentType::StaticStruct());
+	return SharedStruct.IsValid() ? SharedStruct.GetPtr<const TMassSharedFragmentType>() : nullptr;
+}
+
+template<typename TMassSharedFragmentType>
+TMassSharedFragmentType* GetMutableSharedFragmentPtr(FMassArchetypeSharedFragmentValues& Values)
+{
+	FSharedStruct SharedStruct = Values.GetSharedFragmentStruct(TMassSharedFragmentType::StaticStruct());
+	return SharedStruct.IsValid() ? SharedStruct.GetPtr<TMassSharedFragmentType>() : nullptr;
+}
+
+//-----------------------------------------------------------------------------
+// FMassArchetypeSharedFragmentValues Tests
+//-----------------------------------------------------------------------------
+struct FSharedFragmentValues_Create : FExecutionTestBase
+{
+	virtual bool InstantTest() override
+	{
+		constexpr int32 TestIntValue = 59;
+		FMassArchetypeSharedFragmentValues Values;
+
+		{
+			FTestSharedFragment_Int FragmentInstance(TestIntValue);
+			FSharedStruct SharedFragmentInstance = FSharedStruct::Make(FragmentInstance);
+			Values.AddSharedFragment(SharedFragmentInstance);
+		}
+
+		const FTestSharedFragment_Int* ConstInstance = GetConstSharedFragmentPtr<FTestSharedFragment_Int>(Values);
+		FTestSharedFragment_Int* NonConstInstance = GetMutableSharedFragmentPtr<FTestSharedFragment_Int>(Values);
+
+		AITEST_NULL("Fetching fragment as a const shared fragment should fail", ConstInstance);
+		AITEST_NOT_NULL("Fetching fragment as a shared fragment should not fail", NonConstInstance);
+		AITEST_EQUAL("The fetched value should match the expectations", NonConstInstance->Value, TestIntValue);
+
+		return true;
+	}
+};
+IMPLEMENT_AI_INSTANT_TEST(FSharedFragmentValues_Create, "System.Mass.SharedFragments.CreateValue");
+
+struct FSharedFragmentValues_CreateConst : FExecutionTestBase
+{
+	virtual bool InstantTest() override
+	{
+		constexpr int32 TestIntValue = 59;
+		FMassArchetypeSharedFragmentValues Values;
+
+		{
+			FConstSharedStruct SharedFragmentInstance = FConstSharedStruct::Make<FTestSharedFragment_Int>(TestIntValue);
+			Values.AddConstSharedFragment(SharedFragmentInstance);
+		}
+
+		const FTestSharedFragment_Int* ConstInstance = GetConstSharedFragmentPtr<FTestSharedFragment_Int>(Values);
+		FTestSharedFragment_Int* NonConstInstance = GetMutableSharedFragmentPtr<FTestSharedFragment_Int>(Values);
+
+		AITEST_NULL("Fetching fragment as a shared fragment should fail", NonConstInstance);
+		AITEST_NOT_NULL("Fetching fragment as a const shared fragment should not fail", ConstInstance);
+		AITEST_EQUAL("The fetched value should match the expectations", ConstInstance->Value, TestIntValue);
+
+		return true;
+	}
+};
+IMPLEMENT_AI_INSTANT_TEST(FSharedFragmentValues_CreateConst, "System.Mass.SharedFragments.CreateConstValue");
 
 struct FSharedFragmentValues_Contains : FExecutionTestBase
 {
 	virtual bool InstantTest() override
 	{
+		constexpr int32 TestIntValue = 31;
+		constexpr float TestFloatValue = 63.f;
 		FMassArchetypeSharedFragmentValues Values;
 
 		AITEST_FALSE("Empty FMassArchetypeSharedFragmentValues should fail ContainsType tests"
@@ -24,9 +97,7 @@ struct FSharedFragmentValues_Contains : FExecutionTestBase
 			, Values.ContainsType<FTestSharedFragment_Int>());
 
 		{
-			FTestSharedFragment_Int FragmentInstance(1);
-			FConstSharedStruct SharedFragmentInstance;
-			SharedFragmentInstance.InitializeAs(FragmentInstance);
+			FConstSharedStruct SharedFragmentInstance = FConstSharedStruct::Make<FTestSharedFragment_Int>(TestIntValue);
 			Values.AddConstSharedFragment(SharedFragmentInstance);
 		}
 
@@ -36,9 +107,7 @@ struct FSharedFragmentValues_Contains : FExecutionTestBase
 			, Values.ContainsType<FTestSharedFragment_Int>());
 
 		{
-			FTestSharedFragment_Float FragmentInstance(1.f);
-			FConstSharedStruct SharedFragmentInstance;
-			SharedFragmentInstance.InitializeAs(FragmentInstance);
+			FConstSharedStruct SharedFragmentInstance = FConstSharedStruct::Make<FTestSharedFragment_Float>(TestFloatValue);
 			Values.AddConstSharedFragment(SharedFragmentInstance);
 		}
 		AITEST_TRUE("Empty FMassArchetypeSharedFragmentValues should fail ContainsType tests"
@@ -51,6 +120,313 @@ struct FSharedFragmentValues_Contains : FExecutionTestBase
 };
 IMPLEMENT_AI_INSTANT_TEST(FSharedFragmentValues_Contains, "System.Mass.SharedFragments.Contains");
 
+struct FSharedFragmentValues_Append : FExecutionTestBase
+{
+	virtual bool InstantTest() override
+	{
+		constexpr int32 TestIntValue = 31;
+		constexpr float TestFloatValue = 63.f;
+		
+		FMassArchetypeSharedFragmentValues ValuesNonConstInt;
+		ValuesNonConstInt.AddSharedFragment(FSharedStruct::Make<FTestSharedFragment_Int>(TestIntValue));
+		FMassArchetypeSharedFragmentValues ValuesNonConstFloat;
+		ValuesNonConstFloat.AddSharedFragment(FSharedStruct::Make<FTestSharedFragment_Float>(TestFloatValue));
+		FMassArchetypeSharedFragmentValues ValuestNonConstIntFloat;
+		ValuestNonConstIntFloat.AddSharedFragment(FSharedStruct::Make<FTestSharedFragment_Int>(TestIntValue));
+		ValuestNonConstIntFloat.AddSharedFragment(FSharedStruct::Make<FTestSharedFragment_Float>(TestFloatValue));
+
+		// `Appending` int/float values to an new Values instance should result in the same result as `Adding` them
+		{
+			FMassArchetypeSharedFragmentValues Values;
+			Values.Append(ValuesNonConstInt);
+			AITEST_TRUE("#1 Append results should match expectations", Values.HasSameValues(ValuesNonConstInt));
+			Values.Append(ValuesNonConstFloat);
+			AITEST_TRUE("#2 Append results should match expectations", Values.HasSameValues(ValuestNonConstIntFloat));
+		}
+		{
+			FMassArchetypeSharedFragmentValues Values;
+			Values.Append(ValuesNonConstFloat);
+			AITEST_TRUE("#3 Append results should match expectations", Values.HasSameValues(ValuesNonConstFloat));
+			Values.Append(ValuesNonConstInt);
+			AITEST_TRUE("#4 Append results should match expectations", Values.HasSameValues(ValuestNonConstIntFloat));
+		}
+
+		FMassArchetypeSharedFragmentValues ValuesConstInt;
+		ValuesConstInt.AddConstSharedFragment(FConstSharedStruct::Make<FTestSharedFragment_Int>(TestIntValue));
+		FMassArchetypeSharedFragmentValues ValuesConstFloat;
+		ValuesConstFloat.AddConstSharedFragment(FConstSharedStruct::Make<FTestSharedFragment_Float>(TestFloatValue));
+		FMassArchetypeSharedFragmentValues ValuestConstIntFloat;
+		ValuestConstIntFloat.AddConstSharedFragment(FConstSharedStruct::Make<FTestSharedFragment_Int>(TestIntValue));
+		ValuestConstIntFloat.AddConstSharedFragment(FConstSharedStruct::Make<FTestSharedFragment_Float>(TestFloatValue));
+
+		{
+			FMassArchetypeSharedFragmentValues Values;
+			Values.Append(ValuesConstInt);
+			AITEST_TRUE("#5 Append results should match expectations", Values.HasSameValues(ValuesConstInt));
+			Values.Append(ValuesConstFloat);
+			AITEST_TRUE("#6 Append results should match expectations", Values.HasSameValues(ValuestConstIntFloat));
+		}
+		{
+			FMassArchetypeSharedFragmentValues Values;
+			Values.Append(ValuesConstFloat);
+			AITEST_TRUE("#7 Append results should match expectations", Values.HasSameValues(ValuesConstFloat));
+			Values.Append(ValuesConstInt);
+			AITEST_TRUE("#8 Append results should match expectations", Values.HasSameValues(ValuestConstIntFloat));
+		}
+
+		// test mismatching "mode" (const vs non-const) failure
+		AITEST_SCOPED_CHECK("trying to switch", 2);
+		{
+			FMassArchetypeSharedFragmentValues Values;
+			Values.Append(ValuesNonConstInt);
+			AITEST_EQUAL("#1 Adding mismatching mode should fail", Values.Append(ValuesConstInt), 0);
+		}
+		{
+			FMassArchetypeSharedFragmentValues Values;
+			Values.Append(ValuesConstInt);
+			AITEST_EQUAL("#2 Adding mismatching mode should fail", Values.Append(ValuesNonConstInt), 0);
+		}
+
+		return true;
+	}
+};
+IMPLEMENT_AI_INSTANT_TEST(FSharedFragmentValues_Append, "System.Mass.SharedFragments.Append");
+
+struct FSharedFragmentValues_Remove : FExecutionTestBase
+{
+	virtual bool InstantTest() override
+	{
+		constexpr int32 TestIntValue = 31;
+		constexpr float TestFloatValue = 63.f;
+
+		{
+			FMassArchetypeSharedFragmentValues ValuesNonConstInt;
+			ValuesNonConstInt.AddSharedFragment(FSharedStruct::Make<FTestSharedFragment_Int>(TestIntValue));
+			FMassArchetypeSharedFragmentValues ValuesNonConstFloat;
+			ValuesNonConstFloat.AddSharedFragment(FSharedStruct::Make<FTestSharedFragment_Float>(TestFloatValue));
+			FMassArchetypeSharedFragmentValues ValuestNonConstIntFloat;
+			ValuestNonConstIntFloat.AddSharedFragment(FSharedStruct::Make<FTestSharedFragment_Int>(TestIntValue));
+			ValuestNonConstIntFloat.AddSharedFragment(FSharedStruct::Make<FTestSharedFragment_Float>(TestFloatValue));
+
+			{
+				FMassArchetypeSharedFragmentValues Values = ValuestNonConstIntFloat;
+				AITEST_TRUE("Assignment should result in same values", Values.HasSameValues(ValuestNonConstIntFloat));
+
+				// removing just the Int shared fragment
+				Values.Remove(ValuesNonConstInt.GetSharedFragmentBitSet());
+				AITEST_TRUE("#1 Removal results should match expectations", Values.HasSameValues(ValuesNonConstFloat));
+			}
+			{
+				FMassArchetypeSharedFragmentValues Values = ValuestNonConstIntFloat;
+				// removing just the Float shared fragment
+				Values.Remove(ValuesNonConstFloat.GetSharedFragmentBitSet());
+				AITEST_TRUE("#2 Removal results should match expectations", Values.HasSameValues(ValuesNonConstInt));
+			}
+		}
+		{
+			FMassArchetypeSharedFragmentValues ValuesConstInt;
+			ValuesConstInt.AddConstSharedFragment(FConstSharedStruct::Make<FTestSharedFragment_Int>(TestIntValue));
+			FMassArchetypeSharedFragmentValues ValuesConstFloat;
+			ValuesConstFloat.AddConstSharedFragment(FConstSharedStruct::Make<FTestSharedFragment_Float>(TestFloatValue));
+			FMassArchetypeSharedFragmentValues ValuestConstIntFloat;
+			ValuestConstIntFloat.AddConstSharedFragment(FConstSharedStruct::Make<FTestSharedFragment_Int>(TestIntValue));
+			ValuestConstIntFloat.AddConstSharedFragment(FConstSharedStruct::Make<FTestSharedFragment_Float>(TestFloatValue));
+
+			{
+				FMassArchetypeSharedFragmentValues Values = ValuestConstIntFloat;
+				AITEST_TRUE("Assignment should result in same values", Values.HasSameValues(ValuestConstIntFloat));
+
+				// removing just the Int shared fragment
+				Values.Remove(ValuesConstInt.GetSharedFragmentBitSet());
+				AITEST_TRUE("#3 Removal results should match expectations", Values.HasSameValues(ValuesConstFloat));
+			}
+			{
+				FMassArchetypeSharedFragmentValues Values = ValuestConstIntFloat;
+				// removing just the Float shared fragment
+				Values.Remove(ValuesConstFloat.GetSharedFragmentBitSet());
+				AITEST_TRUE("#4 Removal results should match expectations", Values.HasSameValues(ValuesConstInt));
+			}
+		}
+
+		return true;
+	}
+};
+IMPLEMENT_AI_INSTANT_TEST(FSharedFragmentValues_Remove, "System.Mass.SharedFragments.Remove");
+
+struct FSharedFragmentValues_Hash : FExecutionTestBase
+{
+	virtual bool InstantTest() override
+	{
+		constexpr int32 TestIntValue = 31;
+		constexpr float TestFloatValue = 63.f;
+
+		FMassArchetypeSharedFragmentValues ValuestNonConstIntFloat;
+		ValuestNonConstIntFloat.AddSharedFragment(FSharedStruct::Make<FTestSharedFragment_Int>(TestIntValue));
+		ValuestNonConstIntFloat.AddSharedFragment(FSharedStruct::Make<FTestSharedFragment_Float>(TestFloatValue));
+
+		AITEST_SCOPED_CHECK("Expecting the containers to be sorted", 1);
+		const uint32 EmptyHash = ValuestNonConstIntFloat.CalculateHash();
+		AITEST_EQUAL("Expecting unsorted collection hashing to result in 0", EmptyHash, 0u);
+
+		ValuestNonConstIntFloat.Sort();
+		const uint32 ValidHash = ValuestNonConstIntFloat.CalculateHash();
+		AITEST_NOT_EQUAL("Expecting sorted collection hashing to result in non 0", ValidHash, 0u);
+
+		return true;
+	}
+};
+IMPLEMENT_AI_INSTANT_TEST(FSharedFragmentValues_Hash, "System.Mass.SharedFragments.Hash");
+
+//-----------------------------------------------------------------------------
+// Entity-related Tests
+//-----------------------------------------------------------------------------
+struct FSharedFragmentBase : public FEntityTestBase
+{
+	template<typename TSharedStruct, typename TSharedFragment>
+	FConstStructView GetSharedFragmentView(const FMassEntityHandle EntityHandle)
+	{
+		if constexpr (TIsDerivedFrom<TSharedStruct, FSharedStruct>::IsDerived)
+		{
+			return EntityManager->GetSharedFragmentDataStruct(EntityHandle, TSharedFragment::StaticStruct());
+		}
+		else
+		{
+			return EntityManager->GetConstSharedFragmentDataStruct(EntityHandle, TSharedFragment::StaticStruct());
+		}
+	};
+
+	template<typename TSharedStruct, typename TSharedFragment>
+	void CreateEntities(TArray<FMassEntityHandle>& OutEntityHandles, const int32 NumToCreate, const typename TSharedFragment::FValueType TestValue)
+	{
+		TSharedFragment FragmentInstance(TestValue);
+		FMassArchetypeSharedFragmentValues SharedIntValues;
+
+		TSharedStruct SharedFragmentInstance = TSharedStruct::Make(FragmentInstance);
+
+		if constexpr (TIsDerivedFrom<TSharedStruct, FSharedStruct>::IsDerived)
+		{
+			SharedIntValues.AddSharedFragment(SharedFragmentInstance);
+		}
+		else
+		{
+			SharedIntValues.AddConstSharedFragment(SharedFragmentInstance);
+		}
+
+		EntityManager->BatchCreateEntities(FloatsArchetype, SharedIntValues, NumToCreate, OutEntityHandles);
+	}
+
+	template<typename TSharedStruct, typename TSharedFragment>
+	void CreateEntity(FMassEntityHandle& OutEntityHandle, const typename TSharedFragment::FValueType TestValue)
+	{
+		TSharedFragment FragmentInstance(TestValue);
+		FMassArchetypeSharedFragmentValues SharedIntValues;
+
+		TSharedStruct SharedFragmentInstance = TSharedStruct::Make(FragmentInstance);
+
+		if constexpr (TIsDerivedFrom<TSharedStruct, FSharedStruct>::IsDerived)
+		{
+			SharedIntValues.AddSharedFragment(SharedFragmentInstance);
+		}
+		else
+		{
+			SharedIntValues.AddConstSharedFragment(SharedFragmentInstance);
+		}
+
+		OutEntityHandle = EntityManager->CreateEntity(FloatsArchetype, SharedIntValues);
+	}
+};
+
+template<typename TSharedStruct>
+struct FSharedFragment_CreateEntitesWithSharedFragment : public FSharedFragmentBase
+{
+	virtual bool InstantTest() override
+	{
+		constexpr int32 TestIntValueA = 1023;
+		constexpr int32 TestIntValueB = 63;
+		FMassEntityHandle EntityA, EntityB;
+
+		CreateEntity<TSharedStruct, FTestSharedFragment_Int>(EntityA, TestIntValueA);
+		CreateEntity<TSharedStruct, FTestSharedFragment_Int>(EntityB, TestIntValueB);
+
+		AITEST_EQUAL("Both entities should end up in the same archetype", EntityManager->GetArchetypeForEntityUnsafe(EntityA)
+			, EntityManager->GetArchetypeForEntityUnsafe(EntityB));
+
+		/*FConstStructView SharedFragmentA, SharedFragmentB;
+		if constexpr (TIsDerivedFrom<TSharedStruct, FSharedStruct>::IsDerived)
+		{
+			SharedFragmentA = EntityManager->GetSharedFragmentDataStruct(EntityA, FTestSharedFragment_Int::StaticStruct());
+			SharedFragmentB = EntityManager->GetSharedFragmentDataStruct(EntityB, FTestSharedFragment_Int::StaticStruct());
+		}
+		else
+		{
+			SharedFragmentA = EntityManager->GetConstSharedFragmentDataStruct(EntityA, FTestSharedFragment_Int::StaticStruct());
+			SharedFragmentB = EntityManager->GetConstSharedFragmentDataStruct(EntityB, FTestSharedFragment_Int::StaticStruct());
+		}*/
+		FConstStructView SharedFragmentA = GetSharedFragmentView<TSharedStruct, FTestSharedFragment_Int>(EntityA);
+		FConstStructView SharedFragmentB = GetSharedFragmentView<TSharedStruct, FTestSharedFragment_Int>(EntityB);
+
+		AITEST_TRUE("SharedFragmentA should be valid", SharedFragmentA.IsValid());
+		AITEST_TRUE("SharedFragmentB should be valid", SharedFragmentB.IsValid());
+		AITEST_EQUAL("SharedFragmentA should be of expected type", SharedFragmentA.GetScriptStruct(), FTestSharedFragment_Int::StaticStruct());
+		AITEST_EQUAL("SharedFragmentB should be of expected type", SharedFragmentB.GetScriptStruct(), FTestSharedFragment_Int::StaticStruct());
+		AITEST_NOT_EQUAL("SharedFragmentA and SharedFragmentB should be different instanceS", SharedFragmentA, SharedFragmentB);
+		AITEST_NOT_EQUAL("SharedFragmentA and SharedFragmentB should be distinct", SharedFragmentA, SharedFragmentB);
+		AITEST_EQUAL("SharedFragmentA's value should match the expected value", SharedFragmentA.Get<const FTestSharedFragment_Int>().Value, TestIntValueA);
+		AITEST_EQUAL("SharedFragmentB's value should match the expected value", SharedFragmentB.Get<const FTestSharedFragment_Int>().Value, TestIntValueB);
+
+		return true;
+	}
+};
+using FSharedFragment_CreateEntitesWithNonConstSharedFragment = FSharedFragment_CreateEntitesWithSharedFragment<FSharedStruct>;
+IMPLEMENT_AI_INSTANT_TEST(FSharedFragment_CreateEntitesWithNonConstSharedFragment, "System.Mass.SharedFragments.CreateEntities");
+using FSharedFragment_CreateEntitesWithConstSharedFragment = FSharedFragment_CreateEntitesWithSharedFragment<FConstSharedStruct>;
+IMPLEMENT_AI_INSTANT_TEST(FSharedFragment_CreateEntitesWithConstSharedFragment, "System.Mass.SharedFragments.CreateEntitiesConst");
+
+template<typename TSharedStruct>
+struct FSharedFragment_BatchCreateEntitesWithSharedFragment : public FSharedFragmentBase
+{
+	virtual bool InstantTest() override
+	{
+		constexpr int32 TestIntValueA = 1023;
+		constexpr int32 TestIntValueB = 63;
+		const int32 EntitiesPerChunk = FMassArchetypeHelper::ArchetypeDataFromHandleChecked(FloatsArchetype).GetNumEntitiesPerChunk();
+		// we create more than one chunk can handle properly test moving entities between chunks
+		const int32 EntitiesToCreateNumA = FMath::FloorToInt(float(EntitiesPerChunk) * 1.2f);
+		const int32 EntitiesToCreateNumB = 1;
+		constexpr int32 ExpectedNumberOfInitialChunks = 3;
+		TArray<FMassEntityHandle> EntitiesA;
+		TArray<FMassEntityHandle> EntitiesB;
+
+		CreateEntities<TSharedStruct, FTestSharedFragment_Int>(EntitiesA, EntitiesToCreateNumA, TestIntValueA);
+		CreateEntities<TSharedStruct, FTestSharedFragment_Int>(EntitiesB, EntitiesToCreateNumB, TestIntValueB);
+
+		FMassArchetypeHandle CommonArchetype = EntityManager->GetArchetypeForEntityUnsafe(EntitiesA[0]);
+		AITEST_EQUAL("All the entities should end up in the same archetype"
+			, FMassArchetypeHelper::ArchetypeDataFromHandleChecked(CommonArchetype).GetNumEntities(), EntitiesToCreateNumA + EntitiesToCreateNumB);
+		AITEST_EQUAL("The total number of chunks in the resulting archetype should match expectations"
+			, FMassArchetypeHelper::ArchetypeDataFromHandleChecked(CommonArchetype).GetChunkCount(), ExpectedNumberOfInitialChunks);
+
+		for (FMassEntityHandle EntityHandle : EntitiesA)
+		{
+			FConstStructView SharedFragment = GetSharedFragmentView<TSharedStruct, FTestSharedFragment_Int>(EntityHandle);
+			AITEST_TRUE("SharedFragment for entity type A should be valid", SharedFragment.IsValid());
+			AITEST_EQUAL("SharedFragment for entity type A  should be of expected type", SharedFragment.GetScriptStruct(), FTestSharedFragment_Int::StaticStruct());
+			AITEST_EQUAL("SharedFragment's value for entity type A should match the expected value", SharedFragment.Get<const FTestSharedFragment_Int>().Value, TestIntValueA);
+		}
+
+		FConstStructView SharedFragment = GetSharedFragmentView<TSharedStruct, FTestSharedFragment_Int>(EntitiesB[0]);
+		AITEST_TRUE("SharedFragment for entity type B should be valid", SharedFragment.IsValid());
+		AITEST_EQUAL("SharedFragment for entity type B  should be of expected type", SharedFragment.GetScriptStruct(), FTestSharedFragment_Int::StaticStruct());
+		AITEST_EQUAL("SharedFragment's value for entity type B should match the expected value", SharedFragment.Get<const FTestSharedFragment_Int>().Value, TestIntValueB);
+
+		return true;
+	}
+};
+using FSharedFragment_BatchCreateEntitesWithNonConstSharedFragment = FSharedFragment_BatchCreateEntitesWithSharedFragment<FSharedStruct>;
+IMPLEMENT_AI_INSTANT_TEST(FSharedFragment_BatchCreateEntitesWithNonConstSharedFragment, "System.Mass.SharedFragments.BatchCreateEntities");
+using FSharedFragment_BatchCreateEntitesWithConstSharedFragment = FSharedFragment_BatchCreateEntitesWithSharedFragment<FConstSharedStruct>;
+IMPLEMENT_AI_INSTANT_TEST(FSharedFragment_BatchCreateEntitesWithConstSharedFragment, "System.Mass.SharedFragments.BatchCreateEntitiesConst");
+
 } // FMassEntityTest
 
 UE_ENABLE_OPTIMIZATION_SHIP
diff --git a/Engine/Plugins/Runtime/MassEntity/Source/MassEntityTestSuite/Public/MassEntityTestTypes.h b/Engine/Plugins/Runtime/MassEntity/Source/MassEntityTestSuite/Public/MassEntityTestTypes.h
--- a/Engine/Plugins/Runtime/MassEntity/Source/MassEntityTestSuite/Public/MassEntityTestTypes.h
+++ b/Engine/Plugins/Runtime/MassEntity/Source/MassEntityTestSuite/Public/MassEntityTestTypes.h
@@ -89,7 +89,11 @@ struct FTestChunkFragment_Int : public FMassChunkFragment
 USTRUCT()
 struct FTestSharedFragment_Int : public FMassSharedFragment
 {
+	using FValueType = int32;
+
 	GENERATED_BODY()
+
+	UPROPERTY()
 	int32 Value = 0;
 
 	FTestSharedFragment_Int(const int32 InValue = 0) : Value(InValue) {}
@@ -98,7 +102,11 @@ struct FTestSharedFragment_Int : public FMassSharedFragment
 USTRUCT()
 struct FTestSharedFragment_Float : public FMassSharedFragment
 {
+	using FValueType = float;
+
 	GENERATED_BODY()
+
+	UPROPERTY()
 	float Value = 0.f;
 
 	FTestSharedFragment_Float(const float InValue = 0) : Value(InValue) {}
