Added an optional FMassArchetypeSharedFragmentValues parameter to FMassArchetypeData.MoveEntityToAnotherArchetype that will be used to override current values, if provided.
Also FMassArchetypeData.BatchMoveEntitiesToAnotherArchetype got two new optional parameters allowing adding new and removing existing shared fragments

diff --git a/Engine/Plugins/Runtime/MassEntity/Source/MassEntity/Private/MassArchetypeData.cpp b/Engine/Plugins/Runtime/MassEntity/Source/MassEntity/Private/MassArchetypeData.cpp
--- a/Engine/Plugins/Runtime/MassEntity/Source/MassEntity/Private/MassArchetypeData.cpp
+++ b/Engine/Plugins/Runtime/MassEntity/Source/MassEntity/Private/MassArchetypeData.cpp
@@ -399,7 +399,7 @@ void FMassArchetypeData::SetFragmentData(FMassArchetypeEntityCollection::FConstE
 	}
 }
 
-void FMassArchetypeData::MoveEntityToAnotherArchetype(const FMassEntityHandle Entity, FMassArchetypeData& NewArchetype)
+void FMassArchetypeData::MoveEntityToAnotherArchetype(const FMassEntityHandle Entity, FMassArchetypeData& NewArchetype, const FMassArchetypeSharedFragmentValues* SharedFragmentValuesOverride)
 {
 	check(&NewArchetype != this);
 
@@ -408,7 +408,7 @@ void FMassArchetypeData::MoveEntityToAnotherArchetype(const FMassEntityHandle En
 	const int32 IndexWithinChunk = AbsoluteIndex % NumEntitiesPerChunk;
 	FMassArchetypeChunk& Chunk = Chunks[ChunkIndex];
 
-	const int32 NewAbsoluteIndex = NewArchetype.AddEntityInternal(Entity, Chunk.GetSharedFragmentValues());
+	const int32 NewAbsoluteIndex = NewArchetype.AddEntityInternal(Entity, SharedFragmentValuesOverride ? *SharedFragmentValuesOverride : Chunk.GetSharedFragmentValues());
 	const int32 NewChunkIndex = NewAbsoluteIndex / NewArchetype.NumEntitiesPerChunk;
 	const int32 NewIndexWithinChunk = NewAbsoluteIndex % NewArchetype.NumEntitiesPerChunk;
 	FMassArchetypeChunk& NewChunk = NewArchetype.Chunks[NewChunkIndex];
@@ -1008,11 +1008,13 @@ void FMassArchetypeData::BatchAddEntities(TConstArrayView<FMassEntityHandle> Ent
 	} while (NumberMoved < Entities.Num());
 }
 
-void FMassArchetypeData::BatchMoveEntitiesToAnotherArchetype(const FMassArchetypeEntityCollection& EntityCollection, FMassArchetypeData& NewArchetype, TArray<FMassEntityHandle>& OutEntitiesBeingMoved, TArray<FMassArchetypeEntityCollection::FArchetypeEntityRange>* OutNewRanges)
+void FMassArchetypeData::BatchMoveEntitiesToAnotherArchetype(const FMassArchetypeEntityCollection& EntityCollection
+	, FMassArchetypeData& NewArchetype, TArray<FMassEntityHandle>& OutEntitiesBeingMoved
+	, TArray<FMassArchetypeEntityCollection::FArchetypeEntityRange>* OutNewRanges, const FMassArchetypeSharedFragmentValues* SharedFragmentValuesToAdd
+	, const FMassSharedFragmentBitSet* SharedFragmentToRemoveBitSet)
 {
 	check(&NewArchetype != this);
 
-
 	TArray<FMassArchetypeEntityCollection::FArchetypeEntityRange> Subchunks(EntityCollection.GetRanges());
 
 	const int32 InitialOutEntitiesCount = OutEntitiesBeingMoved.Num();
@@ -1038,14 +1040,32 @@ void FMassArchetypeData::BatchMoveEntitiesToAnotherArchetype(const FMassArchetyp
 		ResultSubChunk.ChunkIndex = 0;
 		ResultSubChunk.Length = 0;
 		int32 NumberMoved = 0;
+		const bool ChangeSharedFragments = SharedFragmentValuesToAdd || SharedFragmentToRemoveBitSet;
 
 		do
 		{
 			const int32 IndexWithinChunk = EntityRange.SubchunkStart + NumberMoved;
 
-			ResultSubChunk = NewArchetype.PrepareNextEntitiesSpanInternal(MakeArrayView(DyingEntityPtr + NumberMoved, EntityRange.Length - NumberMoved)
-				, Chunk.GetSharedFragmentValues()
-				, ResultSubChunk.ChunkIndex);
+			if (ChangeSharedFragments == false)
+			{
+				ResultSubChunk = NewArchetype.PrepareNextEntitiesSpanInternal(MakeArrayView(DyingEntityPtr + NumberMoved, EntityRange.Length - NumberMoved)
+					, Chunk.GetSharedFragmentValues(), ResultSubChunk.ChunkIndex);
+			}
+			else
+			{
+				// create new shared values
+				FMassArchetypeSharedFragmentValues NewSharedValues = Chunk.GetSharedFragmentValues();
+				if (SharedFragmentToRemoveBitSet)
+				{
+					NewSharedValues.Remove(*SharedFragmentToRemoveBitSet);
+				}
+				if (SharedFragmentValuesToAdd)
+				{
+					NewSharedValues.Append(*SharedFragmentValuesToAdd);
+				}
+				ResultSubChunk = NewArchetype.PrepareNextEntitiesSpanInternal(MakeArrayView(DyingEntityPtr + NumberMoved, EntityRange.Length - NumberMoved)
+					, NewSharedValues, ResultSubChunk.ChunkIndex);
+			}
 
 			FMassArchetypeChunk& NewChunk = NewArchetype.Chunks[ResultSubChunk.ChunkIndex];
 			MoveFragmentsToAnotherArchetypeInternal(NewArchetype, {NewChunk.GetRawMemory(), ResultSubChunk.SubchunkStart}, {Chunk.GetRawMemory(), IndexWithinChunk}, ResultSubChunk.Length);
diff --git a/Engine/Plugins/Runtime/MassEntity/Source/MassEntity/Private/MassArchetypeData.h b/Engine/Plugins/Runtime/MassEntity/Source/MassEntity/Private/MassArchetypeData.h
--- a/Engine/Plugins/Runtime/MassEntity/Source/MassEntity/Private/MassArchetypeData.h
+++ b/Engine/Plugins/Runtime/MassEntity/Source/MassEntity/Private/MassArchetypeData.h
@@ -223,8 +223,8 @@ public:
 	void Initialize(const FMassArchetypeCompositionDescriptor& InCompositionDescriptor, const uint32 ArchetypeDataVersion);
 
 	/** 
-	 * A special way of initializing an archetype resulting in a copy of SiblingArchetype's setup with OverrideTags
-	 * replacing original tags of SiblingArchetype
+	 * A special way of initializing an archetype resulting in a copy of BaseArchetype's setup with OverrideTags
+	 * replacing original tags of BaseArchetype
 	 */
 	void InitializeWithSimilar(const FMassArchetypeData& BaseArchetype, FMassArchetypeCompositionDescriptor&& NewComposition, const uint32 ArchetypeDataVersion);
 
@@ -267,8 +267,9 @@ public:
 	 * Moves the entity from this archetype to another, will only copy all matching fragment types
 	 * @param Entity is the entity to move
 	 * @param NewArchetype the archetype to move to
+	 * @param SharedFragmentValuesOverride if provided will override all given Entity's shared fragment values
 	 */
-	void MoveEntityToAnotherArchetype(const FMassEntityHandle Entity, FMassArchetypeData& NewArchetype);
+	void MoveEntityToAnotherArchetype(const FMassEntityHandle Entity, FMassArchetypeData& NewArchetype, const FMassArchetypeSharedFragmentValues* SharedFragmentValuesOverride = nullptr);
 
 	/**
 	 * Set all fragment sources data on specified entity, will check if there are fragment sources type that does not exist in the archetype
@@ -363,8 +364,15 @@ public:
 	//////////////////////////////////////////////////////////////////////
 	// batched api
 	void BatchDestroyEntityChunks(FMassArchetypeEntityCollection::FConstEntityRangeArrayView EntityRangeContainer, TArray<FMassEntityHandle>& OutEntitiesRemoved);
-	void BatchAddEntities(TConstArrayView<FMassEntityHandle> Entities, const FMassArchetypeSharedFragmentValues& SharedFragmentValues, TArray<FMassArchetypeEntityCollection::FArchetypeEntityRange>& OutNewRanges);
-	void BatchMoveEntitiesToAnotherArchetype(const FMassArchetypeEntityCollection& EntityCollection, FMassArchetypeData& NewArchetype, TArray<FMassEntityHandle>& OutEntitiesBeingMoved, TArray<FMassArchetypeEntityCollection::FArchetypeEntityRange>* OutNewChunks = nullptr);
+	void BatchAddEntities(TConstArrayView<FMassEntityHandle> Entities, const FMassArchetypeSharedFragmentValues& SharedFragmentValues
+		, TArray<FMassArchetypeEntityCollection::FArchetypeEntityRange>& OutNewRanges);
+	/** 
+	 * @param SharedFragmentValuesOverride if provided will override shared fragment values for the entities being moved
+	 */
+	void BatchMoveEntitiesToAnotherArchetype(const FMassArchetypeEntityCollection& EntityCollection, FMassArchetypeData& NewArchetype
+		, TArray<FMassEntityHandle>& OutEntitiesBeingMoved, TArray<FMassArchetypeEntityCollection::FArchetypeEntityRange>* OutNewChunks = nullptr
+		, const FMassArchetypeSharedFragmentValues* SharedFragmentValuesToAdd = nullptr
+		, const FMassSharedFragmentBitSet* SharedFragmentToRemoveBitSet = nullptr);
 	void BatchSetFragmentValues(TConstArrayView<FMassArchetypeEntityCollection::FArchetypeEntityRange> EntityCollection, const FMassGenericPayloadViewSlice& Payload);
 
 protected:
