Tiny FMassArchetypeSharedFragmentValues.Remove comment extension

diff --git a/Engine/Plugins/Runtime/MassEntity/Source/MassEntity/Public/MassEntityTypes.h b/Engine/Plugins/Runtime/MassEntity/Source/MassEntity/Public/MassEntityTypes.h
--- a/Engine/Plugins/Runtime/MassEntity/Source/MassEntity/Public/MassEntityTypes.h
+++ b/Engine/Plugins/Runtime/MassEntity/Source/MassEntity/Public/MassEntityTypes.h
@@ -316,7 +316,10 @@ struct MASSENTITY_API FMassArchetypeSharedFragmentValues
 	 */
 	int32 Append(const FMassArchetypeSharedFragmentValues& Other);
 
-	/** Note that the function removes the shared fragments by type */
+	/** 
+	 * Note that the function removes the shared fragments by type
+	 * @return number of fragments types removed
+	 */
 	int32 Remove(const FMassSharedFragmentBitSet& SharedFragmentToRemoveBitSet);
 
 	FORCEINLINE const TArray<FConstSharedStruct>& GetConstSharedFragments() const
