Added means of adding shared fragments to entities.

The change includes accompanying changes:
* Changed a bit how we handle creating entities with a given archetype handle and a set of shared fragment values. In the past if the shared fragments didn't match the archetype we simply failed the operation. New version creates a new archetype combining the given archetypes composition (sans shared fragments) with composition of given shared fragments and uses that to create entities.
* A new function for batch-reserving entities

diff --git a/Engine/Plugins/Runtime/MassEntity/Source/MassEntity/Private/MassEntityManager.cpp b/Engine/Plugins/Runtime/MassEntity/Source/MassEntity/Private/MassEntityManager.cpp
--- a/Engine/Plugins/Runtime/MassEntity/Source/MassEntity/Private/MassEntityManager.cpp
+++ b/Engine/Plugins/Runtime/MassEntity/Source/MassEntity/Private/MassEntityManager.cpp
@@ -206,6 +206,19 @@ FMassArchetypeHandle FMassEntityManager::CreateArchetype(const TSharedPtr<FMassA
 	return CreateArchetype(Composition, CreationParams);
 }
 
+FMassArchetypeHandle FMassEntityManager::GetOrCreateSuitableArchetype(const FMassArchetypeHandle& ArchetypeHandle, const FMassSharedFragmentBitSet& SharedFragmentBitSet
+	, const FMassArchetypeCreationParams& CreationParams)
+{
+	FMassArchetypeData& ArchetypeData = FMassArchetypeHelper::ArchetypeDataFromHandleChecked(ArchetypeHandle);
+	if (SharedFragmentBitSet != ArchetypeData.GetSharedFragmentBitSet())
+	{
+		FMassArchetypeCompositionDescriptor NewDescriptor = ArchetypeData.GetCompositionDescriptor();
+		NewDescriptor.SharedFragments = SharedFragmentBitSet;
+		return CreateArchetype(NewDescriptor);
+	}
+	return ArchetypeHandle;
+}
+
 FMassArchetypeHandle FMassEntityManager::CreateArchetype(const FMassArchetypeCompositionDescriptor& Composition, const FMassArchetypeCreationParams& CreationParams)
 {
 	const uint32 TypeHash = Composition.CalculateHash();
@@ -470,6 +483,17 @@ void FMassEntityManager::BuildEntity(FMassEntityHandle Entity, TConstArrayView<F
 	EntityData.CurrentArchetype->SetFragmentsData(Entity, FragmentInstanceList);
 }
 
+TConstArrayView<FMassEntityHandle> FMassEntityManager::BatchReserveEntities(const int32 Count, TArray<FMassEntityHandle>& InOutEntities)
+{
+	int32 Index = InOutEntities.Num();
+	InOutEntities.Reserve(Index + Count);
+	for (int32 Counter = 0; Counter < Count; ++Counter)
+	{
+		InOutEntities.Add(ReserveEntity());
+	}
+	return MakeArrayView(InOutEntities.GetData() + Index, Count);
+}
+
 void FMassEntityManager::BatchBuildEntities(const FMassArchetypeEntityCollectionWithPayload& EncodedEntitiesWithPayload
 	, const FMassFragmentBitSet& FragmentsAffected, const FMassArchetypeSharedFragmentValues& SharedFragmentValues, const FMassArchetypeCreationParams& CreationParams)
 {
@@ -548,37 +572,23 @@ TSharedRef<FMassEntityManager::FEntityCreationContext> FMassEntityManager::Batch
 	checkf(IsProcessing() == false, TEXT("Synchronous API function %hs called during mass processing. Use asynchronous API instead."), __FUNCTION__);
 	checkf(!ReservedEntities.IsEmpty(), TEXT("No reserved entities given to batch create."));
 
-	// verify that SharedFragmentValues contains all the data needed for the archetype indicated by ArchetypeHandle
-	FMassArchetypeData& ArchetypeData = FMassArchetypeHelper::ArchetypeDataFromHandleChecked(ArchetypeHandle);
-	if (!ensureMsgf(SharedFragmentValues.HasAllRequiredFragmentTypes(ArchetypeData.GetSharedFragmentBitSet())
-		, TEXT("Trying to create entities with mismatching shared fragments collection. This would have lead to crashes, so we're bailing out. Make sure you pass in values for all the shared fragments declared in archetype's composition.")))
-	{
-		return MakeShareable(new FEntityCreationContext(0));
-	}
-
-	return InternalBatchCreateReservedEntities(ArchetypeHandle, SharedFragmentValues, ReservedEntities);
+	return InternalBatchCreateReservedEntities(GetOrCreateSuitableArchetype(ArchetypeHandle, SharedFragmentValues.GetSharedFragmentBitSet())
+		, SharedFragmentValues, ReservedEntities);
 }
 
 TSharedRef<FMassEntityManager::FEntityCreationContext> FMassEntityManager::BatchCreateEntities(const FMassArchetypeHandle& ArchetypeHandle
-	, const FMassArchetypeSharedFragmentValues& SharedFragmentValues, const int32 Count, TArray<FMassEntityHandle>& OutEntities)
+	, const FMassArchetypeSharedFragmentValues& SharedFragmentValues, const int32 Count, TArray<FMassEntityHandle>& InOutEntities)
 {
 	TRACE_CPUPROFILER_EVENT_SCOPE(Mass_BatchCreateEntities);
 
-	FMassArchetypeData& ArchetypeData = FMassArchetypeHelper::ArchetypeDataFromHandleChecked(ArchetypeHandle);
-	if (!ensureMsgf(SharedFragmentValues.HasAllRequiredFragmentTypes(ArchetypeData.GetSharedFragmentBitSet())
-		, TEXT("Trying to create entities with mismatching shared fragments collection. This would have lead to crashes, so we're bailing out. Make sure you pass in values for all the shared fragments declared in archetype's composition.")))
-	{
-		return MakeShareable(new FEntityCreationContext(0));
-	}
+	checkf(IsProcessing() == false, TEXT("Synchronous API function %hs called during mass processing. Use asynchronous API instead."), __FUNCTION__);
+	testableCheckfReturn(ArchetypeHandle.IsValid(), return MakeShareable(new FEntityCreationContext(0))
+		, TEXT("%hs expecting a valid ArchetypeHandle"), __FUNCTION__);
 
-	int32 Index = OutEntities.Num();
-	OutEntities.Reserve(Index + Count);
-	for (int32 Counter = 0; Counter < Count; ++Counter)
-	{
-		OutEntities.Add(ReserveEntity());
-	}
+	TConstArrayView<FMassEntityHandle> ReservedEntities = BatchReserveEntities(Count, InOutEntities);
 	
-	return InternalBatchCreateReservedEntities(ArchetypeHandle, SharedFragmentValues, MakeArrayView(OutEntities.GetData() + Index, Count));
+	return InternalBatchCreateReservedEntities(GetOrCreateSuitableArchetype(ArchetypeHandle, SharedFragmentValues.GetSharedFragmentBitSet())
+		, SharedFragmentValues, ReservedEntities);
 }
 
 TSharedRef<FMassEntityManager::FEntityCreationContext> FMassEntityManager::InternalBatchCreateReservedEntities(const FMassArchetypeHandle& ArchetypeHandle
@@ -1009,6 +1019,55 @@ void FMassEntityManager::RemoveTagFromEntity(FMassEntityHandle Entity, const USc
 	}
 }
 
+bool FMassEntityManager::AddConstSharedFragmentToEntity(const FMassEntityHandle Entity, const FConstSharedStruct& InConstSharedFragment)
+{
+	if (!ensureMsgf(InConstSharedFragment.IsValid(), TEXT("%hs parameter Fragment is expected to be valid"), __FUNCTION__))
+	{
+		return false;
+	}
+
+	CheckIfEntityIsActive(Entity);
+
+	FEntityData& EntityData = Entities[Entity.Index];
+	FMassArchetypeData* CurrentArchetype = EntityData.CurrentArchetype.Get();
+	check(CurrentArchetype);
+
+	const UScriptStruct* StructType = InConstSharedFragment.GetScriptStruct();
+	CA_ASSUME(StructType);
+	if (CurrentArchetype->GetCompositionDescriptor().SharedFragments.Contains(*StructType))
+	{
+		const FMassArchetypeSharedFragmentValues& SharedFragmentValues = CurrentArchetype->GetSharedFragmentValues(Entity);
+		FConstSharedStruct ExistingConstSharedStruct = SharedFragmentValues.GetConstSharedFragmentStruct(StructType);
+		if (ExistingConstSharedStruct == InConstSharedFragment || ExistingConstSharedStruct.CompareStructValues(InConstSharedFragment))
+		{
+			// nothing to do
+			return true;
+		}
+		UE_LOG(LogMass, Warning, TEXT("Changing shared fragment value of entities is not supported"));
+		return false;
+	}
+	
+	FMassArchetypeCompositionDescriptor NewComposition(CurrentArchetype->GetCompositionDescriptor());
+	NewComposition.SharedFragments.Add(*StructType);
+	const FMassArchetypeHandle NewArchetypeHandle = CreateArchetype(NewComposition);
+	check(NewArchetypeHandle.IsValid());
+	FMassArchetypeData* NewArchetype = NewArchetypeHandle.DataPtr.Get();
+	check(NewArchetype);
+
+	const FMassArchetypeSharedFragmentValues& OldSharedFragmentValues = CurrentArchetype->GetSharedFragmentValues(Entity.Index);
+	check(!OldSharedFragmentValues.ContainsType(StructType));
+	FMassArchetypeSharedFragmentValues NewSharedFragmentValues(OldSharedFragmentValues);
+	NewSharedFragmentValues.AddConstSharedFragment(InConstSharedFragment);
+	NewSharedFragmentValues.Sort();
+
+	CurrentArchetype->MoveEntityToAnotherArchetype(Entity, *NewArchetype, &NewSharedFragmentValues);
+
+	// Change the entity archetype
+	EntityData.CurrentArchetype = NewArchetypeHandle.DataPtr;
+
+	return true;
+}
+
 void FMassEntityManager::BatchChangeTagsForEntities(TConstArrayView<FMassArchetypeEntityCollection> EntityCollections, const FMassTagBitSet& TagsToAdd, const FMassTagBitSet& TagsToRemove)
 {
 	TRACE_CPUPROFILER_EVENT_SCOPE(Mass_BatchChangeTagsForEntities);
@@ -1192,6 +1251,42 @@ void FMassEntityManager::BatchAddFragmentInstancesForEntities(TConstArrayView<FM
 	}
 }
 
+void FMassEntityManager::BatchAddSharedFragmentsForEntities(TConstArrayView<FMassArchetypeEntityCollection> EntityCollections
+	, const FMassArchetypeSharedFragmentValues& AddedFragmentValues)
+{
+	TRACE_CPUPROFILER_EVENT_SCOPE(Mass_BatchAddConstSharedFragmentForEntities);
+
+	for (const FMassArchetypeEntityCollection& Collection : EntityCollections)
+	{
+		FMassArchetypeData* CurrentArchetype = Collection.GetArchetype().DataPtr.Get();
+		testableCheckfReturn(CurrentArchetype, continue, TEXT("Adding shared fragments to archetype-less entities is not supported"));
+
+		FMassArchetypeCompositionDescriptor NewComposition(CurrentArchetype->GetCompositionDescriptor());
+		NewComposition.SharedFragments += AddedFragmentValues.GetSharedFragmentBitSet();
+
+		const FMassArchetypeHandle NewArchetypeHandle = CreateArchetype(NewComposition);
+		check(NewArchetypeHandle.IsValid());
+		FMassArchetypeData* NewArchetype = NewArchetypeHandle.DataPtr.Get();
+		check(NewArchetype);
+		if (!testableEnsureMsgf(CurrentArchetype != NewArchetype, TEXT("Setting shared fragment values without archetype change is not supported")))
+		{
+			UE_LOG(LogMass, Warning, TEXT("Trying to set shared fragment values, without adding new shared fragments, is not supported."));
+			continue;
+		}
+
+		TArray<FMassEntityHandle> EntitiesBeingMoved;
+		CurrentArchetype->BatchMoveEntitiesToAnotherArchetype(Collection, *NewArchetype, EntitiesBeingMoved, /*OutNewChunks=*/nullptr, &AddedFragmentValues);
+
+		for (const FMassEntityHandle& Entity : EntitiesBeingMoved)
+		{
+			check(Entities.IsValidIndex(Entity.Index));
+
+			FEntityData& EntityData = Entities[Entity.Index];
+			EntityData.CurrentArchetype = NewArchetypeHandle.DataPtr;
+		}
+	}
+}
+
 void FMassEntityManager::MoveEntityToAnotherArchetype(FMassEntityHandle Entity, FMassArchetypeHandle NewArchetypeHandle)
 {
 	CheckIfEntityIsActive(Entity);
diff --git a/Engine/Plugins/Runtime/MassEntity/Source/MassEntity/Public/MassEntityManager.h b/Engine/Plugins/Runtime/MassEntity/Source/MassEntity/Public/MassEntityManager.h
--- a/Engine/Plugins/Runtime/MassEntity/Source/MassEntity/Public/MassEntityManager.h
+++ b/Engine/Plugins/Runtime/MassEntity/Source/MassEntity/Public/MassEntityManager.h
@@ -134,7 +134,18 @@ public:
 	 *   types that SourceArchetype doesn't already have. If the caller cannot guarantee it use of AddFragment functions
 	 *   family is recommended.
 	 */
-	FMassArchetypeHandle CreateArchetype(const TSharedPtr<FMassArchetypeData>& SourceArchetype, const FMassFragmentBitSet& InFragments, const FMassArchetypeCreationParams& CreationParams = FMassArchetypeCreationParams());
+	FMassArchetypeHandle CreateArchetype(const TSharedPtr<FMassArchetypeData>& SourceArchetype, const FMassFragmentBitSet& InFragments
+		, const FMassArchetypeCreationParams& CreationParams = FMassArchetypeCreationParams());
+
+	/** 
+	 * A helper function to be used when creating entities with shared fragments provided, or when adding shared fragments
+	 * to existing entities
+	 * @param ArchetypeHandle that's the assumed target archetype. But we'll be making sure its composition matches SharedFragmentsBitSet
+	 * @param SharedFragmentBitSet indicates which shared fragments we want the target archetype to have. If ArchetypeHandle 
+	 *	doesn't have these a new archetype will be created.
+	 */
+	FMassArchetypeHandle GetOrCreateSuitableArchetype(const FMassArchetypeHandle& ArchetypeHandle, const FMassSharedFragmentBitSet& SharedFragmentBitSet
+		, const FMassArchetypeCreationParams& CreationParams = FMassArchetypeCreationParams());
 
 	/** Fetches the archetype for a given Entity. If Entity is not valid it will still return a handle, just with an invalid archetype */
 	FMassArchetypeHandle GetArchetypeForEntity(FMassEntityHandle Entity) const;
@@ -204,12 +215,12 @@ public:
 	 * @param Archetype you want this entity to be
 	 * @param SharedFragmentValues to be associated with the entities
 	 * @param Count number of entities to create
-	 * @param OutEntities the newly created entities are appended to given array, i.e. the pre-existing content of OutEntities won't be affected by the call
+	 * @param InOutEntities the newly created entities are appended to given array, i.e. the pre-existing content of OutEntities won't be affected by the call
 	 * @return a creation context that will notify all the interested observers about newly created fragments once the context is released */
-	TSharedRef<FEntityCreationContext> BatchCreateEntities(const FMassArchetypeHandle& ArchetypeHandle, const FMassArchetypeSharedFragmentValues& SharedFragmentValues, const int32 Count, TArray<FMassEntityHandle>& OutEntities);
-	FORCEINLINE TSharedRef<FEntityCreationContext> BatchCreateEntities(const FMassArchetypeHandle& ArchetypeHandle, const int32 Count, TArray<FMassEntityHandle>& OutEntities)
+	TSharedRef<FEntityCreationContext> BatchCreateEntities(const FMassArchetypeHandle& ArchetypeHandle, const FMassArchetypeSharedFragmentValues& SharedFragmentValues, const int32 Count, TArray<FMassEntityHandle>& InOutEntities);
+	FORCEINLINE TSharedRef<FEntityCreationContext> BatchCreateEntities(const FMassArchetypeHandle& ArchetypeHandle, const int32 Count, TArray<FMassEntityHandle>& InOutEntities)
 	{
-		return BatchCreateEntities(ArchetypeHandle, FMassArchetypeSharedFragmentValues(), Count, OutEntities);
+		return BatchCreateEntities(ArchetypeHandle, FMassArchetypeSharedFragmentValues(), Count, InOutEntities);
 	}
 
 	/**
@@ -268,6 +279,21 @@ public:
 	void RemoveTagFromEntity(FMassEntityHandle Entity, const UScriptStruct* TagType);
 	void SwapTagsForEntity(FMassEntityHandle Entity, const UScriptStruct* FromFragmentType, const UScriptStruct* ToFragmentType);
 
+	/** 
+	 * Adds a new const shared fragment to the given entity. Note that it only works if the given entity doesn't have
+	 * a shared fragment of the given type. The function will give a soft "pass" if the entity has the shared fragment
+	 * of the same value. Setting shared fragment value (i.e. changing) is not supported and the function will log
+	 * a warning if that's attempted.
+	 * @return whether the Entity has the Fragment value assigned to it, regardless of its original state (i.e. the function will
+	 *	return true also if the Entity already had the same values associated with it)
+	 */
+	bool AddConstSharedFragmentToEntity(const FMassEntityHandle Entity, const FConstSharedStruct& InConstSharedFragment);
+
+	/** 
+	 * Reserves Count number of entities and appends them to InOutEntities
+	 * @return a view into InOutEntities containing only the freshly reserved entities
+	 */
+	TConstArrayView<FMassEntityHandle> BatchReserveEntities(const int32 Count, TArray<FMassEntityHandle>& InOutEntities);
 	void BatchBuildEntities(const FMassArchetypeEntityCollectionWithPayload& EncodedEntitiesWithPayload, const FMassFragmentBitSet& FragmentsAffected
 		, const FMassArchetypeSharedFragmentValues& SharedFragmentValues = {}, const FMassArchetypeCreationParams& CreationParams = FMassArchetypeCreationParams());
 	void BatchBuildEntities(const FMassArchetypeEntityCollectionWithPayload& EncodedEntitiesWithPayload, FMassArchetypeCompositionDescriptor&& Composition
@@ -275,6 +301,10 @@ public:
 	void BatchChangeTagsForEntities(TConstArrayView<FMassArchetypeEntityCollection> EntityCollections, const FMassTagBitSet& TagsToAdd, const FMassTagBitSet& TagsToRemove);
 	void BatchChangeFragmentCompositionForEntities(TConstArrayView<FMassArchetypeEntityCollection> EntityCollections, const FMassFragmentBitSet& FragmentsToAdd, const FMassFragmentBitSet& FragmentsToRemove);
 	void BatchAddFragmentInstancesForEntities(TConstArrayView<FMassArchetypeEntityCollectionWithPayload> EntityCollections, const FMassFragmentBitSet& FragmentsAffected);
+	/** 
+	 * Adds a new const and non-const shared fragments to all entities provided via EntityCollections 
+	 */
+	void BatchAddSharedFragmentsForEntities(TConstArrayView<FMassArchetypeEntityCollection> EntityCollections, const FMassArchetypeSharedFragmentValues& AddedFragmentValues);
 
 	/**
 	 * Adds fragments and tags indicated by InOutDescriptor to the Entity. The function also figures out which elements
