Fixed FMassArchetypeData::BatchMoveEntitiesToAnotherArchetype's path moving the entities while changing the shared fragments. The recently added path was missing shared fragment value sorting (which was causing ensured when used).

In addition a test has been added to the function verifying that the target archetype is configured to accept all provided shared fragment values

diff --git a/Engine/Plugins/Runtime/MassEntity/Source/MassEntity/Private/MassArchetypeData.cpp b/Engine/Plugins/Runtime/MassEntity/Source/MassEntity/Private/MassArchetypeData.cpp
--- a/Engine/Plugins/Runtime/MassEntity/Source/MassEntity/Private/MassArchetypeData.cpp
+++ b/Engine/Plugins/Runtime/MassEntity/Source/MassEntity/Private/MassArchetypeData.cpp
@@ -987,6 +987,9 @@ void FMassArchetypeData::BatchAddEntities(TConstArrayView<FMassEntityHandle> Ent
 {
 	SCOPE_CYCLE_COUNTER(STAT_Mass_ArchetypeBatchAdd);
 
+	testableCheckfReturn(SharedFragmentValues.HasExactFragmentTypesMatch(GetCompositionDescriptor().SharedFragments)
+		, return, TEXT("%hs parameter SharedFragmentValues doesn't match archetype's composition"), __FUNCTION__);
+
 	FMassArchetypeEntityCollection::FArchetypeEntityRange ResultSubchunk;
 	ResultSubchunk.ChunkIndex = 0;
 	int32 NumberMoved = 0;
@@ -1015,6 +1018,23 @@ void FMassArchetypeData::BatchMoveEntitiesToAnotherArchetype(const FMassArchetyp
 {
 	check(&NewArchetype != this);
 
+	if (SharedFragmentValuesToAdd || SharedFragmentToRemoveBitSet)
+	{
+		// verify the new archetype's shared fragment composition matches current archetype's composition modified as requested
+		FMassSharedFragmentBitSet NewSharedFragmentsBitset = GetSharedFragmentBitSet();
+		if (SharedFragmentToRemoveBitSet)
+		{
+			NewSharedFragmentsBitset -= *SharedFragmentToRemoveBitSet;
+		}
+		if (SharedFragmentValuesToAdd)
+		{
+			NewSharedFragmentsBitset += SharedFragmentValuesToAdd->GetSharedFragmentBitSet();
+		}
+		
+		testableCheckfReturn(NewArchetype.GetCompositionDescriptor().SharedFragments == NewSharedFragmentsBitset
+			, return, TEXT("%hs parameter SharedFragmentValues doesn't match archetype's composition"), __FUNCTION__);
+	}
+
 	TArray<FMassArchetypeEntityCollection::FArchetypeEntityRange> Subchunks(EntityCollection.GetRanges());
 
 	const int32 InitialOutEntitiesCount = OutEntitiesBeingMoved.Num();
@@ -1063,6 +1083,8 @@ void FMassArchetypeData::BatchMoveEntitiesToAnotherArchetype(const FMassArchetyp
 				{
 					NewSharedValues.Append(*SharedFragmentValuesToAdd);
 				}
+				NewSharedValues.Sort();
+
 				ResultSubChunk = NewArchetype.PrepareNextEntitiesSpanInternal(MakeArrayView(DyingEntityPtr + NumberMoved, EntityRange.Length - NumberMoved)
 					, NewSharedValues, ResultSubChunk.ChunkIndex);
 			}
