Added some shared fragments unit tests related to recently added new functionality.

diff --git a/Engine/Plugins/Runtime/MassEntity/Source/MassEntityTestSuite/Private/MassSharedFragmentsTest.cpp b/Engine/Plugins/Runtime/MassEntity/Source/MassEntityTestSuite/Private/MassSharedFragmentsTest.cpp
--- a/Engine/Plugins/Runtime/MassEntity/Source/MassEntityTestSuite/Private/MassSharedFragmentsTest.cpp
+++ b/Engine/Plugins/Runtime/MassEntity/Source/MassEntityTestSuite/Private/MassSharedFragmentsTest.cpp
@@ -427,6 +427,152 @@ IMPLEMENT_AI_INSTANT_TEST(FSharedFragment_BatchCreateEntitesWithNonConstSharedFr
 using FSharedFragment_BatchCreateEntitesWithConstSharedFragment = FSharedFragment_BatchCreateEntitesWithSharedFragment<FConstSharedStruct>;
 IMPLEMENT_AI_INSTANT_TEST(FSharedFragment_BatchCreateEntitesWithConstSharedFragment, "System.Mass.SharedFragments.BatchCreateEntitiesConst");
 
+
+struct FSharedFragment_AddToEntity : FEntityTestBase
+{
+	virtual bool InstantTest() override
+	{
+		constexpr int32 TestIntValue = 1023;
+		FTestSharedFragment_Int FragmentInstance(TestIntValue);
+		FSharedStruct SharedFragmentInstance = FSharedStruct::Make(FragmentInstance);
+
+		const FMassEntityHandle EntityHandle = EntityManager->CreateEntity(FloatsArchetype);
+
+		FTestSharedFragment_Int* EntitySharedFragment = EntityManager->GetSharedFragmentDataPtr<FTestSharedFragment_Int>(EntityHandle);
+		AITEST_NULL("Initially the entity is not expected to have the shared fragment", EntitySharedFragment);
+
+		EntityManager->AddConstSharedFragmentToEntity(EntityHandle, SharedFragmentInstance);
+
+		EntitySharedFragment = EntityManager->GetConstSharedFragmentDataPtr<FTestSharedFragment_Int>(EntityHandle);
+		AITEST_NOT_NULL("The entity is expected to have the shared fragment after the operation", EntitySharedFragment);
+		AITEST_EQUAL("The the shared fragment is expected to store the configured value", EntitySharedFragment->Value, TestIntValue);
+
+		// at this point the Entity already has a shared fragment of a given type
+		// now we're going to add it again and test the systems behavior, we'll be adding the same FMasSharedFragment type
+		// in both const and non-const way.
+		constexpr int32 DifferentTestIntValue = TestIntValue + 1;
+		FTestSharedFragment_Int DifferentFragmentInstance(DifferentTestIntValue);
+		FSharedStruct DifferentSharedFragmentInstance = FSharedStruct::Make(DifferentFragmentInstance);
+		FConstSharedStruct DifferentConstSharedFragmentInstance = FConstSharedStruct::Make(DifferentFragmentInstance);
+
+		GetTestRunner().AddExpectedError(TEXT("Changing shared fragment value of entities is not supported"), EAutomationExpectedErrorFlags::Contains, 2);
+
+		const bool bSuccessfullyAddedSharedFragment = EntityManager->AddConstSharedFragmentToEntity(EntityHandle, DifferentSharedFragmentInstance);
+		AITEST_FALSE("Adding existing shared fragment type should fail", bSuccessfullyAddedSharedFragment);
+		const bool bSuccessfullyAddedConstSharedFragment = EntityManager->AddConstSharedFragmentToEntity(EntityHandle, DifferentConstSharedFragmentInstance);
+		AITEST_FALSE("Adding existing const shared fragment type should fail", bSuccessfullyAddedConstSharedFragment);
+
+		return true;
+	}
+};
+IMPLEMENT_AI_INSTANT_TEST(FSharedFragment_AddToEntity, "System.Mass.SharedFragments.AddToEntity");
+
+struct FSharedFragment_BatchAddToEntity : FEntityTestBase
+{
+	virtual bool InstantTest() override
+	{
+		constexpr int32 TestIntValue = 1023;
+
+		const FMassArchetypeHandle InitialArchetype = FloatsArchetype;
+		const int32 EntitiesPerChunk = FMassArchetypeHelper::ArchetypeDataFromHandleChecked(InitialArchetype).GetNumEntitiesPerChunk();
+		const int32 EntitiesToCreateNum = FMath::FloorToInt(EntitiesPerChunk * 2.2f);
+		const int32 EntitiesToMoveNum = FMath::FloorToInt(EntitiesPerChunk * 1.2f);
+
+		TArray<FMassEntityHandle> CreatedEntityHandles;
+
+		EntityManager->BatchCreateEntities(InitialArchetype, EntitiesToCreateNum, CreatedEntityHandles);
+
+		TArray<FMassEntityHandle> EntitiesToMove = CreatedEntityHandles;
+		Algo::RandomShuffle(EntitiesToMove);
+		TConstArrayView<FMassEntityHandle> EntitiesMoved = MakeArrayView(EntitiesToMove.GetData(), EntitiesToMoveNum);
+		FMassArchetypeEntityCollection EntityCollection(InitialArchetype, EntitiesMoved, FMassArchetypeEntityCollection::EDuplicatesHandling::NoDuplicates);
+
+		FMassArchetypeSharedFragmentValues SharedValues;
+		FConstSharedStruct ConstSharedFragment = FConstSharedStruct::Make<FTestSharedFragment_Int>(TestIntValue);
+		SharedValues.AddConstSharedFragment(ConstSharedFragment);
+		EntityManager->BatchAddSharedFragmentsForEntities(MakeArrayView(&EntityCollection, 1), SharedValues);
+
+		const FMassArchetypeHandle TargetArchetype = EntityManager->GetArchetypeForEntityUnsafe(EntitiesToMove[0]);
+		const int32 EntitiesMovedNum = FMassArchetypeHelper::ArchetypeDataFromHandleChecked(TargetArchetype).GetNumEntities();
+		AITEST_EQUAL("Number of entities moves needs to match expectations", EntitiesMovedNum, EntitiesToMoveNum);
+		for (const FMassEntityHandle& EntityHandle : EntitiesMoved)
+		{
+			FTestSharedFragment_Int* SharedFragmentInstance = EntityManager->GetConstSharedFragmentDataPtr<FTestSharedFragment_Int>(EntityHandle);
+			AITEST_NOT_NULL("Every entity moved needs to have a valid shared fragment", SharedFragmentInstance);
+			AITEST_EQUAL("The shared fragment's value needs to match expectations", SharedFragmentInstance->Value, TestIntValue);
+		}
+
+		return true;
+	}
+};
+IMPLEMENT_AI_INSTANT_TEST(FSharedFragment_BatchAddToEntity, "System.Mass.SharedFragments.BatchAddToEntity");
+
+struct FSharedFragment_BatchSetAttempt : FEntityTestBase
+{
+	virtual bool InstantTest() override
+	{
+		constexpr int32 TestIntValue = 1023;
+		constexpr int32 OtherTestIntValue = TestIntValue + 1;
+
+		const int32 EntitiesPerChunk = FMassArchetypeHelper::ArchetypeDataFromHandleChecked(FloatsArchetype).GetNumEntitiesPerChunk();
+		const int32 EntitiesToCreateNum = FMath::FloorToInt(EntitiesPerChunk * 2.2f);
+		const int32 EntitiesToMoveNum = FMath::FloorToInt(EntitiesPerChunk * 1.2f);
+
+		TArray<FMassEntityHandle> CreatedEntityHandles;
+		FMassArchetypeSharedFragmentValues SharedIntValues;
+		FConstSharedStruct ConstSharedFragment = FConstSharedStruct::Make<FTestSharedFragment_Int>(TestIntValue);
+		SharedIntValues.AddConstSharedFragment(ConstSharedFragment);
+
+		TSharedRef<FMassEntityManager::FEntityCreationContext> CreationContext = EntityManager->BatchCreateEntities(FloatsArchetype, SharedIntValues, EntitiesToCreateNum, CreatedEntityHandles);
+		const FMassArchetypeHandle ResultingArchetype = CreationContext->GetEntityCollection().GetArchetype();
+
+		FMassArchetypeEntityCollection EntityCollection(ResultingArchetype, CreatedEntityHandles, FMassArchetypeEntityCollection::EDuplicatesHandling::NoDuplicates);
+		
+		// attempting to add the same values again should fail with checks and ensures
+		{	
+			AITEST_SCOPED_CHECK("Setting shared fragment values without archetype change is not supported", 1);
+			AITEST_SCOPED_CHECK("Trying to set shared fragment values, without adding new shared fragments", 1);
+			EntityManager->BatchAddSharedFragmentsForEntities(MakeArrayView(&EntityCollection, 1), SharedIntValues);
+		}
+		{
+			FMassArchetypeSharedFragmentValues DifferentSharedIntValues;
+			FConstSharedStruct OtherConstSharedFragment = FConstSharedStruct::Make<FTestSharedFragment_Int>(OtherTestIntValue);
+			DifferentSharedIntValues.AddConstSharedFragment(OtherConstSharedFragment);
+
+			AITEST_SCOPED_CHECK("Setting shared fragment values without archetype change is not supported", 1);
+			AITEST_SCOPED_CHECK("Trying to set shared fragment values, without adding new shared fragments", 1);
+			EntityManager->BatchAddSharedFragmentsForEntities(MakeArrayView(&EntityCollection, 1), DifferentSharedIntValues);
+		}
+
+		return true;
+	}
+};
+IMPLEMENT_AI_INSTANT_TEST(FSharedFragment_BatchSetAttempt, "System.Mass.SharedFragments.BatchSetAttempt");
+
+struct FSharedFragment_BatchAddToEmpty : FEntityTestBase
+{
+	virtual bool InstantTest() override
+	{
+		constexpr int32 NumToReserve = 32;
+
+		FMassArchetypeSharedFragmentValues SharedIntValues;
+		FConstSharedStruct ConstSharedFragment = FConstSharedStruct::Make<FTestSharedFragment_Int>();
+		SharedIntValues.AddConstSharedFragment(ConstSharedFragment);
+
+		TArray<FMassEntityHandle> ReservedEntityHandles;
+		EntityManager->BatchReserveEntities(NumToReserve, ReservedEntityHandles);
+		
+		FMassArchetypeEntityCollection EntityCollection(FMassArchetypeHandle(), ReservedEntityHandles, FMassArchetypeEntityCollection::EDuplicatesHandling::NoDuplicates);
+		// attempting to add the values before the entities are created is not a valid operation
+		{
+			AITEST_SCOPED_CHECK("Adding shared fragments to archetype-less entities is not supported", 1);
+			EntityManager->BatchAddSharedFragmentsForEntities(MakeArrayView(&EntityCollection, 1), SharedIntValues);
+		}
+		return true;
+	}
+};
+IMPLEMENT_AI_INSTANT_TEST(FSharedFragment_BatchAddToEmpty, "System.Mass.SharedFragments.BatchAddToEmpty");
+
 } // FMassEntityTest
 
 UE_ENABLE_OPTIMIZATION_SHIP
