CIS issue #450660
added explicit type casts to newly added shared fragment unit tests

diff --git a/Engine/Plugins/Runtime/MassEntity/Source/MassEntityTestSuite/Private/MassSharedFragmentsTest.cpp b/Engine/Plugins/Runtime/MassEntity/Source/MassEntityTestSuite/Private/MassSharedFragmentsTest.cpp
--- a/Engine/Plugins/Runtime/MassEntity/Source/MassEntityTestSuite/Private/MassSharedFragmentsTest.cpp
+++ b/Engine/Plugins/Runtime/MassEntity/Source/MassEntityTestSuite/Private/MassSharedFragmentsTest.cpp
@@ -475,8 +475,8 @@ struct FSharedFragment_BatchAddToEntity : FEntityTestBase
 
 		const FMassArchetypeHandle InitialArchetype = FloatsArchetype;
 		const int32 EntitiesPerChunk = FMassArchetypeHelper::ArchetypeDataFromHandleChecked(InitialArchetype).GetNumEntitiesPerChunk();
-		const int32 EntitiesToCreateNum = FMath::FloorToInt(EntitiesPerChunk * 2.2f);
-		const int32 EntitiesToMoveNum = FMath::FloorToInt(EntitiesPerChunk * 1.2f);
+		const int32 EntitiesToCreateNum = FMath::FloorToInt(float(EntitiesPerChunk) * 2.2f);
+		const int32 EntitiesToMoveNum = FMath::FloorToInt(float(EntitiesPerChunk) * 1.2f);
 
 		TArray<FMassEntityHandle> CreatedEntityHandles;
 
@@ -515,8 +515,8 @@ struct FSharedFragment_BatchSetAttempt : FEntityTestBase
 		constexpr int32 OtherTestIntValue = TestIntValue + 1;
 
 		const int32 EntitiesPerChunk = FMassArchetypeHelper::ArchetypeDataFromHandleChecked(FloatsArchetype).GetNumEntitiesPerChunk();
-		const int32 EntitiesToCreateNum = FMath::FloorToInt(EntitiesPerChunk * 2.2f);
-		const int32 EntitiesToMoveNum = FMath::FloorToInt(EntitiesPerChunk * 1.2f);
+		const int32 EntitiesToCreateNum = FMath::FloorToInt(float(EntitiesPerChunk) * 2.2f);
+		const int32 EntitiesToMoveNum = FMath::FloorToInt(float(EntitiesPerChunk) * 1.2f);
 
 		TArray<FMassEntityHandle> CreatedEntityHandles;
 		FMassArchetypeSharedFragmentValues SharedIntValues;
