Fixed mass processor game thread requirement

should consider the processor's thread requirement as well
removed unused bool in FMassEntityQuery

diff --git a/Engine/Plugins/Runtime/MassEntity/Source/MassEntity/Private/MassProcessor.cpp b/Engine/Plugins/Runtime/MassEntity/Source/MassEntity/Private/MassProcessor.cpp
--- a/Engine/Plugins/Runtime/MassEntity/Source/MassEntity/Private/MassProcessor.cpp
+++ b/Engine/Plugins/Runtime/MassEntity/Source/MassEntity/Private/MassProcessor.cpp
@@ -179,18 +179,18 @@ void UMassProcessor::PostInitProperties()
 	{
 		ConfigureQueries();
 
-		bool bNeedsGameThread = false;
-		for (FMassEntityQuery* QueryPtr : OwnedQueries)
+		bool bNeedsGameThread = ProcessorRequirements.DoesRequireGameThreadExecution();
+		for (const FMassEntityQuery* QueryPtr : OwnedQueries)
 		{
 			CA_ASSUME(QueryPtr);
 			bNeedsGameThread = (bNeedsGameThread || QueryPtr->DoesRequireGameThreadExecution());
 		}
-		
+
 		UE_CLOG(bRequiresGameThreadExecution != bNeedsGameThread, LogMass, Verbose
-			, TEXT("%s is marked bRequiresGameThreadExecution = %s, while the registered quries' requirement indicate the opposite")
+			, TEXT("%s is marked bRequiresGameThreadExecution = %s, while the registered queries' or processor requirements indicate the opposite")
 			, *GetProcessorName(), bRequiresGameThreadExecution ? TEXT("TRUE") : TEXT("FALSE"));
 
-		// better safe than sorry - if queries indicate the game thread execution is required then we marked the whole processor as such
+		// better safe than sorry - if queries or processor requirements indicate the game thread execution is required, then we marked the whole processor as such
 		bRequiresGameThreadExecution = bRequiresGameThreadExecution || bNeedsGameThread;
 	}
 #if CPUPROFILERTRACE_ENABLED
diff --git a/Engine/Plugins/Runtime/MassEntity/Source/MassEntity/Public/MassEntityQuery.h b/Engine/Plugins/Runtime/MassEntity/Source/MassEntity/Public/MassEntityQuery.h
--- a/Engine/Plugins/Runtime/MassEntity/Source/MassEntity/Public/MassEntityQuery.h
+++ b/Engine/Plugins/Runtime/MassEntity/Source/MassEntity/Public/MassEntityQuery.h
@@ -172,7 +172,6 @@ private:
 	 *	will cause crashes if the underlying code does try to issue commands. 
 	 */
 	uint8 bAllowParallelCommands : 1 = true;
-	uint8 bRequiresGameThreadExecution : 1 = false;
 	uint8 bRequiresMutatingWorldAccess : 1 = false;
 
 	EMassExecutionContextType ExpectedContextType = EMassExecutionContextType::Local;
