Fixed FMassEntityQuery::ParallelForEachEntityChunk to ensure the command buffers get created and used in the same thread

diff --git a/Engine/Plugins/Runtime/MassEntity/Source/MassEntity/Private/MassEntityQuery.cpp b/Engine/Plugins/Runtime/MassEntity/Source/MassEntity/Private/MassEntityQuery.cpp
--- a/Engine/Plugins/Runtime/MassEntity/Source/MassEntity/Private/MassEntityQuery.cpp
+++ b/Engine/Plugins/Runtime/MassEntity/Source/MassEntity/Private/MassEntityQuery.cpp
@@ -375,10 +375,18 @@ void FMassEntityQuery::ParallelForEachEntityChunk(FMassEntityManager& EntityMana
 		{
 			struct FTaskContext
 			{
-				FTaskContext()
+				FTaskContext() = default;
+
+				TSharedPtr<FMassCommandBuffer> GetCommandBuffer()
 				{
-					CommandBuffer = MakeShared<FMassCommandBuffer>();
+					if (!CommandBuffer)
+					{
+						// lazily creating the command buffer to ensure we create it in the same thread it's going to be used in
+						CommandBuffer = MakeShared<FMassCommandBuffer>();
+					}
+					return CommandBuffer;
 				}
+			private:
 				TSharedPtr<FMassCommandBuffer> CommandBuffer;
 			};
 
@@ -388,7 +396,7 @@ void FMassEntityQuery::ParallelForEachEntityChunk(FMassEntityManager& EntityMana
 				{
 					FMassExecutionContext LocalExecutionContext = ExecutionContext;
 
-					LocalExecutionContext.SetDeferredCommandBuffer(TaskContext.CommandBuffer);
+					LocalExecutionContext.SetDeferredCommandBuffer(TaskContext.GetCommandBuffer());
 
 					Jobs[JobIndex].Archetype.ExecutionFunctionForChunk(LocalExecutionContext, ExecuteFunction
 						, Jobs[JobIndex].ArchetypeIndex != INDEX_NONE ? ArchetypeFragmentMapping[Jobs[JobIndex].ArchetypeIndex] : FMassQueryRequirementIndicesMapping()
@@ -399,7 +407,7 @@ void FMassEntityQuery::ParallelForEachEntityChunk(FMassEntityManager& EntityMana
 			// merge all command buffers
 			for (FTaskContext& CommandContext : TaskContext)
 			{
-				ExecutionContext.Defer().MoveAppend(*CommandContext.CommandBuffer);
+				ExecutionContext.Defer().MoveAppend(*CommandContext.GetCommandBuffer());
 			}
 		}
 		else
