Fixed Mass commands appending. Before the change the MassEntityManager's command buffer was always growing (in some use cases).

Also:
* added checks to ensure commands are not being pushed while the buffer is being flushed, or from a thread different then the one the given buffer was created in (that's an unsupported operation)
* changed FMassCommandBuffer::Flush to return a boolean indicating whether any commands have been actually executed
* changed implementation details of FMassCommandBuffer::HasPendingCommands to reuse freshly added "active commands counter"

diff --git a/Engine/Plugins/Runtime/MassEntity/Source/MassEntity/Private/MassCommandBuffer.cpp b/Engine/Plugins/Runtime/MassEntity/Source/MassEntity/Private/MassCommandBuffer.cpp
--- a/Engine/Plugins/Runtime/MassEntity/Source/MassEntity/Private/MassCommandBuffer.cpp
+++ b/Engine/Plugins/Runtime/MassEntity/Source/MassEntity/Private/MassCommandBuffer.cpp
@@ -65,6 +65,11 @@ std::atomic<uint32> FMassBatchedCommand::CommandsCounter;
 //////////////////////////////////////////////////////////////////////
 // FMassCommandBuffer
 
+FMassCommandBuffer::FMassCommandBuffer()
+	: OwnerThreadId(FPlatformTLS::GetCurrentThreadId())
+{	
+}
+
 FMassCommandBuffer::~FMassCommandBuffer()
 {
 	ensureMsgf(HasPendingCommands() == false, TEXT("Destroying FMassCommandBuffer while there are still unprocessed commands. These operations will never be performed now."));
@@ -75,7 +80,7 @@ FMassCommandBuffer::~FMassCommandBuffer()
 	}
 }
 
-void FMassCommandBuffer::Flush(FMassEntityManager& EntityManager)
+bool FMassCommandBuffer::Flush(FMassEntityManager& EntityManager)
 {
 	check(!bIsFlushing);
 	TGuardValue FlushingGuard(bIsFlushing, true);
@@ -83,7 +88,7 @@ void FMassCommandBuffer::Flush(FMassEntityManager& EntityManager)
 	// short-circuit exit
 	if (HasPendingCommands() == false)
 	{
-		return;
+		return false;
 	}
 
 	{
@@ -115,18 +120,28 @@ void FMassCommandBuffer::Flush(FMassEntityManager& EntityManager)
 			bool IsValid() const { return GroupOrder < MAX_int32; }
 			bool operator<(const FBatchedCommandsSortedIndex& Other) const { return GroupOrder < Other.GroupOrder; }
 		};
+		
 		TArray<FBatchedCommandsSortedIndex> CommandsOrder;
-		CommandsOrder.Reserve(CommandInstances.Num());
-		for (int32 i = 0; i < CommandInstances.Num(); ++i)
+		const int32 OwnedCommandsCount = CommandInstances.Num();
+
+		CommandsOrder.Reserve(OwnedCommandsCount);
+		for (int32 i = 0; i < OwnedCommandsCount; ++i)
 		{
 			const FMassBatchedCommand* Command = CommandInstances[i];
 			CommandsOrder.Add(FBatchedCommandsSortedIndex(i, (Command && Command->HasWork())? CommandTypeOrder[(int)Command->GetOperationType()] : MAX_int32));
 		}
+		for (int32 i = 0; i < AppendedCommandInstances.Num(); ++i)
+		{
+			const FMassBatchedCommand* Command = AppendedCommandInstances[i];
+			CommandsOrder.Add(FBatchedCommandsSortedIndex(i + OwnedCommandsCount, (Command && Command->HasWork()) ? CommandTypeOrder[(int)Command->GetOperationType()] : MAX_int32));
+		}
 		CommandsOrder.StableSort();
 				
 		for (int32 k = 0; k < CommandsOrder.Num() && CommandsOrder[k].IsValid(); ++k)
 		{
-			FMassBatchedCommand* Command = CommandInstances[CommandsOrder[k].Index];
+			FMassBatchedCommand* Command = k < OwnedCommandsCount
+				? CommandInstances[CommandsOrder[k].Index]
+				: AppendedCommandInstances[CommandsOrder[k].Index - OwnedCommandsCount];
 			check(Command)
 
 #if CSV_PROFILER
@@ -145,18 +160,35 @@ void FMassCommandBuffer::Flush(FMassEntityManager& EntityManager)
 			Command->Execute(EntityManager);
 			Command->Reset();
 		}
+
+		// explicitly destroy the appended commands (stored in AppendedCommandInstances)
+		for (FMassBatchedCommand* Command : AppendedCommandInstances)
+		{
+			delete Command;
+		}
+		AppendedCommandInstances.Reset();
+
+		ActiveCommandsCounter = 0;
 	}
+
+	return true;
 }
  
 void FMassCommandBuffer::CleanUp()
 {
 	for (FMassBatchedCommand* Command : CommandInstances)
 	{
-		if (Command)
-		{
-			Command->Reset();
-		}
+		delete Command;
 	}
+	CommandInstances.Reset();
+
+	for (FMassBatchedCommand* Command : AppendedCommandInstances)
+	{
+		delete Command;
+	}
+	AppendedCommandInstances.Reset();
+
+	ActiveCommandsCounter = 0;
 }
 
 void FMassCommandBuffer::MoveAppend(FMassCommandBuffer& Other)
@@ -169,7 +201,10 @@ void FMassCommandBuffer::MoveAppend(FMassCommandBuffer& Other)
 	{
 		FScopeLock Lock(&AppendingCommandsCS);
 		UE_MT_SCOPED_WRITE_ACCESS(PendingBatchCommandsDetector);
-		CommandInstances.Append(MoveTemp(Other.CommandInstances));
+		AppendedCommandInstances.Append(MoveTemp(Other.CommandInstances));
+		AppendedCommandInstances.Append(MoveTemp(Other.AppendedCommandInstances));
+		ActiveCommandsCounter += Other.ActiveCommandsCounter;
+		Other.ActiveCommandsCounter = 0;
 	}
 }
 
diff --git a/Engine/Plugins/Runtime/MassEntity/Source/MassEntity/Public/MassCommandBuffer.h b/Engine/Plugins/Runtime/MassEntity/Source/MassEntity/Public/MassCommandBuffer.h
--- a/Engine/Plugins/Runtime/MassEntity/Source/MassEntity/Public/MassCommandBuffer.h
+++ b/Engine/Plugins/Runtime/MassEntity/Source/MassEntity/Public/MassCommandBuffer.h
@@ -11,42 +11,58 @@
 //@TODO: Consider debug information in case there is an assert when replaying the command buffer
 // (e.g., which system added the command, or even file/line number in development builds for the specific call via a macro)
 
+#define COMMAND_PUSHING_CHECK() \
+checkf(IsFlushing() == false, TEXT("Trying to push commands is not supported while the given buffer is being flushed")); \
+checkf(OwnerThreadId == FPlatformTLS::GetCurrentThreadId(), TEXT("Commands can be pushed only in the same thread where the command buffer was created."))
+
 struct MASSENTITY_API FMassCommandBuffer
 {
 public:
-	FMassCommandBuffer() = default;
+	FMassCommandBuffer();
 	~FMassCommandBuffer();
 
 	/** Adds a new entry to a given TCommand batch command instance */
 	template< template<typename... TArgs> typename TCommand, typename... TArgs >
 	void PushCommand(const FMassEntityHandle Entity, TArgs&&... InArgs)
 	{
+		COMMAND_PUSHING_CHECK();
+
 		LLM_SCOPE_BYNAME(TEXT("Mass/PushCommand"));
 		TCommand<TArgs...>& Instance = CreateOrAddCommand<TCommand<TArgs...>>();
 		Instance.Add(Entity, Forward<TArgs>(InArgs)...);
+		++ActiveCommandsCounter;
 	}
 
 	template<typename TCommand, typename... TArgs>
 	void PushCommand(TArgs&&... InArgs)
 	{
+		COMMAND_PUSHING_CHECK();
+
 		LLM_SCOPE_BYNAME(TEXT("Mass/PushCommand"));
 		TCommand& Instance = CreateOrAddCommand<TCommand>();
 		Instance.Add(Forward<TArgs>(InArgs)...);
+		++ActiveCommandsCounter;
 	}
 
 	/** Adds a new entry to a given TCommand batch command instance */
 	template< typename TCommand>
 	void PushCommand(const FMassEntityHandle Entity)
 	{
+		COMMAND_PUSHING_CHECK();
+
 		LLM_SCOPE_BYNAME(TEXT("Mass/PushCommand"));
 		CreateOrAddCommand<TCommand>().Add(Entity);
+		++ActiveCommandsCounter;
 	}
 
 	template< typename TCommand>
 	void PushCommand(TConstArrayView<FMassEntityHandle> Entities)
 	{
+		COMMAND_PUSHING_CHECK();
+
 		LLM_SCOPE_BYNAME(TEXT("Mass/PushCommand"));
 		CreateOrAddCommand<TCommand>().Add(Entities);
+		++ActiveCommandsCounter;
 	}
 
 	template<typename T>
@@ -145,13 +161,7 @@ public:
 
 	bool HasPendingCommands() const 
 	{
-		bool bHasCommands = false;		
-		for (int32 i = 0; i < CommandInstances.Num() && !bHasCommands; ++i)
-		{
-			bHasCommands = CommandInstances[i] && CommandInstances[i]->HasWork();
-		}
-
-		return bHasCommands;
+		return ActiveCommandsCounter > 0;
 	}
 	bool IsFlushing() const { return bIsFlushing; }
 
@@ -163,7 +173,6 @@ private:
 	{
 		const int32 Index = FMassBatchedCommand::GetCommandIndex<T>();
 
-		UE_MT_SCOPED_WRITE_ACCESS(PendingBatchCommandsDetector);
 		if (CommandInstances.IsValidIndex(Index) == false)
 		{
 			CommandInstances.AddZeroed(Index - CommandInstances.Num() + 1);
@@ -179,14 +188,38 @@ private:
 		return *NewCommandInstance;
 	}
 
-	void Flush(FMassEntityManager& EntityManager);
+	/** 
+	 * Executes all accumulated commands. 
+	 * @return whether any commands have actually been executed
+	 */
+	bool Flush(FMassEntityManager& EntityManager);
 	void CleanUp();
 
 	FCriticalSection AppendingCommandsCS;
 
 	UE_MT_DECLARE_RW_ACCESS_DETECTOR(PendingBatchCommandsDetector);
+	/** 
+	 * Commands created for this specific command buffer. All commands in the array are unique (by type) and reusable 
+	 * with subsequent PushCommand calls
+	 */
 	TArray<FMassBatchedCommand*> CommandInstances;
+	/** 
+	 * Commands appended to this command buffer (via FMassCommandBuffer::MoveAppend). These commands are just naive list
+	 * of commands, potentially containing duplicates with multiple MoveAppend calls. Once appended these commands are 
+	 * not being reused and consumed, destructively, during flushing
+	 */
+	TArray<FMassBatchedCommand*> AppendedCommandInstances;
+
+	int32 ActiveCommandsCounter = 0;
 
 	/** Indicates that this specific MassCommandBuffer is currently flushing its contents */
 	bool bIsFlushing = false;
+
+	/** 
+	 * Identifies the thread where given FMassCommandBuffer instance was created. Adding commands from other
+	 * threads is not supported and we use this value to check that
+	 */
+	const uint32 OwnerThreadId;
 };
+
+#undef COMMAND_PUSHING_CHECK
diff --git a/Engine/Plugins/Runtime/MassEntity/Source/MassEntity/Public/MassCommands.h b/Engine/Plugins/Runtime/MassEntity/Source/MassEntity/Public/MassCommands.h
--- a/Engine/Plugins/Runtime/MassEntity/Source/MassEntity/Public/MassCommands.h
+++ b/Engine/Plugins/Runtime/MassEntity/Source/MassEntity/Public/MassCommands.h
@@ -83,7 +83,7 @@ struct MASSENTITY_API FMassBatchedCommand
 		, DebugName(DebugName)
 	{}
 #endif // CSV_PROFILER || WITH_MASSENTITY_DEBUG
-	virtual ~FMassBatchedCommand() {}
+	virtual ~FMassBatchedCommand() { Reset(); }
 
 	virtual void Execute(FMassEntityManager& System) const = 0;
 	virtual void Reset()
@@ -95,7 +95,7 @@ struct MASSENTITY_API FMassBatchedCommand
 	EMassCommandOperationType GetOperationType() const { return OperationType; }
 
 	template<typename T>
-	FORCENOINLINE static uint32 GetCommandIndex()
+	static uint32 GetCommandIndex()
 	{
 		static const uint32 ThisTypesStaticIndex = CommandsCounter++;
 		return ThisTypesStaticIndex;
@@ -363,8 +363,7 @@ struct FMassCommandAddFragmentInstances : public FMassBatchedEntityCommand
 protected:
 	virtual void Reset() override
 	{
-		Fragments.Reset(); 
-		FragmentsAffected.Reset();
+		Fragments.Reset();
 		Super::Reset();
 	}
 
@@ -389,7 +388,7 @@ protected:
 	}
 
 	mutable UE::Mass::TMultiArray<TOthers...> Fragments;
-	FMassFragmentBitSet FragmentsAffected;
+	const FMassFragmentBitSet FragmentsAffected;
 };
 
 template<typename... TOthers>
