[Mass] Fix for array index out of bounds crash in FMassEntityManager::FlushCommands
NOTE: The fix was implemented by mikko.mononen
fyi bill.colby

diff --git a/Engine/Plugins/Runtime/MassEntity/Source/MassEntity/Private/MassCommandBuffer.cpp b/Engine/Plugins/Runtime/MassEntity/Source/MassEntity/Private/MassCommandBuffer.cpp
--- a/Engine/Plugins/Runtime/MassEntity/Source/MassEntity/Private/MassCommandBuffer.cpp
+++ b/Engine/Plugins/Runtime/MassEntity/Source/MassEntity/Private/MassCommandBuffer.cpp
@@ -139,9 +139,10 @@ bool FMassCommandBuffer::Flush(FMassEntityManager& EntityManager)
 				
 		for (int32 k = 0; k < CommandsOrder.Num() && CommandsOrder[k].IsValid(); ++k)
 		{
-			FMassBatchedCommand* Command = k < OwnedCommandsCount
-				? CommandInstances[CommandsOrder[k].Index]
-				: AppendedCommandInstances[CommandsOrder[k].Index - OwnedCommandsCount];
+			const int32 CommandIndex = CommandsOrder[k].Index;
+			FMassBatchedCommand* Command = CommandIndex < OwnedCommandsCount
+				? CommandInstances[CommandIndex]
+				: AppendedCommandInstances[CommandIndex - OwnedCommandsCount];
 			check(Command)
 
 #if CSV_PROFILER
