Switched MassCommandBuffer from storing raw command pointers to using TUniquePtrs

diff --git a/Engine/Plugins/Runtime/MassEntity/Source/MassEntity/Private/MassCommandBuffer.cpp b/Engine/Plugins/Runtime/MassEntity/Source/MassEntity/Private/MassCommandBuffer.cpp
--- a/Engine/Plugins/Runtime/MassEntity/Source/MassEntity/Private/MassCommandBuffer.cpp
+++ b/Engine/Plugins/Runtime/MassEntity/Source/MassEntity/Private/MassCommandBuffer.cpp
@@ -74,10 +74,7 @@ FMassCommandBuffer::~FMassCommandBuffer()
 {
 	ensureMsgf(HasPendingCommands() == false, TEXT("Destroying FMassCommandBuffer while there are still unprocessed commands. These operations will never be performed now."));
 
-	for (FMassBatchedCommand*& Command : CommandInstances)
-	{
-		delete Command;
-	}
+	CleanUp();
 }
 
 void FMassCommandBuffer::ForceUpdateCurrentThreadID()
@@ -132,12 +129,12 @@ bool FMassCommandBuffer::Flush(FMassEntityManager& EntityManager)
 		CommandsOrder.Reserve(OwnedCommandsCount);
 		for (int32 i = 0; i < OwnedCommandsCount; ++i)
 		{
-			const FMassBatchedCommand* Command = CommandInstances[i];
+			const TUniquePtr<FMassBatchedCommand>& Command = CommandInstances[i];
 			CommandsOrder.Add(FBatchedCommandsSortedIndex(i, (Command && Command->HasWork())? CommandTypeOrder[(int)Command->GetOperationType()] : MAX_int32));
 		}
 		for (int32 i = 0; i < AppendedCommandInstances.Num(); ++i)
 		{
-			const FMassBatchedCommand* Command = AppendedCommandInstances[i];
+			const TUniquePtr<FMassBatchedCommand>& Command = AppendedCommandInstances[i];
 			CommandsOrder.Add(FBatchedCommandsSortedIndex(i + OwnedCommandsCount, (Command && Command->HasWork()) ? CommandTypeOrder[(int)Command->GetOperationType()] : MAX_int32));
 		}
 		CommandsOrder.StableSort();
@@ -145,7 +142,7 @@ bool FMassCommandBuffer::Flush(FMassEntityManager& EntityManager)
 		for (int32 k = 0; k < CommandsOrder.Num() && CommandsOrder[k].IsValid(); ++k)
 		{
 			const int32 CommandIndex = CommandsOrder[k].Index;
-			FMassBatchedCommand* Command = CommandIndex < OwnedCommandsCount
+			TUniquePtr<FMassBatchedCommand>& Command = CommandIndex < OwnedCommandsCount
 				? CommandInstances[CommandIndex]
 				: AppendedCommandInstances[CommandIndex - OwnedCommandsCount];
 			check(Command)
@@ -167,11 +164,6 @@ bool FMassCommandBuffer::Flush(FMassEntityManager& EntityManager)
 			Command->Reset();
 		}
 
-		// explicitly destroy the appended commands (stored in AppendedCommandInstances)
-		for (FMassBatchedCommand* Command : AppendedCommandInstances)
-		{
-			delete Command;
-		}
 		AppendedCommandInstances.Reset();
 
 		ActiveCommandsCounter = 0;
@@ -182,16 +174,7 @@ bool FMassCommandBuffer::Flush(FMassEntityManager& EntityManager)
  
 void FMassCommandBuffer::CleanUp()
 {
-	for (FMassBatchedCommand* Command : CommandInstances)
-	{
-		delete Command;
-	}
 	CommandInstances.Reset();
-
-	for (FMassBatchedCommand* Command : AppendedCommandInstances)
-	{
-		delete Command;
-	}
 	AppendedCommandInstances.Reset();
 
 	ActiveCommandsCounter = 0;
@@ -217,7 +200,11 @@ void FMassCommandBuffer::MoveAppend(FMassCommandBuffer& Other)
 SIZE_T FMassCommandBuffer::GetAllocatedSize() const
 {
 	SIZE_T TotalSize = 0;
-	for (FMassBatchedCommand* Command : CommandInstances)
+	for (const TUniquePtr<FMassBatchedCommand>& Command : CommandInstances)
+	{
+		TotalSize += Command ? Command->GetAllocatedSize() : 0;
+	}
+	for (const TUniquePtr<FMassBatchedCommand>& Command : AppendedCommandInstances)
 	{
 		TotalSize += Command ? Command->GetAllocatedSize() : 0;
 	}
diff --git a/Engine/Plugins/Runtime/MassEntity/Source/MassEntity/Public/MassCommandBuffer.h b/Engine/Plugins/Runtime/MassEntity/Source/MassEntity/Public/MassCommandBuffer.h
--- a/Engine/Plugins/Runtime/MassEntity/Source/MassEntity/Public/MassCommandBuffer.h
+++ b/Engine/Plugins/Runtime/MassEntity/Source/MassEntity/Public/MassCommandBuffer.h
@@ -181,13 +181,11 @@ private:
 		}
 		else if (CommandInstances[Index])
 		{
-			return (T&)(*CommandInstances[Index]);
+			return (T&)(*CommandInstances[Index].Get());
 		}
 
-		T* NewCommandInstance = new T();
-		check(NewCommandInstance);
-		CommandInstances[Index] = NewCommandInstance;
-		return *NewCommandInstance;
+		CommandInstances[Index] = MakeUnique<T>();
+		return (T&)(*CommandInstances[Index].Get());
 	}
 
 	/** 
@@ -204,24 +202,24 @@ private:
 	 * Commands created for this specific command buffer. All commands in the array are unique (by type) and reusable 
 	 * with subsequent PushCommand calls
 	 */
-	TArray<FMassBatchedCommand*> CommandInstances;
+	TArray<TUniquePtr<FMassBatchedCommand>> CommandInstances;
 	/** 
 	 * Commands appended to this command buffer (via FMassCommandBuffer::MoveAppend). These commands are just naive list
 	 * of commands, potentially containing duplicates with multiple MoveAppend calls. Once appended these commands are 
 	 * not being reused and consumed, destructively, during flushing
 	 */
-	TArray<FMassBatchedCommand*> AppendedCommandInstances;
+	TArray<TUniquePtr<FMassBatchedCommand>> AppendedCommandInstances;
 
 	int32 ActiveCommandsCounter = 0;
 
 	/** Indicates that this specific MassCommandBuffer is currently flushing its contents */
 	bool bIsFlushing = false;
 
-	/**
-	* Identifies the thread where given FMassCommandBuffer instance was created. Adding commands from other
-	* threads is not supported and we use this value to check that.
-	* Note that it could be const since we set it in the constructor, but we need to recache on server forking.
-	*/
+	/** 
+	 * Identifies the thread where given FMassCommandBuffer instance was created. Adding commands from other
+	 * threads is not supported and we use this value to check that.
+	 * Note that it could be const since we set it in the constructor, but we need to recache on server forking.
+	 */
 	uint32 OwnerThreadId;
 };
 
