Tiny change to MassExecutor to avoid operations on an empty command buffer

diff --git a/Engine/Plugins/Runtime/MassEntity/Source/MassEntity/Private/MassExecutor.cpp b/Engine/Plugins/Runtime/MassEntity/Source/MassEntity/Private/MassExecutor.cpp
--- a/Engine/Plugins/Runtime/MassEntity/Source/MassEntity/Private/MassExecutor.cpp
+++ b/Engine/Plugins/Runtime/MassEntity/Source/MassEntity/Private/MassExecutor.cpp
@@ -125,7 +125,7 @@ void RunProcessorsView(TArrayView<UMassProcessor* const> Processors, FMassProces
 	}
 	// else make sure we don't just lose the commands. Append to the command buffer requested via
 	// ProcessingContext.CommandBuffer or to the default EntityManager's command buffer.
-	else if (CommandBuffer != ProcessingContext.CommandBuffer)
+	else if (CommandBuffer->HasPendingCommands() && CommandBuffer != ProcessingContext.CommandBuffer)
 	{
 		if (ProcessingContext.CommandBuffer)
 		{
