MassEntityManager's command buffer double buffering.

diff --git a/Engine/Plugins/Runtime/MassEntity/Source/MassEntity/Private/MassEntityManager.cpp b/Engine/Plugins/Runtime/MassEntity/Source/MassEntity/Private/MassEntityManager.cpp
--- a/Engine/Plugins/Runtime/MassEntity/Source/MassEntity/Private/MassEntityManager.cpp
+++ b/Engine/Plugins/Runtime/MassEntity/Source/MassEntity/Private/MassEntityManager.cpp
@@ -11,6 +11,7 @@
 #include "MassDebugger.h"
 #include "Misc/Fork.h"
 #include "Misc/CoreDelegates.h"
+#include "Algo/Find.h"
 
 
 const FMassEntityHandle FMassEntityManager::InvalidEntity;
@@ -61,11 +62,16 @@ FMassEntityManager::~FMassEntityManager()
 
 void FMassEntityManager::GetResourceSizeEx(FResourceSizeEx& CumulativeResourceSize)
 {
-	const SIZE_T MyExtraSize = Entities.GetAllocatedSize() + 
-		EntityFreeIndexList.GetAllocatedSize() +
-		(DeferredCommandBuffer != nullptr ? DeferredCommandBuffer->GetAllocatedSize() : 0) +
-		FragmentHashToArchetypeMap.GetAllocatedSize() +
-		FragmentTypeToArchetypeMap.GetAllocatedSize();
+	SIZE_T MyExtraSize = Entities.GetAllocatedSize()
+		+ EntityFreeIndexList.GetAllocatedSize()
+		+ FragmentHashToArchetypeMap.GetAllocatedSize()
+		+ FragmentTypeToArchetypeMap.GetAllocatedSize();
+
+	for (const TSharedPtr<FMassCommandBuffer>& CommandBuffer : DeferredCommandBuffers)
+	{
+		MyExtraSize += (CommandBuffer ? CommandBuffer->GetAllocatedSize() : 0);
+	}
+	
 	CumulativeResourceSize.AddDedicatedSystemMemoryBytes(MyExtraSize);
 
 	for (const auto& KVP : FragmentHashToArchetypeMap)
@@ -107,8 +113,11 @@ void FMassEntityManager::Initialize()
 	Entities.Add();
 	SerialNumberGenerator.fetch_add(FMath::Max(1,NumReservedEntities));
 
-	DeferredCommandBuffer = MakeShareable(new FMassCommandBuffer());
-	
+	for (TSharedPtr<FMassCommandBuffer>& CommandBuffer : DeferredCommandBuffers)
+	{
+		CommandBuffer = MakeShareable(new FMassCommandBuffer());
+	}
+
 	// if we get forked we need to update the command buffer's CurrentThreadID
 	if (FForkProcessHelper::IsForkRequested())
 	{
@@ -176,7 +185,13 @@ void FMassEntityManager::Deinitialize()
 		FCoreDelegates::OnPostFork.Remove(OnPostForkHandle);
 
 		// closing down so no point in actually flushing commands, but need to clean them up to avoid warnings on destruction
-		DeferredCommandBuffer->CleanUp();
+		for (TSharedPtr<FMassCommandBuffer>& CommandBuffer : DeferredCommandBuffers)
+		{
+			if (CommandBuffer)
+			{
+				CommandBuffer->CleanUp();
+			}
+		}
 
 #if WITH_MASSENTITY_DEBUG
 		FMassDebugger::UnregisterEntityManager(*this);
@@ -195,13 +210,16 @@ void FMassEntityManager::OnPostFork(EForkProcessRole Role)
 {
 	if (Role == EForkProcessRole::Child)
 	{
-		if (!DeferredCommandBuffer)
-		{
-			DeferredCommandBuffer = MakeShareable(new FMassCommandBuffer());
-		}
-		else
+		for (TSharedPtr<FMassCommandBuffer>& CommandBuffer : DeferredCommandBuffers)
 		{
-			DeferredCommandBuffer->ForceUpdateCurrentThreadID();
+			if (CommandBuffer)
+			{
+				CommandBuffer->ForceUpdateCurrentThreadID();
+			}
+			else
+			{
+				CommandBuffer = MakeShareable(new FMassCommandBuffer());
+			}
 		}
 	}
 }
@@ -1609,66 +1627,77 @@ void FMassEntityManager::GetMatchingArchetypes(const FMassFragmentRequirements&
 FMassExecutionContext FMassEntityManager::CreateExecutionContext(const float DeltaSeconds)
 {
 	FMassExecutionContext ExecutionContext(*this, DeltaSeconds);
-	ExecutionContext.SetDeferredCommandBuffer(DeferredCommandBuffer);
+	ExecutionContext.SetDeferredCommandBuffer(DeferredCommandBuffers[OpenedCommandBufferIndex]);
 	return MoveTemp(ExecutionContext);
 }
 
-void FMassEntityManager::FlushCommands(const TSharedPtr<FMassCommandBuffer>& InCommandBuffer)
+void FMassEntityManager::FlushCommands(TSharedPtr<FMassCommandBuffer>& InCommandBuffer)
+{
+	if (!ensureMsgf(IsInGameThread(), TEXT("Calling %hs is supported only on the Game Tread"), __FUNCTION__))
+	{
+		return;
+	}
+	if (!ensureMsgf(IsProcessing() == false, TEXT("Calling %hs is not supported while Mass Processing is active. Call FMassEntityManager::AppendCommands instead."), __FUNCTION__))
+	{
+		return;
+	}
+
+	if (InCommandBuffer && InCommandBuffer->HasPendingCommands()
+		&& (Algo::Find(DeferredCommandBuffers, InCommandBuffer) == nullptr))
+	{
+		AppendCommands(InCommandBuffer);
+	}
+	FlushCommands();
+}
+
+void FMassEntityManager::FlushCommands()
 {
 	constexpr int32 MaxIterations = 5;
 
-	if (InCommandBuffer)
+	if (!ensureMsgf(IsInGameThread(), TEXT("Calling %hs is supported only on the Game Tread"), __FUNCTION__))
 	{
-		if (InCommandBuffer->HasPendingCommands())
-		{
-			FlushedCommandBufferQueue.Enqueue(InCommandBuffer);
-		}
+		return;
 	}
-	else
+	if (!ensureMsgf(IsProcessing() == false, TEXT("Calling %hs is not supported while Mass Processing is active. Call FMassEntityManager::AppendCommands instead."), __FUNCTION__))
 	{
-		if (DeferredCommandBuffer->HasPendingCommands())
-		{
-			FlushedCommandBufferQueue.Enqueue(DeferredCommandBuffer);
-		}
+		return;
 	}
 
 	if (bCommandBufferFlushingInProgress == false && IsProcessing() == false)
 	{
+		ON_SCOPE_EXIT
+		{
+			bCommandBufferFlushingInProgress = false;
+		};
 		bCommandBufferFlushingInProgress = true;
-		
-		const int32 IterationsLimit = bFirstCommandFlush ? MAX_int32 : MaxIterations;
-		int32 IterationsCounter = 0;
-		TOptional<TSharedPtr<FMassCommandBuffer>> CurrentCommandBuffer = FlushedCommandBufferQueue.Dequeue();
 
-		while (IterationsCounter < IterationsLimit && CurrentCommandBuffer.IsSet())
+		int32 IterationCount = 0;
+		do 
 		{
-			IterationsCounter++;
-			(*CurrentCommandBuffer)->Flush(*this);
-			CurrentCommandBuffer = FlushedCommandBufferQueue.Dequeue();
-		}
-		ensure(IterationsCounter >= IterationsLimit || CurrentCommandBuffer.IsSet() == false);
-		UE_CVLOG_UELOG(IterationsCounter >= IterationsLimit, GetOwner(), LogMass, Error, TEXT("Reached loop count limit while flushing commands"));
+			const int32 CommandBufferIndexToFlush = OpenedCommandBufferIndex;
+
+			// buffer swap. Code instigated by observers can still use Defer() to push commands.
+			OpenedCommandBufferIndex = (OpenedCommandBufferIndex + 1) % DeferredCommandBuffers.Num();
+			ensureMsgf(DeferredCommandBuffers[OpenedCommandBufferIndex]->HasPendingCommands() == false
+				, TEXT("The freshly opened command buffer is expected to be empty upon switching"));
+
+			DeferredCommandBuffers[CommandBufferIndexToFlush]->Flush(*this);
 
-		bCommandBufferFlushingInProgress = false;
+			// repeat if there were commands submitted while commands were being flushed (by observers for example)
+		} while (DeferredCommandBuffers[OpenedCommandBufferIndex]->HasPendingCommands() && ++IterationCount < MaxIterations);
+
+		UE_CVLOG_UELOG(IterationCount >= MaxIterations, GetOwner(), LogMass, Error, TEXT("Reached loop count limit while flushing commands. Limiting the number of commands pushed during commands flushing could help."));
 	}
 }
 
 void FMassEntityManager::AppendCommands(TSharedPtr<FMassCommandBuffer>& InOutCommandBuffer)
 {
-	if (!ensureMsgf(InOutCommandBuffer != DeferredCommandBuffer, TEXT("We don't expect AppendCommands to be called with EntityManager's command buffer as the input parameter")))
+	if (!ensureMsgf(Algo::Find(DeferredCommandBuffers, InOutCommandBuffer) == nullptr
+		, TEXT("We don't expect AppendCommands to be called with EntityManager's command buffer as the input parameter")))
 	{
 		return;
 	}
-	else if (DeferredCommandBuffer->IsFlushing())
-	{
-		// in this case we'll add InOutCommandBuffer to FlushedCommandBufferQueue
-		FlushCommands(InOutCommandBuffer);
-	}
-	else
-	{
-		// otherwise we just move all the commands out of InOutCommandBuffer and into the main buffer
-		DeferredCommandBuffer->MoveAppend(*InOutCommandBuffer.Get());
-	}
+	Defer().MoveAppend(*InOutCommandBuffer.Get());
 }
 
 void FMassEntityManager::SetDebugName(const FString& NewDebugGame) 
diff --git a/Engine/Plugins/Runtime/MassEntity/Source/MassEntity/Public/MassEntityManager.h b/Engine/Plugins/Runtime/MassEntity/Source/MassEntity/Public/MassEntityManager.h
--- a/Engine/Plugins/Runtime/MassEntity/Source/MassEntity/Public/MassEntityManager.h
+++ b/Engine/Plugins/Runtime/MassEntity/Source/MassEntity/Public/MassEntityManager.h
@@ -453,13 +453,15 @@ public:
 	 */
 	bool IsProcessing() const { return ProcessingScopeCount > 0; }
 
-	FMassCommandBuffer& Defer() const { return *DeferredCommandBuffer.Get(); }
+	FMassCommandBuffer& Defer() const { return *DeferredCommandBuffers[OpenedCommandBufferIndex].Get(); }
 	/** 
 	 * @param InCommandBuffer if not set then the default command buffer will be flushed. If set and there's already 
 	 *		a command buffer being flushed (be it the main one or a previously requested one) then this command buffer 
 	 *		will be queue itself.
 	 */
-	void FlushCommands(const TSharedPtr<FMassCommandBuffer>& InCommandBuffer = TSharedPtr<FMassCommandBuffer>());
+	void FlushCommands(TSharedPtr<FMassCommandBuffer>& InCommandBuffer);
+
+	void FlushCommands();
 
 	/** 
 	 * Depending on the current state of Manager's command buffer the function will either move all the commands out of 
@@ -614,8 +616,15 @@ private:
 	TArray<int32> EntityFreeIndexList;
 
 	std::atomic<bool> bCommandBufferFlushingInProgress = false;
-	TSharedPtr<FMassCommandBuffer> DeferredCommandBuffer;
-	TMpscQueue<TSharedPtr<FMassCommandBuffer>> FlushedCommandBufferQueue;
+	/**
+	 * This index will be enough to control which buffer is available for pushing commands since flashing is taking place 
+	 * in the game thread and pushing commands to the buffer fetched by Defer() is only supported also on the game thread
+	 * (due to checking the cached thread ID).
+	 * The whole CL aims to support non-mass code trying to push commands while the flushing is going on (as triggered
+	 * by MassObservers reacting to the commands being flushed currently).
+	 */
+	uint8 OpenedCommandBufferIndex = 0;
+	TStaticArray<TSharedPtr<FMassCommandBuffer>, 2> DeferredCommandBuffers;
 
 	std::atomic<int32> SerialNumberGenerator = 0;
 	std::atomic<int32> ProcessingScopeCount = 0;
diff --git a/Engine/Plugins/Runtime/MassEntity/Source/MassEntityTestSuite/Private/MassCommandsTest.cpp b/Engine/Plugins/Runtime/MassEntity/Source/MassEntityTestSuite/Private/MassCommandsTest.cpp
--- a/Engine/Plugins/Runtime/MassEntity/Source/MassEntityTestSuite/Private/MassCommandsTest.cpp
+++ b/Engine/Plugins/Runtime/MassEntity/Source/MassEntityTestSuite/Private/MassCommandsTest.cpp
@@ -5,7 +5,7 @@
 #include "MassEntityTestTypes.h"
 #include "MassEntityTypes.h"
 #include "MassEntityView.h"
-
+#include "MassExecutionContext.h"
 #include "Algo/Sort.h"
 #include "Algo/RandomShuffle.h"
 
@@ -241,6 +241,60 @@ struct FCommands_DeferredFunction : FEntityTestBase
 };
 IMPLEMENT_AI_INSTANT_TEST(FCommands_DeferredFunction, "System.Mass.Commands.DeferredFunction");
 
+// pushing commands while the main buffer is being flushed
+struct FCommands_PushWhileFlushing : FEntityTestBase
+{
+	virtual bool InstantTest() override
+	{
+		constexpr int32 Count = 5;
+
+		// here's what we want to do:
+		// 1. Create a Count number of Int entities
+		// 2. Register TagA observer that will add a float fragment when the tag is added
+		//	a. The observer will use EntityManager.Defer() directly for the testing purposes - it should use Context.Defer() in real world scenarios
+		// 3. Add TagA to all the created Entities
+		// 4. Test if all the affected entities have the float fragment after the flushing
+
+		TArray<FMassEntityHandle> Entities;
+		EntityManager->BatchCreateEntities(IntsArchetype, Count, Entities);
+		for (const FMassEntityHandle& EntityHandle : Entities)
+		{
+			AITEST_NULL(TEXT("None of the freshly created entities is expexted to contain a float fragment")
+				, EntityManager->GetFragmentDataPtr<FTestFragment_Float>(EntityHandle));
+		}
+
+
+		UMassTestProcessorBase* ObserverProcessor = NewObject<UMassTestProcessorBase>();
+		ObserverProcessor->ForEachEntityChunkExecutionFunction = [](FMassExecutionContext& Context)
+			{
+				for (const FMassEntityHandle& EntityHandle : Context.GetEntities())
+				{
+					Context.GetEntityManagerChecked().Defer().AddFragment<FTestFragment_Float>(EntityHandle);
+				}
+			};
+		FMassObserverManager& ObserverManager = EntityManager->GetObserverManager();
+		ObserverManager.AddObserverInstance(*FTestTag_A::StaticStruct(), EMassObservedOperation::Add, *ObserverProcessor);
+
+		EntityManager->Defer().PushCommand<FMassCommandAddTag<FTestTag_A>>(Entities);
+		for (const FMassEntityHandle& EntityHandle : Entities)
+		{
+			AITEST_NULL(TEXT("Pushing the AddTag command should not result in adding the float fragment")
+				, EntityManager->GetFragmentDataPtr<FTestFragment_Float>(EntityHandle));
+		}
+
+		EntityManager->FlushCommands();
+		
+		for (const FMassEntityHandle& EntityHandle : Entities)
+		{
+			AITEST_NOT_NULL(TEXT("After flushing all the observed entities should have the float fragment")
+				, EntityManager->GetFragmentDataPtr<FTestFragment_Float>(EntityHandle));
+		}
+
+		return true;
+	}
+};
+IMPLEMENT_AI_INSTANT_TEST(FCommands_PushWhileFlushing, "System.Mass.Commands.PushWhileFlushing");
+
 #endif // WITH_MASSENTITY_DEBUG
 } // FMassCommandsTest
 
