Fix to FMassEntityManager::CreateEntity that would produce undefined behavior if the provided archetype did not match the provided shared fragment values

diff --git a/Engine/Plugins/Runtime/MassEntity/Source/MassEntity/Private/MassEntityManager.cpp b/Engine/Plugins/Runtime/MassEntity/Source/MassEntity/Private/MassEntityManager.cpp
--- a/Engine/Plugins/Runtime/MassEntity/Source/MassEntity/Private/MassEntityManager.cpp
+++ b/Engine/Plugins/Runtime/MassEntity/Source/MassEntity/Private/MassEntityManager.cpp
@@ -451,7 +451,8 @@ FMassEntityHandle FMassEntityManager::CreateEntity(const FMassArchetypeHandle& A
 	check(ArchetypeHandle.IsValid());
 
 	const FMassEntityHandle Entity = ReserveEntity();
-	InternalBuildEntity(Entity, ArchetypeHandle, SharedFragmentValues);
+	InternalBuildEntity(Entity, GetOrCreateSuitableArchetype(ArchetypeHandle, SharedFragmentValues.GetSharedFragmentBitSet()), SharedFragmentValues);
+
 	return Entity;
 }
 
