Implementation to concurrently reserve entities in Mass
* Default Mass still uses single threaded implementation

Notable changes:
* Extracted the entity allocation into an interface.  Implementation will be inlined into MassEntityManager as it was before.
* Concurrent implementation uses fix chunked array instead of TChunkArray - there isn't a generic container for this as far as I am aware - would be nice to have that too to clean up some implementation.
* Changed all places where MassEntityManager uses Entities and FEntityData and EntityFreeList directly and pushed through the interface.

Limitations:
* MassEntityManager will only support the concurrent backend in editor builds - other build configurations will be limited to single threaded mode for now.  Non-editor build configurations should devirtualize all calls.
* Tests all use single threaded configuration still. Validated that tests all pass in concurrent configuration.

diff --git a/Engine/Plugins/Runtime/MassEntity/Source/MassEntity/Private/MassDebugger.cpp b/Engine/Plugins/Runtime/MassEntity/Source/MassEntity/Private/MassDebugger.cpp
--- a/Engine/Plugins/Runtime/MassEntity/Source/MassEntity/Private/MassDebugger.cpp
+++ b/Engine/Plugins/Runtime/MassEntity/Source/MassEntity/Private/MassDebugger.cpp
@@ -4,6 +4,7 @@
 #if WITH_MASSENTITY_DEBUG
 #include "MassProcessor.h"
 #include "MassEntityManager.h"
+#include "MassEntityManagerStorage.h"
 #include "MassEntitySubsystem.h"
 #include "MassArchetypeTypes.h"
 #include "MassArchetypeData.h"
@@ -563,21 +564,20 @@ void FMassDebugger::OutputArchetypeDescription(FOutputDevice& Ar, const FMassArc
 
 void FMassDebugger::OutputEntityDescription(FOutputDevice& Ar, const FMassEntityManager& EntityManager, const int32 EntityIndex, const TCHAR* InPrefix)
 {
-	if (EntityIndex >= EntityManager.Entities.Num())
+	if (EntityIndex >= EntityManager.DebugGetEntityStorageInterface().Num())
 	{
 		Ar.Logf(ELogVerbosity::Log, TEXT("Unable to list fragments values for out of range index in EntityManager owned by %s"), *GetPathNameSafe(EntityManager.GetOwner()));
 		return;
 	}
-
-	const FMassEntityManager::FEntityData& EntityData = EntityManager.Entities[EntityIndex];
-	if (!EntityData.IsValid())
+	
+	if (!EntityManager.DebugGetEntityStorageInterface().IsValid(EntityIndex))
 	{
 		Ar.Logf(ELogVerbosity::Log, TEXT("Unable to list fragments values for invalid entity in EntityManager owned by %s"), *GetPathNameSafe(EntityManager.GetOwner()));
 	}
-
+	
 	FMassEntityHandle Entity;
 	Entity.Index = EntityIndex;
-	Entity.SerialNumber = EntityData.SerialNumber;
+	Entity.SerialNumber = EntityManager.DebugGetEntityStorageInterface().GetSerialNumber(EntityIndex);
 	OutputEntityDescription(Ar, EntityManager, Entity, InPrefix);
 }
 
@@ -590,8 +590,7 @@ void FMassDebugger::OutputEntityDescription(FOutputDevice& Ar, const FMassEntity
 
 	Ar.Logf(ELogVerbosity::Log, TEXT("Listing fragments values for Entity[%s] in EntityManager owned by %s"), *Entity.DebugGetDescription(), *GetPathNameSafe(EntityManager.GetOwner()));
 
-	const FMassEntityManager::FEntityData& EntityData = EntityManager.Entities[Entity.Index];
-	FMassArchetypeData* Archetype = EntityData.CurrentArchetype.Get();
+	FMassArchetypeData* Archetype = EntityManager.DebugGetEntityStorageInterface().GetArchetypeAsShared(Entity.Index).Get();
 	if (Archetype == nullptr)
 	{
 		Ar.Logf(ELogVerbosity::Log, TEXT("Unable to list fragments values for invalid entity in EntityManager owned by %s"), *GetPathNameSafe(EntityManager.GetOwner()));
diff --git a/Engine/Plugins/Runtime/MassEntity/Source/MassEntity/Private/MassEntityManager.cpp b/Engine/Plugins/Runtime/MassEntity/Source/MassEntity/Private/MassEntityManager.cpp
--- a/Engine/Plugins/Runtime/MassEntity/Source/MassEntity/Private/MassEntityManager.cpp
+++ b/Engine/Plugins/Runtime/MassEntity/Source/MassEntity/Private/MassEntityManager.cpp
@@ -3,6 +3,7 @@
 #include "MassEntityManager.h"
 #include "MassArchetypeData.h"
 #include "MassCommandBuffer.h"
+#include "MassEntityManagerStorage.h"
 #include "HAL/IConsoleManager.h"
 #include "Engine/World.h"
 #include "UObject/UObjectIterator.h"
@@ -38,10 +39,86 @@ namespace UE::Mass::Private
 			}
 		}
 	}
+
+	// Index 0 is a sentinel for an Empty/Unset EntityHandle
+	constexpr int32 InvalidEntityIndex = 0;
 }
 
-//////////////////////////////////////////////////////////////////////
-// FMassEntityManager
+#if MASS_CONCURRENT_RESERVE
+UE::Mass::IEntityStorageInterface& FMassEntityManager::GetEntityStorageInterface()
+{
+	using namespace UE::Mass;
+	struct StorageSelector
+	{
+		UE::Mass::IEntityStorageInterface* operator()(FEmptyVariantState&) const
+		{
+			checkf(false, TEXT("Attempt to use EntityStorageInterface without initialization"));
+			return nullptr;
+		}
+		UE::Mass::IEntityStorageInterface* operator()(FSingleThreadedEntityStorage& Storage) const
+		{
+			return &Storage;
+		}
+		UE::Mass::IEntityStorageInterface* operator()(FConcurrentEntityStorage& Storage) const
+		{
+			return &Storage;
+		}
+	};
+
+	UE::Mass::IEntityStorageInterface* Interface = Visit(StorageSelector{}, EntityStorage);
+
+	return *Interface;
+}
+
+const UE::Mass::IEntityStorageInterface& FMassEntityManager::GetEntityStorageInterface() const
+{
+	using namespace UE::Mass;
+	struct StorageSelector
+	{
+		const UE::Mass::IEntityStorageInterface* operator()(const FEmptyVariantState&) const
+		{
+			checkf(false, TEXT("Attempt to use EntityStorageInterface without initialization"));
+			return nullptr;
+		}
+		const UE::Mass::IEntityStorageInterface* operator()(const FSingleThreadedEntityStorage& Storage) const
+		{
+			return &Storage;
+		}
+		const UE::Mass::IEntityStorageInterface* operator()(const FConcurrentEntityStorage& Storage) const
+		{
+			return &Storage;
+		}
+	};
+
+	const UE::Mass::IEntityStorageInterface* Interface = Visit(StorageSelector{}, EntityStorage);
+
+	return *Interface;
+}
+#else
+UE::Mass::FSingleThreadedEntityStorage& FMassEntityManager::GetEntityStorageInterface()
+{
+	// Get will assert if not initialized
+	return EntityStorage.Get<UE::Mass::FSingleThreadedEntityStorage>();
+}
+
+const UE::Mass::FSingleThreadedEntityStorage& FMassEntityManager::GetEntityStorageInterface() const
+{
+	// Get will assert if not initialized
+	return EntityStorage.Get<UE::Mass::FSingleThreadedEntityStorage>();
+}
+#endif
+
+#if WITH_MASSENTITY_DEBUG
+UE::Mass::IEntityStorageInterface& FMassEntityManager::DebugGetEntityStorageInterface()
+{
+	return GetEntityStorageInterface();
+}
+
+const UE::Mass::IEntityStorageInterface& FMassEntityManager::DebugGetEntityStorageInterface() const
+{
+	return GetEntityStorageInterface();
+}
+#endif
 
 FMassEntityManager::FMassEntityManager(UObject* InOwner)
 	: ObserverManager(*this)
@@ -52,6 +129,9 @@ FMassEntityManager::FMassEntityManager(UObject* InOwner)
 #endif
 }
 
+//////////////////////////////////////////////////////////////////////
+// FMassEntityManager
+
 FMassEntityManager::~FMassEntityManager()
 {
 	if (bInitialized)
@@ -62,8 +142,7 @@ FMassEntityManager::~FMassEntityManager()
 
 void FMassEntityManager::GetResourceSizeEx(FResourceSizeEx& CumulativeResourceSize)
 {
-	SIZE_T MyExtraSize = Entities.GetAllocatedSize()
-		+ EntityFreeIndexList.GetAllocatedSize()
+	SIZE_T MyExtraSize = GetEntityStorageInterface().GetAllocatedSize()
 		+ FragmentHashToArchetypeMap.GetAllocatedSize()
 		+ FragmentTypeToArchetypeMap.GetAllocatedSize();
 
@@ -94,13 +173,42 @@ void FMassEntityManager::AddReferencedObjects(FReferenceCollector& Collector)
 	{
 		Struct.AddStructReferencedObjects(Collector);
 	}
-
+ 
 	const class UScriptStruct* ScriptStruct = FMassObserverManager::StaticStruct();
 	TWeakObjectPtr<const UScriptStruct> ScriptStructPtr{ScriptStruct};
 	Collector.AddReferencedObjects(ScriptStructPtr, &ObserverManager);
 }
 
 void FMassEntityManager::Initialize()
+{
+	FMassEntityManagerStorageInitParams InitializationParams;
+	InitializationParams.Emplace<FMassEntityManager_InitParams_SingleThreaded>();
+	Initialize(InitializationParams);
+}
+
+namespace UE::Mass::Private
+{
+	struct FEntityStorageInitializer
+	{
+		void operator()(const FMassEntityManager_InitParams_SingleThreaded& Params)
+		{
+			EntityStorage->Emplace<UE::Mass::FSingleThreadedEntityStorage>();
+		}
+		void operator()(const FMassEntityManager_InitParams_Concurrent& Params)
+		{
+#if MASS_CONCURRENT_RESERVE
+			EntityStorage->Emplace<UE::Mass::FConcurrentEntityStorage>();
+			EntityStorage->Get<UE::Mass::FConcurrentEntityStorage>().Initialize(Params);
+#else
+			checkf(false, TEXT("Mass does not support this storage backend"));
+#endif
+		}
+		
+		FMassEntityManager::FEntityStorageContainerType* EntityStorage = nullptr;
+	};
+}
+
+void FMassEntityManager::Initialize(const FMassEntityManagerStorageInitParams& InitializationParams)
 {
 	if (bInitialized)
 	{
@@ -109,9 +217,11 @@ void FMassEntityManager::Initialize()
 		return;
 	}
 
+	Visit(UE::Mass::Private::FEntityStorageInitializer{&EntityStorage}, InitializationParams);
+
 	// Index 0 is reserved so we can treat that index as an invalid entity handle
-	Entities.Add();
-	SerialNumberGenerator.fetch_add(FMath::Max(1,NumReservedEntities));
+	const FMassEntityHandle SentinelEntity = GetEntityStorageInterface().AcquireOne();
+	check(SentinelEntity.Index == UE::Mass::Private::InvalidEntityIndex);
 
 	for (TSharedPtr<FMassCommandBuffer>& CommandBuffer : DeferredCommandBuffers)
 	{
@@ -197,6 +307,8 @@ void FMassEntityManager::Deinitialize()
 		FMassDebugger::UnregisterEntityManager(*this);
 #endif // WITH_MASSENTITY_DEBUG
 
+		EntityStorage.Emplace<FEmptyVariantState>();
+		
 		bInitialized = false;
 	}
 	else
@@ -401,15 +513,15 @@ FMassArchetypeHandle FMassEntityManager::GetArchetypeForEntity(FMassEntityHandle
 {
 	if (IsEntityValid(Entity))
 	{
-		return FMassArchetypeHelper::ArchetypeHandleFromData(Entities[Entity.Index].CurrentArchetype);
+		return FMassArchetypeHelper::ArchetypeHandleFromData(GetEntityStorageInterface().GetArchetypeAsShared(Entity.Index));
 	}
 	return FMassArchetypeHandle();
 }
 
 FMassArchetypeHandle FMassEntityManager::GetArchetypeForEntityUnsafe(FMassEntityHandle Entity) const
 {
-	check(Entities.IsValidIndex(Entity.Index));
-	return FMassArchetypeHelper::ArchetypeHandleFromData(Entities[Entity.Index].CurrentArchetype);
+	check(GetEntityStorageInterface().IsValidIndex(Entity.Index));
+	return FMassArchetypeHelper::ArchetypeHandleFromData(GetEntityStorageInterface().GetArchetypeAsShared(Entity.Index));
 }
 
 void FMassEntityManager::ForEachArchetypeFragmentType(const FMassArchetypeHandle& ArchetypeHandle, TFunction< void(const UScriptStruct* /*FragmentType*/)> Function)
@@ -468,19 +580,16 @@ FMassEntityHandle FMassEntityManager::CreateEntity(TConstArrayView<FInstancedStr
 	const FMassEntityHandle Entity = ReserveEntity();
 	InternalBuildEntity(Entity, ArchetypeHandle, SharedFragmentValues);
 
-	const FEntityData& EntityData = Entities[Entity.Index];
-	EntityData.CurrentArchetype->SetFragmentsData(Entity, FragmentInstanceList);
+	FMassArchetypeData* CurrentArchetype = GetEntityStorageInterface().GetArchetype(Entity.Index);
+	check(CurrentArchetype);
+	CurrentArchetype->SetFragmentsData(Entity, FragmentInstanceList);
 
 	return Entity;
 }
 
 FMassEntityHandle FMassEntityManager::ReserveEntity()
 {
-	// @todo: Need to add thread safety to the reservation of an entity
-	FMassEntityHandle Result;
-	Result.Index = (EntityFreeIndexList.Num() > 0) ? EntityFreeIndexList.Pop(EAllowShrinking::No) : Entities.Add();
-	Result.SerialNumber = SerialNumberGenerator.fetch_add(1);
-	Entities[Result.Index].SerialNumber = Result.SerialNumber;
+	FMassEntityHandle Result = GetEntityStorageInterface().AcquireOne();
 
 	return Result;
 }
@@ -523,8 +632,9 @@ void FMassEntityManager::BuildEntity(FMassEntityHandle Entity, TConstArrayView<F
 
 	InternalBuildEntity(Entity, ArchetypeHandle, SharedFragmentValues);
 
-	const FEntityData& EntityData = Entities[Entity.Index];
-	EntityData.CurrentArchetype->SetFragmentsData(Entity, FragmentInstanceList);
+	FMassArchetypeData* CurrentArchetype = GetEntityStorageInterface().GetArchetype(Entity.Index);
+	check(CurrentArchetype);
+	CurrentArchetype->SetFragmentsData(Entity, FragmentInstanceList);
 }
 
 TConstArrayView<FMassEntityHandle> FMassEntityManager::BatchReserveEntities(const int32 Count, TArray<FMassEntityHandle>& InOutEntities)
@@ -581,12 +691,15 @@ void FMassEntityManager::BatchBuildEntities(const FMassArchetypeEntityCollection
 	// to update the manager's information on these entities (stored in FMassEntityManager::Entities)
 	for (FMassEntityHandle& Entity : EntityHandles)
 	{
-		check(Entities.IsValidIndex(Entity.Index));
+		check(GetEntityStorageInterface().IsValidIndex(Entity.Index));
 
-		FEntityData& EntityData = Entities[Entity.Index];
-		checkf(EntityData.SerialNumber != 0, TEXT("Failing this check means we're trying to build entities that have not been reserved yet. Consider using BatchCreateEntities"));
-		Entity.SerialNumber = EntityData.SerialNumber;
-		EntityData.CurrentArchetype = TargetArchetypeHandle.DataPtr;
+		const UE::Mass::IEntityStorageInterface::EEntityState EntityState = GetEntityStorageInterface().GetEntityState(Entity.Index);
+		checkf(EntityState == UE::Mass::IEntityStorageInterface::EEntityState::Reserved, TEXT("Trying to build entities that are not reserved. Check all handles are reserved or consider using BatchCreateEntities"));
+
+		const int32 SerialNumber = GetEntityStorageInterface().GetSerialNumber(Entity.Index);
+		Entity.SerialNumber = SerialNumber;
+		
+		GetEntityStorageInterface().SetArchetypeFromShared(Entity.Index, TargetArchetypeHandle.DataPtr);
 	}
 
 	TargetArchetypeHandle.DataPtr->BatchAddEntities(EntityHandles, SharedFragmentValues, TargetArchetypeEntityRanges);
@@ -645,11 +758,10 @@ TSharedRef<FMassEntityManager::FEntityCreationContext> FMassEntityManager::Inter
 	for (FMassEntityHandle Entity : ReservedEntities)
 	{
 		check(IsEntityValid(Entity));
-		FEntityData& EntityData = Entities[Entity.Index];
-		checkf(!EntityData.CurrentArchetype.IsValid(), TEXT("Batch creating reserved entities can only use entities that have not been constructed yet."));
-		checkf(EntityData.SerialNumber != 0, TEXT("Failing this check means we're trying to build entities that have not been reserved yet"));
-		EntityData.CurrentArchetype = ArchetypeHandle.DataPtr;
-		EntityData.SerialNumber = Entity.SerialNumber;
+		const UE::Mass::IEntityStorageInterface::EEntityState EntityState = GetEntityStorageInterface().GetEntityState(Entity.Index);
+		checkf(EntityState == UE::Mass::IEntityStorageInterface::EEntityState::Reserved, TEXT("Trying to build entities that are not reserved. Check all handles are reserved or consider using BatchCreateEntities"));
+		
+		GetEntityStorageInterface().SetArchetypeFromShared(Entity.Index, ArchetypeHandle.DataPtr);
 	}
 
 	FMassArchetypeEntityCollection::FEntityRangeArray TargetArchetypeEntityRanges;
@@ -675,8 +787,7 @@ void FMassEntityManager::DestroyEntity(FMassEntityHandle Entity)
 	
 	CheckIfEntityIsActive(Entity);
 
-	const FEntityData& EntityData = Entities[Entity.Index];
-	FMassArchetypeData* Archetype = EntityData.CurrentArchetype.Get();
+	FMassArchetypeData* Archetype = GetEntityStorageInterface().GetArchetype(Entity.Index);
 
 	if (Archetype)
 	{
@@ -692,30 +803,27 @@ void FMassEntityManager::BatchDestroyEntities(TConstArrayView<FMassEntityHandle>
 	TRACE_CPUPROFILER_EVENT_SCOPE(Mass_BatchDestroyEntities);
 
 	checkf(IsProcessing() == false, TEXT("Synchronous API function %hs called during mass processing. Use asynchronous API instead."), __FUNCTION__);
-	
-	EntityFreeIndexList.Reserve(EntityFreeIndexList.Num() + InEntities.Num());
 
 	// @todo optimize, we can make savings by implementing Archetype->RemoveEntities()
 	for (const FMassEntityHandle Entity : InEntities)
 	{
-		if (Entities.IsValidIndex(Entity.Index) == false)
+		if (GetEntityStorageInterface().IsValidIndex(Entity.Index) == false)
 		{
 			continue;
 		}
 
-		FEntityData& EntityData = Entities[Entity.Index];
-		if (EntityData.SerialNumber != Entity.SerialNumber)
+		const int32 SerialNumber = GetEntityStorageInterface().GetSerialNumber(Entity.Index);
+		if (SerialNumber != Entity.SerialNumber)
 		{
 			continue;
 		}
 
-		FMassArchetypeData* Archetype = EntityData.CurrentArchetype.Get();
+		FMassArchetypeData* Archetype = GetEntityStorageInterface().GetArchetype(Entity.Index);
 		check(Archetype);
 		ObserverManager.OnPreEntityDestroyed(Archetype->GetCompositionDescriptor(), Entity);
 		Archetype->RemoveEntity(Entity);
 
-		EntityData.Reset();
-		EntityFreeIndexList.Add(Entity.Index);
+		GetEntityStorageInterface().ReleaseOne(Entity);
 	}
 }
 
@@ -740,24 +848,13 @@ void FMassEntityManager::BatchDestroyEntityChunks(const FMassArchetypeEntityColl
 
 		FMassArchetypeData& ArchetypeData = FMassArchetypeHelper::ArchetypeDataFromHandleChecked(EntityCollection.GetArchetype());
 		ArchetypeData.BatchDestroyEntityChunks(EntityCollection.GetRanges(), EntitiesRemoved);
-
-		EntityFreeIndexList.Reserve(EntityFreeIndexList.Num() + EntitiesRemoved.Num());
+		
+		GetEntityStorageInterface().Release(EntitiesRemoved);
 	}
 	else
 	{
 		UE::Mass::Private::ConvertArchetypelessSubchunksIntoEntityHandles(EntityCollection.GetRanges(), EntitiesRemoved);
-	}
-
-	for (const FMassEntityHandle& Entity : EntitiesRemoved)
-	{
-		check(Entities.IsValidIndex(Entity.Index));
-
-		FEntityData& EntityData = Entities[Entity.Index];
-		if (!bValidArchetype || EntityData.SerialNumber == Entity.SerialNumber)
-		{
-			EntityData.Reset();
-			EntityFreeIndexList.Add(Entity.Index);
-		}
+		GetEntityStorageInterface().ForceRelease(EntitiesRemoved);
 	}
 }
 
@@ -780,8 +877,9 @@ void FMassEntityManager::AddFragmentToEntity(FMassEntityHandle Entity, const USc
 	CheckIfEntityIsActive(Entity);
 
 	FMassFragmentBitSet Fragments = InternalAddFragmentListToEntityChecked(Entity, FMassFragmentBitSet(*FragmentType));
-	FEntityData& EntityData = Entities[Entity.Index];
-	void* FragmentData = EntityData.CurrentArchetype->GetFragmentDataForEntity(FragmentType, Entity.Index);
+	FMassArchetypeData* CurrentArchetype = GetEntityStorageInterface().GetArchetype(Entity.Index);
+	check(CurrentArchetype);
+	void* FragmentData = CurrentArchetype->GetFragmentDataForEntity(FragmentType, Entity.Index);
 	Initializer(FragmentData, *FragmentType);
 
 	const FMassArchetypeCompositionDescriptor Descriptor(MoveTemp(Fragments));
@@ -800,8 +898,7 @@ void FMassEntityManager::AddCompositionToEntity_GetDelta(FMassEntityHandle Entit
 {
 	CheckIfEntityIsActive(Entity);
 
-	FEntityData& EntityData = Entities[Entity.Index];
-	FMassArchetypeData* OldArchetype = EntityData.CurrentArchetype.Get();
+	FMassArchetypeData* OldArchetype = GetEntityStorageInterface().GetArchetype(Entity.Index);
 	check(OldArchetype);
 
 	InDescriptor.Fragments -= OldArchetype->GetCompositionDescriptor().Fragments;
@@ -817,13 +914,14 @@ void FMassEntityManager::AddCompositionToEntity_GetDelta(FMassEntityHandle Entit
 
 		const FMassArchetypeHandle NewArchetypeHandle = CreateArchetype(NewDescriptor);
 
-		if (ensure(NewArchetypeHandle.DataPtr != EntityData.CurrentArchetype))
+		if (ensure(NewArchetypeHandle.DataPtr.Get() != OldArchetype))
 		{
 			// Move the entity over
 			FMassArchetypeData& NewArchetype = FMassArchetypeHelper::ArchetypeDataFromHandleChecked(NewArchetypeHandle);
 			NewArchetype.CopyDebugNamesFrom(*OldArchetype);
-			EntityData.CurrentArchetype->MoveEntityToAnotherArchetype(Entity, NewArchetype);
-			EntityData.CurrentArchetype = NewArchetypeHandle.DataPtr;
+			OldArchetype->MoveEntityToAnotherArchetype(Entity, NewArchetype);
+
+			GetEntityStorageInterface().SetArchetypeFromShared(Entity.Index, NewArchetypeHandle.DataPtr);
 
 			ObserverManager.OnPostCompositionAdded(Entity, InDescriptor);
 		}
@@ -836,8 +934,7 @@ void FMassEntityManager::RemoveCompositionFromEntity(FMassEntityHandle Entity, c
 
 	if(InDescriptor.IsEmpty() == false)
 	{
-		FEntityData& EntityData = Entities[Entity.Index];
-		FMassArchetypeData* OldArchetype = EntityData.CurrentArchetype.Get();
+		FMassArchetypeData* OldArchetype = GetEntityStorageInterface().GetArchetype(Entity.Index);
 		check(OldArchetype);
 
 		FMassArchetypeCompositionDescriptor NewDescriptor = OldArchetype->GetCompositionDescriptor();
@@ -854,13 +951,13 @@ void FMassEntityManager::RemoveCompositionFromEntity(FMassEntityHandle Entity, c
 
 			const FMassArchetypeHandle NewArchetypeHandle = CreateArchetype(NewDescriptor);
 
-			if (ensure(NewArchetypeHandle.DataPtr != EntityData.CurrentArchetype))
+			if (ensure(NewArchetypeHandle.DataPtr.Get() != OldArchetype))
 			{
 				// Move the entity over
 				FMassArchetypeData& NewArchetype = FMassArchetypeHelper::ArchetypeDataFromHandleChecked(NewArchetypeHandle);
 				NewArchetype.CopyDebugNamesFrom(*OldArchetype);
-				EntityData.CurrentArchetype->MoveEntityToAnotherArchetype(Entity, NewArchetype);
-				EntityData.CurrentArchetype = NewArchetypeHandle.DataPtr;
+				OldArchetype->MoveEntityToAnotherArchetype(Entity, NewArchetype);
+				GetEntityStorageInterface().SetArchetypeFromShared(Entity.Index, NewArchetypeHandle.DataPtr);
 			}
 		}
 	}
@@ -874,24 +971,22 @@ const FMassArchetypeCompositionDescriptor& FMassEntityManager::GetArchetypeCompo
 
 void FMassEntityManager::InternalBuildEntity(FMassEntityHandle Entity, const FMassArchetypeHandle& ArchetypeHandle, const FMassArchetypeSharedFragmentValues& SharedFragmentValues)
 {
-	FEntityData& EntityData = Entities[Entity.Index];
-	EntityData.CurrentArchetype = ArchetypeHandle.DataPtr;
-	EntityData.CurrentArchetype->AddEntity(Entity, SharedFragmentValues);
+	const TSharedPtr<FMassArchetypeData>& NewArchetype = ArchetypeHandle.DataPtr;
+	GetEntityStorageInterface().SetArchetypeFromShared(Entity.Index, ArchetypeHandle.DataPtr);
+	NewArchetype->AddEntity(Entity, SharedFragmentValues);
 
-	ObserverManager.OnPostCompositionAdded(Entity, EntityData.CurrentArchetype->GetCompositionDescriptor());
+	ObserverManager.OnPostCompositionAdded(Entity, NewArchetype->GetCompositionDescriptor());
 }
 
 void FMassEntityManager::InternalReleaseEntity(FMassEntityHandle Entity)
 {
-	FEntityData& EntityData = Entities[Entity.Index];
-	EntityData.Reset();
-	EntityFreeIndexList.Add(Entity.Index);
+	// Using force release by bypass serial number check since we have verified the validity of the handle earlier.
+	GetEntityStorageInterface().ForceReleaseOne(Entity);
 }
 
 FMassFragmentBitSet FMassEntityManager::InternalAddFragmentListToEntityChecked(FMassEntityHandle Entity, const FMassFragmentBitSet& InFragments)
 {
-	const FEntityData& EntityData = Entities[Entity.Index];
-	FMassArchetypeData* OldArchetype = EntityData.CurrentArchetype.Get();
+	TSharedPtr<FMassArchetypeData>& OldArchetype = GetEntityStorageInterface().GetArchetypeAsShared(Entity.Index);
 	check(OldArchetype);
 
 	UE_CLOG(OldArchetype->GetFragmentBitSet().HasAny(InFragments), LogMass, Log
@@ -909,20 +1004,20 @@ FMassFragmentBitSet FMassEntityManager::InternalAddFragmentListToEntityChecked(F
 void FMassEntityManager::InternalAddFragmentListToEntity(FMassEntityHandle Entity, const FMassFragmentBitSet& InFragments)
 {
 	checkf(InFragments.IsEmpty() == false, TEXT("%hs is intended for internal calls with non empty NewFragments parameter"), __FUNCTION__);
-	check(Entities.IsValidIndex(Entity.Index));
-	FEntityData& EntityData = Entities[Entity.Index];
-	check(EntityData.CurrentArchetype.IsValid());
-	const FMassArchetypeData* OldArchetype = EntityData.CurrentArchetype.Get();
+	check(GetEntityStorageInterface().IsValidIndex(Entity.Index));
+	TSharedPtr<FMassArchetypeData>& OldArchetype = GetEntityStorageInterface().GetArchetypeAsShared(Entity.Index);
+	check(OldArchetype.IsValid());
 
 	// fetch or create the new archetype
-	const FMassArchetypeHandle NewArchetypeHandle = CreateArchetype(EntityData.CurrentArchetype, InFragments);
-	checkf(NewArchetypeHandle.DataPtr != EntityData.CurrentArchetype, TEXT("%hs is intended for internal calls with non overlapping fragment list."), __FUNCTION__);
+	const FMassArchetypeHandle NewArchetypeHandle = CreateArchetype(OldArchetype, InFragments);
+	checkf(NewArchetypeHandle.DataPtr != OldArchetype, TEXT("%hs is intended for internal calls with non overlapping fragment list."), __FUNCTION__);
 
 	// Move the entity over
 	FMassArchetypeData& NewArchetype = FMassArchetypeHelper::ArchetypeDataFromHandleChecked(NewArchetypeHandle);
 	NewArchetype.CopyDebugNamesFrom(*OldArchetype);
-	EntityData.CurrentArchetype->MoveEntityToAnotherArchetype(Entity, NewArchetype);
-	EntityData.CurrentArchetype = NewArchetypeHandle.DataPtr;
+	OldArchetype->MoveEntityToAnotherArchetype(Entity, NewArchetype);
+
+	GetEntityStorageInterface().SetArchetypeFromShared(Entity.Index, NewArchetypeHandle.DataPtr);
 }
 
 void FMassEntityManager::AddFragmentInstanceListToEntity(FMassEntityHandle Entity, TConstArrayView<FInstancedStruct> FragmentInstanceList)
@@ -933,9 +1028,10 @@ void FMassEntityManager::AddFragmentInstanceListToEntity(FMassEntityHandle Entit
 	checkf(FragmentInstanceList.Num() > 0, TEXT("Need to specify at least one fragment instances for this operation"));
 
 	const FMassArchetypeCompositionDescriptor Descriptor(InternalAddFragmentListToEntityChecked(Entity, FMassFragmentBitSet(FragmentInstanceList)));
-
-	const FEntityData& EntityData = Entities[Entity.Index];
-	EntityData.CurrentArchetype->SetFragmentsData(Entity, FragmentInstanceList);
+	
+	FMassArchetypeData* CurrentArchetype = GetEntityStorageInterface().GetArchetype(Entity.Index);
+	check(CurrentArchetype);
+	CurrentArchetype->SetFragmentsData(Entity, FragmentInstanceList);
 
 	ObserverManager.OnPostCompositionAdded(Entity, Descriptor);
 }
@@ -951,8 +1047,7 @@ void FMassEntityManager::RemoveFragmentListFromEntity(FMassEntityHandle Entity,
 
 	CheckIfEntityIsActive(Entity);
 	
-	FEntityData& EntityData = Entities[Entity.Index];
-	FMassArchetypeData* OldArchetype = EntityData.CurrentArchetype.Get();
+	FMassArchetypeData* OldArchetype = GetEntityStorageInterface().GetArchetype(Entity.Index);
 	check(OldArchetype);
 
 	const FMassFragmentBitSet FragmentsToRemove(FragmentList);
@@ -972,7 +1067,8 @@ void FMassEntityManager::RemoveFragmentListFromEntity(FMassEntityHandle Entity,
 		FMassArchetypeData& NewArchetype = FMassArchetypeHelper::ArchetypeDataFromHandleChecked(NewArchetypeHandle);
 		NewArchetype.CopyDebugNamesFrom(*OldArchetype);
 		OldArchetype->MoveEntityToAnotherArchetype(Entity, NewArchetype);
-		EntityData.CurrentArchetype = NewArchetypeHandle.DataPtr;
+		
+		GetEntityStorageInterface().SetArchetypeFromShared(Entity.Index, NewArchetypeHandle.DataPtr);
 	}
 }
 
@@ -984,9 +1080,8 @@ void FMassEntityManager::SwapTagsForEntity(FMassEntityHandle Entity, const UScri
 
 	checkf((OldTagType != nullptr) && OldTagType->IsChildOf(FMassTag::StaticStruct()), TEXT("%hs works only with tags while '%s' is not one."), __FUNCTION__, *GetPathNameSafe(OldTagType));
 	checkf((NewTagType != nullptr) && NewTagType->IsChildOf(FMassTag::StaticStruct()), TEXT("%hs works only with tags while '%s' is not one."), __FUNCTION__, *GetPathNameSafe(NewTagType));
-
-	FEntityData& EntityData = Entities[Entity.Index];
-	FMassArchetypeData* CurrentArchetype = EntityData.CurrentArchetype.Get();
+	
+	TSharedPtr<FMassArchetypeData>& CurrentArchetype = GetEntityStorageInterface().GetArchetypeAsShared(Entity.Index);
 	check(CurrentArchetype);
 
 	FMassTagBitSet NewTagBitSet = CurrentArchetype->GetTagBitSet();
@@ -995,12 +1090,13 @@ void FMassEntityManager::SwapTagsForEntity(FMassEntityHandle Entity, const UScri
 	
 	if (NewTagBitSet != CurrentArchetype->GetTagBitSet())
 	{
-		const FMassArchetypeHandle NewArchetypeHandle = InternalCreateSimilarArchetype(EntityData.CurrentArchetype, NewTagBitSet);
+		const FMassArchetypeHandle NewArchetypeHandle = InternalCreateSimilarArchetype(CurrentArchetype, NewTagBitSet);
 		checkSlow(NewArchetypeHandle.IsValid());
 
 		// Move the entity over
-		EntityData.CurrentArchetype->MoveEntityToAnotherArchetype(Entity, *NewArchetypeHandle.DataPtr.Get());
-		EntityData.CurrentArchetype = NewArchetypeHandle.DataPtr;
+		CurrentArchetype->MoveEntityToAnotherArchetype(Entity, *NewArchetypeHandle.DataPtr.Get());
+		
+		GetEntityStorageInterface().SetArchetypeFromShared(Entity.Index, NewArchetypeHandle.DataPtr);
 	}
 }
 
@@ -1009,9 +1105,8 @@ void FMassEntityManager::AddTagToEntity(FMassEntityHandle Entity, const UScriptS
 	checkf((TagType != nullptr) && TagType->IsChildOf(FMassTag::StaticStruct()), TEXT("%hs works only with tags while '%s' is not one."), __FUNCTION__, *GetPathNameSafe(TagType));
 
 	CheckIfEntityIsActive(Entity);
-
-	FEntityData& EntityData = Entities[Entity.Index];
-	FMassArchetypeData* CurrentArchetype = EntityData.CurrentArchetype.Get();
+	
+	TSharedPtr<FMassArchetypeData>& CurrentArchetype = GetEntityStorageInterface().GetArchetypeAsShared(Entity.Index);
 	check(CurrentArchetype);
 
 	if (CurrentArchetype->HasTagType(TagType) == false)
@@ -1019,12 +1114,12 @@ void FMassEntityManager::AddTagToEntity(FMassEntityHandle Entity, const UScriptS
 		//FMassTagBitSet NewTags = CurrentArchetype->GetTagBitSet() - *TagType;
 		FMassTagBitSet NewTags = CurrentArchetype->GetTagBitSet();
 		NewTags.Add(*TagType);
-		const FMassArchetypeHandle NewArchetypeHandle = InternalCreateSimilarArchetype(EntityData.CurrentArchetype, NewTags);
+		const FMassArchetypeHandle NewArchetypeHandle = InternalCreateSimilarArchetype(CurrentArchetype, NewTags);
 		checkSlow(NewArchetypeHandle.IsValid());
 
 		// Move the entity over
-		EntityData.CurrentArchetype->MoveEntityToAnotherArchetype(Entity, *NewArchetypeHandle.DataPtr.Get());
-		EntityData.CurrentArchetype = NewArchetypeHandle.DataPtr;
+		CurrentArchetype->MoveEntityToAnotherArchetype(Entity, *NewArchetypeHandle.DataPtr.Get());
+		GetEntityStorageInterface().SetArchetypeFromShared(Entity.Index, NewArchetypeHandle.DataPtr);
 
 		FMassArchetypeCompositionDescriptor CompositionDelta;
 		FMassTagBitSet TagDelta;
@@ -1040,8 +1135,7 @@ void FMassEntityManager::RemoveTagFromEntity(FMassEntityHandle Entity, const USc
 
 	CheckIfEntityIsActive(Entity);
 
-	FEntityData& EntityData = Entities[Entity.Index];
-	FMassArchetypeData* CurrentArchetype = EntityData.CurrentArchetype.Get();
+	TSharedPtr<FMassArchetypeData>& CurrentArchetype = GetEntityStorageInterface().GetArchetypeAsShared(Entity.Index);
 	check(CurrentArchetype);
 
 	if (CurrentArchetype->HasTagType(TagType))
@@ -1054,12 +1148,12 @@ void FMassEntityManager::RemoveTagFromEntity(FMassEntityHandle Entity, const USc
 		
 		// CurrentArchetype->GetTagBitSet() -  *TagType
 		const FMassTagBitSet NewTagComposition = CurrentArchetype->GetTagBitSet() - TagDelta;
-		const FMassArchetypeHandle NewArchetypeHandle = InternalCreateSimilarArchetype(EntityData.CurrentArchetype, NewTagComposition);
+		const FMassArchetypeHandle NewArchetypeHandle = InternalCreateSimilarArchetype(CurrentArchetype, NewTagComposition);
 		checkSlow(NewArchetypeHandle.IsValid());
 
 		// Move the entity over
-		EntityData.CurrentArchetype->MoveEntityToAnotherArchetype(Entity, *NewArchetypeHandle.DataPtr.Get());
-		EntityData.CurrentArchetype = NewArchetypeHandle.DataPtr;
+		CurrentArchetype->MoveEntityToAnotherArchetype(Entity, *NewArchetypeHandle.DataPtr.Get());
+		GetEntityStorageInterface().SetArchetypeFromShared(Entity.Index, NewArchetypeHandle.DataPtr);
 	}
 }
 
@@ -1071,9 +1165,8 @@ bool FMassEntityManager::AddConstSharedFragmentToEntity(const FMassEntityHandle
 	}
 
 	CheckIfEntityIsActive(Entity);
-
-	FEntityData& EntityData = Entities[Entity.Index];
-	FMassArchetypeData* CurrentArchetype = EntityData.CurrentArchetype.Get();
+	
+	FMassArchetypeData* CurrentArchetype = GetEntityStorageInterface().GetArchetypeAsShared(Entity.Index).Get();
 	check(CurrentArchetype);
 
 	const UScriptStruct* StructType = InConstSharedFragment.GetScriptStruct();
@@ -1107,7 +1200,7 @@ bool FMassEntityManager::AddConstSharedFragmentToEntity(const FMassEntityHandle
 	CurrentArchetype->MoveEntityToAnotherArchetype(Entity, *NewArchetype, &NewSharedFragmentValues);
 
 	// Change the entity archetype
-	EntityData.CurrentArchetype = NewArchetypeHandle.DataPtr;
+	GetEntityStorageInterface().SetArchetypeFromShared(Entity.Index, NewArchetypeHandle.DataPtr);
 
 	return true;
 }
@@ -1145,10 +1238,9 @@ void FMassEntityManager::BatchChangeTagsForEntities(TConstArrayView<FMassArchety
 
 			for (const FMassEntityHandle& Entity : EntitiesBeingMoved)
 			{
-				check(Entities.IsValidIndex(Entity.Index));
-
-				FEntityData& EntityData = Entities[Entity.Index];
-				EntityData.CurrentArchetype = NewArchetypeHandle.DataPtr;
+				check(GetEntityStorageInterface().IsValidIndex(Entity.Index));
+				
+				GetEntityStorageInterface().SetArchetypeFromShared(Entity.Index, NewArchetypeHandle.DataPtr);
 			}
 
 			if (bTagsAddedAreObserved)
@@ -1196,10 +1288,9 @@ void FMassEntityManager::BatchChangeFragmentCompositionForEntities(TConstArrayVi
 
 				for (const FMassEntityHandle& Entity : EntitiesBeingMoved)
 				{
-					check(Entities.IsValidIndex(Entity.Index));
+					check(GetEntityStorageInterface().IsValidIndex(Entity.Index));
 
-					FEntityData& EntityData = Entities[Entity.Index];
-					EntityData.CurrentArchetype = NewArchetypeHandle.DataPtr;
+					GetEntityStorageInterface().SetArchetypeFromShared(Entity.Index, NewArchetypeHandle.DataPtr);
 				}
 
 				if (bFragmentsAddedAreObserved)
@@ -1262,10 +1353,9 @@ void FMassEntityManager::BatchAddFragmentInstancesForEntities(TConstArrayView<FM
 
 				for (const FMassEntityHandle& Entity : EntitiesBeingMoved)
 				{
-					check(Entities.IsValidIndex(Entity.Index));
+					check(GetEntityStorageInterface().IsValidIndex(Entity.Index));
 
-					FEntityData& EntityData = Entities[Entity.Index];
-					EntityData.CurrentArchetype = NewArchetypeHandle.DataPtr;
+					GetEntityStorageInterface().SetArchetypeFromShared(Entity.Index, NewArchetypeHandle.DataPtr);
 				}
 
 				TargetArchetypeHandle = NewArchetypeHandle;
@@ -1323,10 +1413,9 @@ void FMassEntityManager::BatchAddSharedFragmentsForEntities(TConstArrayView<FMas
 
 		for (const FMassEntityHandle& Entity : EntitiesBeingMoved)
 		{
-			check(Entities.IsValidIndex(Entity.Index));
-
-			FEntityData& EntityData = Entities[Entity.Index];
-			EntityData.CurrentArchetype = NewArchetypeHandle.DataPtr;
+			check(GetEntityStorageInterface().IsValidIndex(Entity.Index));
+			
+			GetEntityStorageInterface().SetArchetypeFromShared(Entity.Index, NewArchetypeHandle.DataPtr);
 		}
 	}
 }
@@ -1338,17 +1427,19 @@ void FMassEntityManager::MoveEntityToAnotherArchetype(FMassEntityHandle Entity,
 	FMassArchetypeData& NewArchetype = FMassArchetypeHelper::ArchetypeDataFromHandleChecked(NewArchetypeHandle);
 
 	// Move the entity over
-	FEntityData& EntityData = Entities[Entity.Index];
-	EntityData.CurrentArchetype->MoveEntityToAnotherArchetype(Entity, NewArchetype);
-	EntityData.CurrentArchetype = NewArchetypeHandle.DataPtr;
+	FMassArchetypeData* CurrentArchetype = GetEntityStorageInterface().GetArchetype(Entity.Index);
+	check(CurrentArchetype);
+	CurrentArchetype->MoveEntityToAnotherArchetype(Entity, NewArchetype);
+	GetEntityStorageInterface().SetArchetypeFromShared(Entity.Index, NewArchetypeHandle.DataPtr);
 }
 
 void FMassEntityManager::SetEntityFragmentsValues(FMassEntityHandle Entity, TArrayView<const FInstancedStruct> FragmentInstanceList)
 {
 	CheckIfEntityIsActive(Entity);
 
-	const FEntityData& EntityData = Entities[Entity.Index];
-	EntityData.CurrentArchetype->SetFragmentsData(Entity, FragmentInstanceList);
+	FMassArchetypeData* CurrentArchetype = GetEntityStorageInterface().GetArchetype(Entity.Index);
+	check(CurrentArchetype);
+	CurrentArchetype->SetFragmentsData(Entity, FragmentInstanceList);
 }
 
 void FMassEntityManager::BatchSetEntityFragmentsValues(const FMassArchetypeEntityCollection& SparseEntities, TArrayView<const FInstancedStruct> FragmentInstanceList)
@@ -1369,8 +1460,9 @@ void* FMassEntityManager::InternalGetFragmentDataChecked(FMassEntityHandle Entit
 	// note that FragmentType is guaranteed to be of valid type - it's either statically checked by the template versions
 	// or `checkf`ed by the non-template one
 	CheckIfEntityIsActive(Entity);
-	const FEntityData& EntityData = Entities[Entity.Index];
-	return EntityData.CurrentArchetype->GetFragmentDataForEntityChecked(FragmentType, Entity.Index);
+	const FMassArchetypeData* CurrentArchetype = GetEntityStorageInterface().GetArchetype(Entity.Index);
+	check(CurrentArchetype);
+	return CurrentArchetype->GetFragmentDataForEntityChecked(FragmentType, Entity.Index);
 }
 
 void* FMassEntityManager::InternalGetFragmentDataPtr(FMassEntityHandle Entity, const UScriptStruct* FragmentType) const
@@ -1378,8 +1470,9 @@ void* FMassEntityManager::InternalGetFragmentDataPtr(FMassEntityHandle Entity, c
 	// note that FragmentType is guaranteed to be of valid type - it's either statically checked by the template versions
 	// or `checkf`ed by the non-template one
 	CheckIfEntityIsActive(Entity);
-	const FEntityData& EntityData = Entities[Entity.Index];
-	return EntityData.CurrentArchetype->GetFragmentDataForEntity(FragmentType, Entity.Index);
+	const FMassArchetypeData* CurrentArchetype = GetEntityStorageInterface().GetArchetype(Entity.Index);
+	check(CurrentArchetype);
+	return CurrentArchetype->GetFragmentDataForEntity(FragmentType, Entity.Index);
 }
 
 const FConstSharedStruct* FMassEntityManager::InternalGetConstSharedFragmentPtr(FMassEntityHandle Entity, const UScriptStruct* ConstSharedFragmentType) const
@@ -1387,8 +1480,9 @@ const FConstSharedStruct* FMassEntityManager::InternalGetConstSharedFragmentPtr(
 	// note that ConstSharedFragmentType is guaranteed to be of valid type - it's either statically checked by the template versions
 	// or `checkf`ed by the non-template one
 	CheckIfEntityIsActive(Entity);
-	const FEntityData& EntityData = Entities[Entity.Index];
-	const FConstSharedStruct* SharedFragment = EntityData.CurrentArchetype->GetSharedFragmentValues(Entity).GetConstSharedFragments().FindByPredicate(FStructTypeEqualOperator(ConstSharedFragmentType));
+	const FMassArchetypeData* CurrentArchetype = GetEntityStorageInterface().GetArchetype(Entity.Index);
+	check(CurrentArchetype);
+	const FConstSharedStruct* SharedFragment = CurrentArchetype->GetSharedFragmentValues(Entity).GetConstSharedFragments().FindByPredicate(FStructTypeEqualOperator(ConstSharedFragmentType));
 	return SharedFragment;
 }
 
@@ -1397,20 +1491,22 @@ const FSharedStruct* FMassEntityManager::InternalGetSharedFragmentPtr(FMassEntit
 	// note that SharedFragmentType is guaranteed to be of valid type - it's either statically checked by the template versions
 	// or `checkf`ed by the non-template one
 	CheckIfEntityIsActive(Entity);
-	const FEntityData& EntityData = Entities[Entity.Index];
-	const FSharedStruct* SharedFragment = EntityData.CurrentArchetype->GetSharedFragmentValues(Entity).GetSharedFragments().FindByPredicate(FStructTypeEqualOperator(SharedFragmentType));
+	const FMassArchetypeData* CurrentArchetype = GetEntityStorageInterface().GetArchetype(Entity.Index);
+	check(CurrentArchetype);
+	const FSharedStruct* SharedFragment = CurrentArchetype->GetSharedFragmentValues(Entity).GetSharedFragments().FindByPredicate(FStructTypeEqualOperator(SharedFragmentType));
 	return SharedFragment;
 }
 
 bool FMassEntityManager::IsEntityValid(FMassEntityHandle Entity) const
 {
-	return (Entity.Index > 0) && Entities.IsValidIndex(Entity.Index) && (Entities[Entity.Index].SerialNumber == Entity.SerialNumber);
+	return (Entity.Index != UE::Mass::Private::InvalidEntityIndex) && GetEntityStorageInterface().IsValidIndex(Entity.Index) && (GetEntityStorageInterface().GetSerialNumber(Entity.Index) == Entity.SerialNumber);
 }
 
 bool FMassEntityManager::IsEntityBuilt(FMassEntityHandle Entity) const
 {
 	CheckIfEntityIsValid(Entity);
-	return Entities[Entity.Index].CurrentArchetype.IsValid();
+	const UE::Mass::IEntityStorageInterface::EEntityState CurrentState = GetEntityStorageInterface().GetEntityState(Entity.Index);
+	return CurrentState == UE::Mass::IEntityStorageInterface::EEntityState::Created;
 }
 
 void FMassEntityManager::CheckIfEntityIsValid(FMassEntityHandle Entity) const
@@ -1772,27 +1868,41 @@ int32 FMassEntityManager::DebugGetArchetypeEntitiesCountPerChunk(const FMassArch
 	return Archetype.IsValid() ? FMassArchetypeHelper::ArchetypeDataFromHandleChecked(Archetype).GetNumEntitiesPerChunk() : 0;
 }
 
+int32 FMassEntityManager::DebugGetEntityCount() const
+{
+	return GetEntityStorageInterface().Num() - NumReservedEntities - GetEntityStorageInterface().ComputeFreeSize();
+}
+
+int32 FMassEntityManager::DebugGetArchetypesCount() const
+{
+	return AllArchetypes.Num();
+}
+
 void FMassEntityManager::DebugRemoveAllEntities()
 {
-	for (int EntityIndex = NumReservedEntities; EntityIndex < Entities.Num(); ++EntityIndex)
+	for (int EntityIndex = NumReservedEntities, EndIndex = GetEntityStorageInterface().Num(); EntityIndex < EndIndex; ++EntityIndex)
 	{
-		FEntityData& EntityData = Entities[EntityIndex];
-		if (EntityData.IsValid() == false)
+		if (GetEntityStorageInterface().IsValid(EntityIndex) == false)
 		{
 			// already dead
 			continue;
 		}
-		const TSharedPtr<FMassArchetypeData>& Archetype = EntityData.CurrentArchetype;
+		FMassArchetypeData* Archetype = GetEntityStorageInterface().GetArchetype(EntityIndex);
+		check(Archetype);
 		FMassEntityHandle Entity;
 		Entity.Index = EntityIndex;
-		Entity.SerialNumber = EntityData.SerialNumber;
+		Entity.SerialNumber = GetEntityStorageInterface().GetSerialNumber(EntityIndex);
 		Archetype->RemoveEntity(Entity);
 
-		EntityData.Reset();
-		EntityFreeIndexList.Add(EntityIndex);
+		GetEntityStorageInterface().ForceReleaseOne(Entity);
 	}
 }
 
+void FMassEntityManager::DebugForceArchetypeDataVersionBump()
+{
+	++ArchetypeDataVersion;
+}
+
 void FMassEntityManager::DebugGetArchetypeStrings(const FMassArchetypeHandle& Archetype, TArray<FName>& OutFragmentNames, TArray<FName>& OutTagNames)
 {
 	if (Archetype.IsValid() == false)
@@ -1812,6 +1922,21 @@ void FMassEntityManager::DebugGetArchetypeStrings(const FMassArchetypeHandle& Ar
 	ArchetypeRef.GetTagBitSet().DebugGetIndividualNames(OutTagNames);
 }
 
+FMassEntityHandle FMassEntityManager::DebugGetEntityIndexHandle(const int32 EntityIndex) const
+{
+	return GetEntityStorageInterface().IsValidIndex(EntityIndex) ? FMassEntityHandle(EntityIndex, GetEntityStorageInterface().GetSerialNumber(EntityIndex)) : FMassEntityHandle();
+}
+
+const FString& FMassEntityManager::DebugGetName() const
+{
+	return DebugName;
+}
+
+FMassRequirementAccessDetector& FMassEntityManager::GetRequirementAccessDetector()
+{
+	return RequirementAccessDetector;
+}
+
 #endif // WITH_MASSENTITY_DEBUG
 
 
@@ -1867,4 +1992,4 @@ void FMassEntityManager::BatchBuildEntities(const FMassArchetypeEntityCollection
 	FMassArchetypeCreationParams Params;
 	Params.DebugName = ArchetypeDebugName;
 	BatchBuildEntities(EncodedEntitiesWithPayload, MoveTemp(Composition), SharedFragmentValues, Params);
-}
+}
\ No newline at end of file
diff --git a/Engine/Plugins/Runtime/MassEntity/Source/MassEntity/Private/MassEntityManagerStorage.cpp b/Engine/Plugins/Runtime/MassEntity/Source/MassEntity/Private/MassEntityManagerStorage.cpp
new file mode 100644
--- /dev/null
+++ b/Engine/Plugins/Runtime/MassEntity/Source/MassEntity/Private/MassEntityManagerStorage.cpp
@@ -0,0 +1,477 @@
+﻿// Copyright Epic Games, Inc. All Rights Reserved.
+#pragma once
+
+#include "MassEntityManagerStorage.h"
+
+#include "MassEntityTypes.h"
+#include "Templates/SharedPointer.h"
+
+namespace UE::Mass
+{
+	//////////////////////////////////////////////////////////////////////
+	// FSingleThreadedEntityStorage
+	
+	FMassArchetypeData* FSingleThreadedEntityStorage::GetArchetype(int32 Index)
+	{
+		return Entities[Index].CurrentArchetype.Get();
+	}
+
+	const FMassArchetypeData* FSingleThreadedEntityStorage::GetArchetype(int32 Index) const
+	{
+		return Entities[Index].CurrentArchetype.Get();
+	}
+
+	TSharedPtr<FMassArchetypeData>& FSingleThreadedEntityStorage::GetArchetypeAsShared(int32 Index)
+	{
+		return Entities[Index].CurrentArchetype;
+	}
+
+	const TSharedPtr<FMassArchetypeData>& FSingleThreadedEntityStorage::GetArchetypeAsShared(int32 Index) const
+	{
+		return Entities[Index].CurrentArchetype;
+	}
+
+	void FSingleThreadedEntityStorage::SetArchetypeFromShared(int32 Index, TSharedPtr<FMassArchetypeData>& Archetype)
+	{
+		Entities[Index].CurrentArchetype = Archetype;
+	}
+
+	void FSingleThreadedEntityStorage::SetArchetypeFromShared(int32 Index, const TSharedPtr<FMassArchetypeData>& Archetype)
+	{
+		Entities[Index].CurrentArchetype = Archetype;
+	}
+
+	IEntityStorageInterface::EEntityState FSingleThreadedEntityStorage::GetEntityState(int32 Index) const
+	{
+		const FMassArchetypeData* CurrentArchetype = Entities[Index].CurrentArchetype.Get();
+		const uint32 CurrentSerialNumber = Entities[Index].SerialNumber;
+
+		if (CurrentSerialNumber != 0)
+		{
+			if (CurrentArchetype != nullptr)
+			{
+				return EEntityState::Created;
+			}
+			else // (CurrentArchetype == nullptr)
+			{
+				return EEntityState::Reserved;
+			}
+		}
+
+		return EEntityState::Free;	
+	}
+
+	int32 FSingleThreadedEntityStorage::GetSerialNumber(int32 Index) const
+	{
+		return Entities[Index].SerialNumber;
+	}
+
+	bool FSingleThreadedEntityStorage::IsValidIndex(int32 Index) const
+	{
+		return Entities.IsValidIndex(Index);
+	}
+
+	SIZE_T FSingleThreadedEntityStorage::GetAllocatedSize() const
+	{
+		return Entities.GetAllocatedSize() + EntityFreeIndexList.GetAllocatedSize();
+	}
+
+	bool FSingleThreadedEntityStorage::IsValid(int32 Index) const
+	{
+		return Entities[Index].IsValid();
+	}
+
+	FMassEntityHandle FSingleThreadedEntityStorage::AcquireOne()
+	{
+		const int32 SerialNumber = SerialNumberGenerator.fetch_add(1);
+		const int32 Index = (EntityFreeIndexList.Num() > 0) ? EntityFreeIndexList.Pop(EAllowShrinking::No) : Entities.Add();
+		Entities[Index].SerialNumber = SerialNumber;
+
+		FMassEntityHandle Handle;
+		Handle.SerialNumber = SerialNumber;
+		Handle.Index = Index;
+		return Handle;
+	}
+
+	int32 FSingleThreadedEntityStorage::Release(TArrayView<FMassEntityHandle> Handles)
+	{
+		int DeallocateCount = 0;
+
+		EntityFreeIndexList.Reserve(EntityFreeIndexList.Num() + Handles.Num());
+
+		for (FMassEntityHandle& Handle : Handles)
+		{
+			FEntityData& EntityData = Entities[Handle.Index];
+			if (EntityData.SerialNumber == Handle.SerialNumber)
+			{
+				EntityData.Reset();
+				EntityFreeIndexList.Add(Handle.Index);
+				++DeallocateCount;
+			}
+		}
+	
+		return DeallocateCount;
+	}
+
+	int32 FSingleThreadedEntityStorage::ReleaseOne(FMassEntityHandle Handle)
+	{
+		return Release(MakeArrayView(&Handle, 1));
+	}
+
+	int32 FSingleThreadedEntityStorage::ForceRelease(TArrayView<FMassEntityHandle> Handles)
+	{
+		EntityFreeIndexList.Reserve(EntityFreeIndexList.Num() + Handles.Num());
+		for (FMassEntityHandle& Handle : Handles)
+		{
+			FEntityData& EntityData = Entities[Handle.Index];
+			EntityData.Reset();
+			EntityFreeIndexList.Add(Handle.Index);
+		}
+		return Handles.Num();
+	}
+
+	int32 FSingleThreadedEntityStorage::ForceReleaseOne(FMassEntityHandle Handle)
+	{
+		return ForceRelease(MakeArrayView(&Handle, 1));
+	}
+
+	int32 FSingleThreadedEntityStorage::Num() const
+	{
+		return Entities.Num();
+	}
+
+	int32 FSingleThreadedEntityStorage::ComputeFreeSize() const
+	{
+		return EntityFreeIndexList.Num();
+	}
+
+	FSingleThreadedEntityStorage::FEntityData::~FEntityData() = default;
+
+	void FSingleThreadedEntityStorage::FEntityData::Reset()
+	{
+		CurrentArchetype.Reset();
+		SerialNumber = 0;
+	}
+
+	bool FSingleThreadedEntityStorage::FEntityData::IsValid() const
+	{
+		return SerialNumber != 0 && CurrentArchetype.IsValid();
+	}
+
+	//////////////////////////////////////////////////////////////////////
+	// FConcurrentEntityStorage
+
+	void FConcurrentEntityStorage::Initialize(const FMassEntityManager_InitParams_Concurrent& InInitializationParams)
+	{
+		// Compute number of pages required
+		check(FMath::IsPowerOfTwo(InInitializationParams.MaxEntitiesPerPage));
+		check(FMath::IsPowerOfTwo(InInitializationParams.MaxEntityCount));
+		MaxEntitiesPerPage = InInitializationParams.MaxEntitiesPerPage;
+		MaximumEntityCountShift = FMath::FloorLog2(InInitializationParams.MaxEntityCount);
+		checkf(MaximumEntityCountShift < 32, TEXT("Invalid maximum entity count, cannot exceed 31 bits"));
+
+		const uint64 PagePointerCount = InInitializationParams.MaxEntityCount / InInitializationParams.MaxEntitiesPerPage;
+
+		const uint64 EntityPageSize = sizeof(void*) * PagePointerCount;
+		EntityPages = static_cast<FEntityData**>(FMemory::Malloc(EntityPageSize, alignof(FEntityData**)));
+		FMemory::Memzero(EntityPages, EntityPageSize);	
+	}
+
+	FConcurrentEntityStorage::~FConcurrentEntityStorage()
+	{
+		if (EntityPages != nullptr)
+		{
+			for (uint32 Index = 0; Index < PageCount; ++Index)
+			{
+				FMemory::Free(EntityPages[Index]);
+				EntityPages[Index] = nullptr;
+			}
+			FMemory::Free(EntityPages);
+			EntityPages = nullptr;
+		}
+	}
+
+	FMassArchetypeData* FConcurrentEntityStorage::GetArchetype(int32 Index)
+	{
+		return LookupEntity(Index).CurrentArchetype.Get();
+	}
+
+	const FMassArchetypeData* FConcurrentEntityStorage::GetArchetype(int32 Index) const
+	{
+		return LookupEntity(Index).CurrentArchetype.Get();
+	}
+
+	TSharedPtr<FMassArchetypeData>& FConcurrentEntityStorage::GetArchetypeAsShared(int32 Index)
+	{
+		return LookupEntity(Index).CurrentArchetype;
+	}
+
+	const TSharedPtr<FMassArchetypeData>& FConcurrentEntityStorage::GetArchetypeAsShared(int32 Index) const
+	{
+		return LookupEntity(Index).CurrentArchetype;
+	}
+
+	void FConcurrentEntityStorage::SetArchetypeFromShared(int32 Index, TSharedPtr<FMassArchetypeData>& Archetype)
+	{
+		LookupEntity(Index).CurrentArchetype = Archetype;
+	}
+
+	void FConcurrentEntityStorage::SetArchetypeFromShared(int32 Index, const TSharedPtr<FMassArchetypeData>& Archetype)
+	{
+		LookupEntity(Index).CurrentArchetype = Archetype;
+	}
+
+	IEntityStorageInterface::EEntityState FConcurrentEntityStorage::GetEntityState(int32 Index) const
+	{
+		//
+		// || Archetype || IsAllocated || Result    |
+		//  |  nullptr   |      0       |  Free     |
+		//  |  nullptr   |      1       |  Reserved |
+		//  | !nullptr   |      1       |  Created  |
+		//
+	
+		const FEntityData& EntityData = LookupEntity(Index);
+		if (EntityData.CurrentArchetype != nullptr)
+		{
+			return EEntityState::Created;
+		}
+		else // EntityData.CurrentArchetype == nullptr
+		{
+			if (EntityData.IsAllocated == 1)
+			{
+				return EEntityState::Reserved;
+			}
+			return EEntityState::Free;
+		}
+	}
+
+	int32 FConcurrentEntityStorage::GetSerialNumber(int32 Index) const
+	{
+		return LookupEntity(Index).GenerationId;
+	}
+
+	bool FConcurrentEntityStorage::IsValidIndex(int32 Index) const
+	{
+		// Page Index is which page in the array of pages we need to access
+		const uint32 PageIndex = static_cast<uint32>(Index) >> MaximumEntityCountShift;
+		return PageIndex < PageCount;
+	}
+
+	SIZE_T FConcurrentEntityStorage::GetAllocatedSize() const
+	{
+		const SIZE_T EntityFreeListSizeBytes = EntityFreeIndexList.GetAllocatedSize();
+
+		// Allocated size to pages
+		const SIZE_T PageSizeBytes = ComputePageSize();
+		const SIZE_T PageAllocatedSizeBytes = PageCount * PageSizeBytes;
+
+		// Size of page pointer array
+		const uint32 MaxEntities = 1 << MaximumEntityCountShift;
+		const uint32 MagPageCount = (MaxEntities / MaxEntitiesPerPage);
+		const SIZE_T PagePointerArraySizeBytes = MagPageCount * sizeof(FEntityData**);
+	
+		return PageAllocatedSizeBytes + PagePointerArraySizeBytes + EntityFreeListSizeBytes;
+	}
+
+	bool FConcurrentEntityStorage::IsValid(int32 Index) const
+	{
+		return LookupEntity(Index).CurrentArchetype != nullptr;
+	}
+
+	FMassEntityHandle FConcurrentEntityStorage::AcquireOne()
+	{
+		int32 EntityIndex;
+		{
+			UE::TUniqueLock FreeListLock(FreeListMutex);
+		
+			if (UNLIKELY(EntityFreeIndexList.IsEmpty()))
+			{
+				check(FreeListMutex.IsLocked());
+				UE::TUniqueLock PageAllocateLock(PageAllocateMutex);
+
+				// Allocate new page
+				const uint32 NewPageIndex = PageCount;
+				checkf((NewPageIndex + 1) * MaxEntitiesPerPage < (1llu << MaximumEntityCountShift), TEXT("Exahusted number of entities"));
+
+				const uint64 PageSize = ComputePageSize();
+				FEntityData* Page = static_cast<FEntityData*>(FMemory::Malloc(PageSize, alignof(FEntityData)));
+			
+				for (int32 Index = 0, End = MaxEntitiesPerPage; Index < End; ++Index)
+				{
+					new (Page + Index) FEntityData();
+				}
+
+				EntityPages[PageCount] = Page;
+				++PageCount;
+
+				const int32 NewEntityIndexStart = NewPageIndex * MaxEntitiesPerPage;
+				const int32 NewEntityIndexEnd = (NewPageIndex + 1) * MaxEntitiesPerPage;
+
+				EntityFreeIndexList.Reserve(MaxEntitiesPerPage);
+
+				// Somewhat tricksy thing here to be aware of
+				// MassEntityManager expects the very first allocated entity to be at index 0
+				// However, EntityFreeIndexList.Pop() will return the last one added to the list
+				// Therefore, populate the free list backwards
+				for (int32 NewEntityIndex = NewEntityIndexEnd - 1; NewEntityIndex >= NewEntityIndexStart; --NewEntityIndex)
+				{
+					// Setup the free list
+					EntityFreeIndexList.Push(NewEntityIndex);
+				}
+			}
+
+			EntityIndex = EntityFreeIndexList.Pop(EAllowShrinking::No);
+		}
+
+		FEntityData& EntityData = LookupEntity(EntityIndex);
+		// NOTE: Technically should not be necessary, however FEntityHandle::IsValid() makes the assumption
+		// that SerialNum == 0 means an invalid Entity.  FMassArchetypeEntityCollection uses this assumption
+		// and will fail IsValid() checks otherwise.
+		++EntityData.GenerationId;
+		EntityData.IsAllocated = 1;
+		int32 SerialNumber = EntityData.GetSerialNumber();
+
+		EntityCount.fetch_add(1llu);
+	
+		FMassEntityHandle Handle;
+		Handle.SerialNumber = SerialNumber;
+		Handle.Index = EntityIndex;
+		return Handle;
+	}
+
+	int32 FConcurrentEntityStorage::Release(TArrayView<FMassEntityHandle> Handles)
+	{
+		int32 DeallocateCount = 0;
+	
+		int32 BeginHandlesIndexToFree = 0;
+		int32 AllocatedRunLength = 0;
+
+		// Helper to add a range of handles to the EntityFreeIndexList
+		auto FreeRunOfHandles = [this, &BeginHandlesIndexToFree, &AllocatedRunLength, Handles]()
+		{
+			if (AllocatedRunLength > 0) // Cheaper than taking the lock for each in case of runs of unallocated handles
+			{
+				UE::TUniqueLock FreeListLock(FreeListMutex);
+				EntityFreeIndexList.Reserve(EntityFreeIndexList.Num() + AllocatedRunLength);
+				for (int32 IndexToFree = BeginHandlesIndexToFree; IndexToFree < BeginHandlesIndexToFree + AllocatedRunLength; ++IndexToFree)
+				{
+					FMassEntityHandle& HandleToFree = Handles[IndexToFree];
+					EntityFreeIndexList.Add(HandleToFree.Index);
+				}
+			}
+			BeginHandlesIndexToFree += (AllocatedRunLength + 1); // +1 to skip to next iteration
+			AllocatedRunLength = 0;
+		};
+	
+		for (int32 Index = 0, End = Handles.Num(); Index < End; ++Index)
+		{
+			FMassEntityHandle& Handle = Handles[Index];
+			FEntityData& EntityData = LookupEntity(Handle.Index);
+			if (EntityData.GetSerialNumber() == Handle.SerialNumber)
+			{
+				++AllocatedRunLength;
+			
+				++EntityData.GenerationId;
+				EntityData.IsAllocated = 0;
+				EntityData.CurrentArchetype.Reset();
+			
+				++DeallocateCount;
+			}
+			else
+			{
+				// Skip, this one isn't allocated
+				// Return the last run to the free list
+				// Ideally this code never runs but we cannot control what is passed into the Release() function
+				FreeRunOfHandles();
+			}
+		}
+
+		// Free any remaining handles
+		FreeRunOfHandles();
+
+		EntityCount.fetch_sub(DeallocateCount);
+	
+		return DeallocateCount;
+	}
+
+	int32 FConcurrentEntityStorage::ReleaseOne(FMassEntityHandle Handle)
+	{
+		return Release(MakeArrayView(&Handle, 1));
+	}
+
+	int32 FConcurrentEntityStorage::ForceRelease(TArrayView<FMassEntityHandle> Handles)
+	{
+		// ForceRelease assumes the caller knows all handles are allocated
+		// no need to have complexity of tracking "runs" of handles 
+		for (FMassEntityHandle& Handle : Handles)
+		{
+			FEntityData& EntityData = LookupEntity(Handle.Index);
+
+			++EntityData.GenerationId;
+			EntityData.IsAllocated = 0;
+			EntityData.CurrentArchetype.Reset();
+		}
+
+		{
+			UE::TUniqueLock FreeListLock(FreeListMutex);
+			EntityFreeIndexList.Reserve(EntityFreeIndexList.Num() + Handles.Num());
+			for (FMassEntityHandle& Handle : Handles)
+			{
+				EntityFreeIndexList.Add(Handle.Index);
+			}
+		}
+
+		EntityCount.fetch_sub(Handles.Num());
+	
+		return Handles.Num();
+	}
+
+	int32 FConcurrentEntityStorage::ForceReleaseOne(FMassEntityHandle Handle)
+	{
+		return ForceRelease(MakeArrayView(&Handle, 1));
+	}
+
+	int32 FConcurrentEntityStorage::Num() const
+	{
+		return MaxEntitiesPerPage * PageCount;;
+	}
+
+	int32 FConcurrentEntityStorage::ComputeFreeSize() const
+	{
+		return EntityFreeIndexList.Num();
+	}
+
+	FConcurrentEntityStorage::FEntityData::~FEntityData() = default;
+
+	int32 FConcurrentEntityStorage::FEntityData::GetSerialNumber() const
+	{
+		return static_cast<int32>(GenerationId);
+	}
+
+	FConcurrentEntityStorage::FEntityData& FConcurrentEntityStorage::LookupEntity(int32 Index)
+	{
+		check(Index >= 0);
+		// PageIndex is which Page in the array of pages we need to access
+		const uint32 PageIndex = static_cast<uint32>(Index) >> FMath::FloorLog2(MaxEntitiesPerPage);
+
+		// Convert the entity index into the index with respect to the page
+		const uint32 EntityOffset = (PageIndex * MaxEntitiesPerPage);
+		check(Index >= static_cast<int32>(EntityOffset)); // Check against negative values
+		const uint32 InternalPageIndex = static_cast<uint32>(Index) - EntityOffset;
+
+		// Pointer to start of page
+		FEntityData* PageStart = EntityPages[PageIndex];
+		FEntityData& EntityData = PageStart[InternalPageIndex];
+		return EntityData;
+	}
+
+	const FConcurrentEntityStorage::FEntityData& FConcurrentEntityStorage::LookupEntity(int32 Index) const
+	{
+		return const_cast<FConcurrentEntityStorage*>(this)->LookupEntity(Index);
+	}
+
+	uint64 FConcurrentEntityStorage::ComputePageSize() const
+	{
+		return sizeof(FEntityData) * MaxEntitiesPerPage;
+	}
+}
diff --git a/Engine/Plugins/Runtime/MassEntity/Source/MassEntity/Public/MassEntityManager.h b/Engine/Plugins/Runtime/MassEntity/Source/MassEntity/Public/MassEntityManager.h
--- a/Engine/Plugins/Runtime/MassEntity/Source/MassEntity/Public/MassEntityManager.h
+++ b/Engine/Plugins/Runtime/MassEntity/Source/MassEntity/Public/MassEntityManager.h
@@ -12,6 +12,7 @@
 #include "Containers/MpscQueue.h"
 #include "MassRequirementAccessDetector.h"
 #include "Templates/FunctionFwd.h"
+#include "MassEntityManagerStorage.h"
 
 
 struct FMassEntityQuery;
@@ -23,6 +24,18 @@ class FOutputDevice;
 struct FMassDebugger;
 enum class EMassFragmentAccess : uint8;
 enum class EForkProcessRole : uint8;
+namespace UE::Mass::Private
+{
+	struct FEntityStorageInitializer;
+}
+
+// Only Editor builds enable Mass to use the concurrent entity reservation implementation
+// Non-editor builds will devirtualize to a single threaded implementation only
+#if UE_EDITOR
+#define MASS_CONCURRENT_RESERVE 1
+#else
+#define MASS_CONCURRENT_RESERVE 0
+#endif
 
 /** 
  * The type responsible for hosting Entities managing Archetypes.
@@ -48,23 +61,6 @@ struct MASSENTITY_API FMassEntityManager : public TSharedFromThis<FMassEntityMan
 private:
 	// Index 0 is reserved so we can treat that index as an invalid entity handle
 	constexpr static int32 NumReservedEntities = 1;
-
-	struct FEntityData
-	{
-		TSharedPtr<FMassArchetypeData> CurrentArchetype;
-		int32 SerialNumber = 0;
-
-		void Reset()
-		{
-			CurrentArchetype.Reset();
-			SerialNumber = 0;
-		}
-
-		bool IsValid() const
-		{
-			return SerialNumber != 0 && CurrentArchetype.IsValid();
-		}
-	};
 	
 public:
 	struct FScopedProcessing
@@ -97,7 +93,9 @@ public:
 	// End of FGCObject interface
 	void GetResourceSizeEx(FResourceSizeEx& CumulativeResourceSize);
 
+	// Default to use single threaded implementation
 	void Initialize();
+	void Initialize(const FMassEntityManagerStorageInitParams& InitializationParams);
 	void PostInitialize();
 	void Deinitialize();
 
@@ -556,15 +554,20 @@ public:
 	void DebugGetArchetypeFragmentTypes(const FMassArchetypeHandle& Archetype, TArray<const UScriptStruct*>& InOutFragmentList) const;
 	int32 DebugGetArchetypeEntitiesCount(const FMassArchetypeHandle& Archetype) const;
 	int32 DebugGetArchetypeEntitiesCountPerChunk(const FMassArchetypeHandle& Archetype) const;
-	int32 DebugGetEntityCount() const { return Entities.Num() - NumReservedEntities - EntityFreeIndexList.Num(); }
-	int32 DebugGetArchetypesCount() const { return AllArchetypes.Num(); }
+	int32 DebugGetEntityCount() const;
+	int32 DebugGetArchetypesCount() const;
 	void DebugRemoveAllEntities();
-	void DebugForceArchetypeDataVersionBump() { ++ArchetypeDataVersion; }
+	void DebugForceArchetypeDataVersionBump();
 	void DebugGetArchetypeStrings(const FMassArchetypeHandle& Archetype, TArray<FName>& OutFragmentNames, TArray<FName>& OutTagNames);
-	FMassEntityHandle DebugGetEntityIndexHandle(const int32 EntityIndex) const { return Entities.IsValidIndex(EntityIndex) ? FMassEntityHandle(EntityIndex, Entities[EntityIndex].SerialNumber) : FMassEntityHandle(); }
-	const FString& DebugGetName() const { return DebugName; }
+	FMassEntityHandle DebugGetEntityIndexHandle(const int32 EntityIndex) const;
+	const FString& DebugGetName() const;
+
+	FMassRequirementAccessDetector& GetRequirementAccessDetector();
 
-	FMassRequirementAccessDetector& GetRequirementAccessDetector() { return RequirementAccessDetector; }
+	// For use by the friend MassDebugger
+	UE::Mass::IEntityStorageInterface& DebugGetEntityStorageInterface();
+	// For use by the friend MassDebugger
+	const UE::Mass::IEntityStorageInterface& DebugGetEntityStorageInterface() const;
 #endif // WITH_MASSENTITY_DEBUG
 
 protected:
@@ -611,9 +614,23 @@ private:
 
 	TSharedRef<FEntityCreationContext> InternalBatchCreateReservedEntities(const FMassArchetypeHandle& ArchetypeHandle,
 		const FMassArchetypeSharedFragmentValues& SharedFragmentValues, TConstArrayView<FMassEntityHandle> ReservedEntities);
+	
+#if MASS_CONCURRENT_RESERVE
+	UE::Mass::IEntityStorageInterface& GetEntityStorageInterface();
+	const UE::Mass::IEntityStorageInterface& GetEntityStorageInterface() const;
+#else
+	UE::Mass::FSingleThreadedEntityStorage& GetEntityStorageInterface();
+	const UE::Mass::FSingleThreadedEntityStorage& GetEntityStorageInterface() const;
+#endif
+	
 private:
-	TChunkedArray<FEntityData> Entities;
-	TArray<int32> EntityFreeIndexList;
+
+	friend struct UE::Mass::Private::FEntityStorageInitializer;
+	using FEntityStorageContainerType = TVariant<
+		FEmptyVariantState,
+		UE::Mass::FSingleThreadedEntityStorage,
+		UE::Mass::FConcurrentEntityStorage>;
+	FEntityStorageContainerType EntityStorage;
 
 	std::atomic<bool> bCommandBufferFlushingInProgress = false;
 	/**
@@ -625,8 +642,7 @@ private:
 	 */
 	uint8 OpenedCommandBufferIndex = 0;
 	TStaticArray<TSharedPtr<FMassCommandBuffer>, 2> DeferredCommandBuffers;
-
-	std::atomic<int32> SerialNumberGenerator = 0;
+	
 	std::atomic<int32> ProcessingScopeCount = 0;
 
 	// the "version" number increased every time an archetype gets added
diff --git a/Engine/Plugins/Runtime/MassEntity/Source/MassEntity/Public/MassEntityManagerStorage.h b/Engine/Plugins/Runtime/MassEntity/Source/MassEntity/Public/MassEntityManagerStorage.h
new file mode 100644
--- /dev/null
+++ b/Engine/Plugins/Runtime/MassEntity/Source/MassEntity/Public/MassEntityManagerStorage.h
@@ -0,0 +1,215 @@
+// Copyright Epic Games, Inc. All Rights Reserved.
+#pragma once
+
+#include "Async/Mutex.h"
+#include "Containers/ChunkedArray.h"
+#include "Misc/TVariant.h"
+#include "Templates/SharedPointer.h"
+
+struct FMassArchetypeData;
+struct FMassEntityHandle;
+
+// Initialization parameters to configure MassEntityManager to reserve entities only single threaded
+// Supported in all build configurations
+struct MASSENTITY_API FMassEntityManager_InitParams_SingleThreaded {};
+
+// Initialization parameters to configure MassEntityManager to concurrently reserve entities
+// Only supported in editor builds.
+//
+// Expected static memory requirement for array of Page pointers can be computed:
+// MaxPages = MaxEntityCount / MaxEntitiesPerPage
+// MemorySize = MaxPages * sizeof(Page**)
+//
+// For default values, expectation is 128kB
+struct MASSENTITY_API FMassEntityManager_InitParams_Concurrent
+{
+	// Maximum supported entities by the MassEntityManager
+	// Must be multiple of 2
+	uint32 MaxEntityCount = 1 << 30; // 1 billion
+
+	// Number of entities per chunk
+	// Must be multiple of 2
+	uint32 MaxEntitiesPerPage = 1 << 16; // 65536
+};
+
+using FMassEntityManagerStorageInitParams = TVariant<FMassEntityManager_InitParams_SingleThreaded, FMassEntityManager_InitParams_Concurrent>;
+
+namespace UE::Mass
+{
+	/**
+	 * Interface that abstracts the storage system for Mass Entities in the EntityManager
+	 * This may be temporary until the concurrent mechanism has been vetted for performance
+	 */
+	class IEntityStorageInterface
+	{
+	public:
+		enum class EEntityState
+		{
+			// Entity index refers to an entity that is free to be reserved or created
+			Free,
+			// Entity index refers to a reserved entity
+			Reserved,
+			// Entity index refers to an entity assigned to an archetype
+			Created
+		};
+		virtual ~IEntityStorageInterface() = default;
+
+		virtual FMassArchetypeData* GetArchetype(int32 Index) = 0;
+		virtual const FMassArchetypeData* GetArchetype(int32 Index) const = 0;
+		virtual TSharedPtr<FMassArchetypeData>& GetArchetypeAsShared(int32 Index) = 0;
+		virtual const TSharedPtr<FMassArchetypeData>& GetArchetypeAsShared(int32 Index) const = 0;
+	
+		virtual void SetArchetypeFromShared(int32 Index, TSharedPtr<FMassArchetypeData>& Archetype) = 0;
+		virtual void SetArchetypeFromShared(int32 Index, const TSharedPtr<FMassArchetypeData>& Archetype) = 0;
+
+		// Returns true if the given entity at index is currently reserved
+		// False if free or assigned an archetype
+		virtual EEntityState GetEntityState(int32 Index) const = 0;
+	
+		virtual int32 GetSerialNumber(int32 Index) const = 0;
+
+		// Checks if index can be used to access entity data
+		virtual bool IsValidIndex(int32 Index) const = 0;
+		virtual SIZE_T GetAllocatedSize() const = 0;
+
+		// Checks if entity at Index is built
+		virtual bool IsValid(int32 Index) const = 0;
+
+		virtual FMassEntityHandle AcquireOne() = 0;
+		
+		virtual int32 Release(TArrayView<FMassEntityHandle> Handles) = 0;
+		virtual int32 ReleaseOne(FMassEntityHandle Handles) = 0;
+		
+		// Bypasses Serial Number Check
+		// Only use if caller has ensured serial number matches or for debug purposes
+		virtual int32 ForceRelease(TArrayView<FMassEntityHandle> Handles) = 0;
+		virtual int32 ForceReleaseOne(FMassEntityHandle Handle) = 0;
+
+		// Returns the number of entities that are not free
+		// For debug purposes only.  In multi-threaded environments, the result is going to be out of date
+		virtual int32 Num() const = 0;
+
+		// Returns the number of entities that are free
+		// For debug purposes only.  In multi-threaded environments, the result is going to be out of date
+		virtual int32 ComputeFreeSize() const = 0;
+	};
+
+	
+	//////////////////////////////////////////////////////////////////////
+	// FSingleThreadedEntityStorage
+
+	/**
+	 * This storage backend should be used when the user of MassEntityManager can guarantee
+	 * that all entity management will be done on a single thread.
+	 */
+	class FSingleThreadedEntityStorage final : public IEntityStorageInterface
+	{
+	public:
+		virtual FMassArchetypeData* GetArchetype(int32 Index) override;
+		virtual const FMassArchetypeData* GetArchetype(int32 Index) const override;
+		virtual TSharedPtr<FMassArchetypeData>& GetArchetypeAsShared(int32 Index) override;
+		virtual const TSharedPtr<FMassArchetypeData>& GetArchetypeAsShared(int32 Index) const override;
+		virtual void SetArchetypeFromShared(int32 Index, TSharedPtr<FMassArchetypeData>&) override;
+		virtual void SetArchetypeFromShared(int32 Index, const TSharedPtr<FMassArchetypeData>&) override;
+		virtual EEntityState GetEntityState(int32 Index) const override;
+		virtual int32 GetSerialNumber(int32 Index) const override;
+		virtual bool IsValidIndex(int32 Index) const override;
+		virtual SIZE_T GetAllocatedSize() const override;
+		virtual bool IsValid(int32 Index) const override;
+		virtual FMassEntityHandle AcquireOne() override;
+		virtual int32 Release(TArrayView<FMassEntityHandle> Handles) override;
+		virtual int32 ReleaseOne(FMassEntityHandle Handle) override;
+		virtual int32 ForceRelease(TArrayView<FMassEntityHandle> Handles) override;
+		virtual int32 ForceReleaseOne(FMassEntityHandle Handle) override;
+		virtual int32 Num() const override;
+		virtual int32 ComputeFreeSize() const override;
+	private:
+	
+		struct FEntityData
+		{
+			TSharedPtr<FMassArchetypeData> CurrentArchetype;
+			int32 SerialNumber = 0;
+			
+			~FEntityData();
+			void Reset();
+			bool IsValid() const;
+		};
+
+		std::atomic<int32> SerialNumberGenerator = 0;
+		TChunkedArray<FEntityData> Entities;
+		TArray<int32> EntityFreeIndexList;
+	};
+
+	//////////////////////////////////////////////////////////////////////
+	// FConcurrentEntityStorage
+
+	/**
+	 * This storage backend allows for entities to be concurrently reserved.  Reserved entities can also
+	 * be concurrently freed.
+	 * Creation of entities (ie. assignment of an archetype and addition of data into chunks) cannot be done
+	 * concurrently with this implementation.
+	 */
+	class FConcurrentEntityStorage final : public IEntityStorageInterface
+	{
+	public:
+
+		void Initialize(const FMassEntityManager_InitParams_Concurrent& InInitializationParams);
+
+		virtual ~FConcurrentEntityStorage() override;
+
+		virtual FMassArchetypeData* GetArchetype(int32 Index) override;
+		virtual const FMassArchetypeData* GetArchetype(int32 Index) const override;
+		virtual TSharedPtr<FMassArchetypeData>& GetArchetypeAsShared(int32 Index) override;
+		virtual const TSharedPtr<FMassArchetypeData>& GetArchetypeAsShared(int32 Index) const override;
+		virtual void SetArchetypeFromShared(int32 Index, TSharedPtr<FMassArchetypeData>& Archetype) override;
+		virtual void SetArchetypeFromShared(int32 Index, const TSharedPtr<FMassArchetypeData>& Archetype) override;
+		virtual EEntityState GetEntityState(int32 Index) const override;
+		virtual int32 GetSerialNumber(int32 Index) const override;
+		virtual bool IsValidIndex(int32 Index) const override;
+		virtual SIZE_T GetAllocatedSize() const override;
+		virtual bool IsValid(int32 Index) const override;
+		virtual FMassEntityHandle AcquireOne() override;
+		virtual int32 Release(TArrayView<FMassEntityHandle> Handles) override;
+		virtual int32 ReleaseOne(FMassEntityHandle Handle) override;
+		virtual int32 ForceRelease(TArrayView<FMassEntityHandle> Handles) override;
+		virtual int32 ForceReleaseOne(FMassEntityHandle Handle) override;
+		virtual int32 Num() const override;
+		virtual int32 ComputeFreeSize() const override;
+	private:	
+	
+		struct FEntityData
+		{
+			static constexpr int MaxGenerationBits = 30;
+		
+			TSharedPtr<FMassArchetypeData> CurrentArchetype;
+			// Generation ID or version of the entity in this slot
+			uint32 GenerationId  : MaxGenerationBits = 0;
+			// 1 if the entity is NOT free
+			uint32 IsAllocated : 1 = 0;
+
+			~FEntityData();
+			// Converts EntityData state into a SerialNumber for public usage
+			int32 GetSerialNumber() const;
+		};
+	
+		FEntityData& LookupEntity(int32 Index);
+		const FEntityData& LookupEntity(int32 Index) const;
+
+		// Returns size of a page in bytes
+		uint64 ComputePageSize() const;
+
+		// Number of allocated Entities
+		std::atomic<uint64> EntityCount = 0;
+		uint32 MaximumEntityCountShift = 0;
+		uint32 MaxEntitiesPerPage = 0;
+		uint32 PageCount = 0;
+		// ALWAYS acquire FreeListMutex before this one
+		UE::FMutex PageAllocateMutex;
+		// Pointer to array of pages
+		FEntityData** EntityPages = nullptr;
+	
+		TArray<int32> EntityFreeIndexList;
+		UE::FMutex FreeListMutex;
+	};
+	
+} // namespace UE::Mass
\ No newline at end of file
diff --git a/Engine/Plugins/Runtime/MassEntity/Source/MassEntityEditor/Private/MassEntityEditorSubsystem.cpp b/Engine/Plugins/Runtime/MassEntity/Source/MassEntityEditor/Private/MassEntityEditorSubsystem.cpp
--- a/Engine/Plugins/Runtime/MassEntity/Source/MassEntityEditor/Private/MassEntityEditorSubsystem.cpp
+++ b/Engine/Plugins/Runtime/MassEntity/Source/MassEntityEditor/Private/MassEntityEditorSubsystem.cpp
@@ -113,7 +113,14 @@ void UMassEntityEditorSubsystem::InitializeMassInternals()
 	// in the header file.
 	EditorTickFunction = new UE::Mass::FMassEditorTickFunction(*this);
 
-	EntityManager->Initialize();
+	FMassEntityManagerStorageInitParams InitParams;
+	InitParams.Emplace<FMassEntityManager_InitParams_Concurrent>(
+		FMassEntityManager_InitParams_Concurrent
+		{
+			.MaxEntityCount = 1llu << 28,
+			.MaxEntitiesPerPage = 1 << 16
+		});
+	EntityManager->Initialize(InitParams);
 
 	// set up ProcessingPhasesConfig
 	TConstArrayView<FMassProcessingPhaseConfig> MainPhasesConfig = GET_MASS_CONFIG_VALUE(GetProcessingPhasesConfig());
