Dedicated type for Mass const shared fragments

diff --git a/Engine/Plugins/AI/MassAI/Source/MassAIBehavior/Private/MassStateTreeExecutionContext.cpp b/Engine/Plugins/AI/MassAI/Source/MassAIBehavior/Private/MassStateTreeExecutionContext.cpp
--- a/Engine/Plugins/AI/MassAI/Source/MassAIBehavior/Private/MassStateTreeExecutionContext.cpp
+++ b/Engine/Plugins/AI/MassAI/Source/MassAIBehavior/Private/MassStateTreeExecutionContext.cpp
@@ -44,12 +44,33 @@ bool CollectExternalData(const FStateTreeExecutionContext& Context, const UState
 			{
 				if (DataDesc.Requirement == EStateTreeExternalDataRequirement::Required)
 				{
+					UE_LOG(LogMass, Error, TEXT("Missing Fragment: %s"), *GetNameSafe(ScriptStruct));
+
 					// Note: Not breaking here, so that we can validate all missing ones in one go.
 					bFoundAll = false;
 				}
 			}
 		}
 		else if (DataDesc.Struct->IsChildOf(FMassSharedFragment::StaticStruct()))
+		{
+			const UScriptStruct* ScriptStruct = Cast<const UScriptStruct>(DataDesc.Struct);
+			FStructView Fragment = EntityView.GetSharedFragmentDataStruct(ScriptStruct);
+			if (Fragment.IsValid())
+			{
+				OutDataViews[Index] = FStateTreeDataView(Fragment.GetScriptStruct(), Fragment.GetMemory());
+			}
+			else
+			{
+				if (DataDesc.Requirement == EStateTreeExternalDataRequirement::Required)
+				{
+					UE_LOG(LogMass, Error, TEXT("Missing Shared Fragment: %s"), *GetNameSafe(ScriptStruct));
+
+					// Note: Not breaking here, so that we can validate all missing ones in one go.
+					bFoundAll = false;
+				}
+			}
+		}
+		else if (DataDesc.Struct->IsChildOf(FMassConstSharedFragment::StaticStruct()))
 		{
 			const UScriptStruct* ScriptStruct = Cast<const UScriptStruct>(DataDesc.Struct);
 			FConstStructView Fragment = EntityView.GetConstSharedFragmentDataStruct(ScriptStruct);
@@ -61,6 +82,8 @@ bool CollectExternalData(const FStateTreeExecutionContext& Context, const UState
 			{
 				if (DataDesc.Requirement == EStateTreeExternalDataRequirement::Required)
 				{
+					UE_LOG(LogMass, Error, TEXT("Missing Const Shared Fragment: %s"), *GetNameSafe(ScriptStruct));
+
 					// Note: Not breaking here, so that we can validate all missing ones in one go.
 					bFoundAll = false;
 				}
@@ -78,6 +101,8 @@ bool CollectExternalData(const FStateTreeExecutionContext& Context, const UState
 			{
 				if (DataDesc.Requirement == EStateTreeExternalDataRequirement::Required)
 				{
+					UE_LOG(LogMass, Error, TEXT("Missing Subsystem: %s"), *GetNameSafe(SubClass));
+
 					// Note: Not breaking here, so that we can validate all missing ones in one go.
 					bFoundAll = false;
 				}
diff --git a/Engine/Plugins/AI/MassAI/Source/MassAIBehavior/Private/MassStateTreeSchema.cpp b/Engine/Plugins/AI/MassAI/Source/MassAIBehavior/Private/MassStateTreeSchema.cpp
--- a/Engine/Plugins/AI/MassAI/Source/MassAIBehavior/Private/MassStateTreeSchema.cpp
+++ b/Engine/Plugins/AI/MassAI/Source/MassAIBehavior/Private/MassStateTreeSchema.cpp
@@ -20,5 +20,6 @@ bool UMassStateTreeSchema::IsExternalItemAllowed(const UStruct& InStruct) const
 	// Allow only WorldSubsystems and fragments as external data.
 	return InStruct.IsChildOf(UWorldSubsystem::StaticClass())
 			|| InStruct.IsChildOf(FMassFragment::StaticStruct())
-			|| InStruct.IsChildOf(FMassSharedFragment::StaticStruct());
+			|| InStruct.IsChildOf(FMassSharedFragment::StaticStruct())
+			|| InStruct.IsChildOf(FMassConstSharedFragment::StaticStruct());
 }
diff --git a/Engine/Plugins/AI/MassAI/Source/MassAIBehavior/Public/MassStateTreeFragments.h b/Engine/Plugins/AI/MassAI/Source/MassAIBehavior/Public/MassStateTreeFragments.h
--- a/Engine/Plugins/AI/MassAI/Source/MassAIBehavior/Public/MassStateTreeFragments.h
+++ b/Engine/Plugins/AI/MassAI/Source/MassAIBehavior/Public/MassStateTreeFragments.h
@@ -23,7 +23,7 @@ struct MASSAIBEHAVIOR_API FMassStateTreeInstanceFragment : public FMassFragment
 
 
 USTRUCT()
-struct MASSAIBEHAVIOR_API FMassStateTreeSharedFragment : public FMassSharedFragment
+struct MASSAIBEHAVIOR_API FMassStateTreeSharedFragment : public FMassConstSharedFragment
 {
 	GENERATED_BODY()
 
diff --git a/Engine/Plugins/AI/MassAI/Source/MassNavigation/Public/Avoidance/MassAvoidanceFragments.h b/Engine/Plugins/AI/MassAI/Source/MassNavigation/Public/Avoidance/MassAvoidanceFragments.h
--- a/Engine/Plugins/AI/MassAI/Source/MassNavigation/Public/Avoidance/MassAvoidanceFragments.h
+++ b/Engine/Plugins/AI/MassAI/Source/MassNavigation/Public/Avoidance/MassAvoidanceFragments.h
@@ -7,7 +7,7 @@
 #include "MassAvoidanceFragments.generated.h"
 
 USTRUCT()
-struct MASSNAVIGATION_API FMassMovingAvoidanceParameters : public FMassSharedFragment
+struct MASSNAVIGATION_API FMassMovingAvoidanceParameters : public FMassConstSharedFragment
 {
 	GENERATED_BODY()
 
@@ -94,7 +94,7 @@ struct MASSNAVIGATION_API FMassMovingAvoidanceParameters : public FMassSharedFra
 };
 
 USTRUCT()
-struct MASSNAVIGATION_API FMassStandingAvoidanceParameters : public FMassSharedFragment
+struct MASSNAVIGATION_API FMassStandingAvoidanceParameters : public FMassConstSharedFragment
 {
 	GENERATED_BODY()
 
diff --git a/Engine/Plugins/AI/MassAI/Source/MassNavigation/Public/SmoothOrientation/MassSmoothOrientationFragments.h b/Engine/Plugins/AI/MassAI/Source/MassNavigation/Public/SmoothOrientation/MassSmoothOrientationFragments.h
--- a/Engine/Plugins/AI/MassAI/Source/MassNavigation/Public/SmoothOrientation/MassSmoothOrientationFragments.h
+++ b/Engine/Plugins/AI/MassAI/Source/MassNavigation/Public/SmoothOrientation/MassSmoothOrientationFragments.h
@@ -27,7 +27,7 @@ struct MASSNAVIGATION_API FMassSmoothOrientationWeights : public FMassSharedFrag
 };
 
 USTRUCT()
-struct MASSNAVIGATION_API FMassSmoothOrientationParameters : public FMassSharedFragment
+struct MASSNAVIGATION_API FMassSmoothOrientationParameters : public FMassConstSharedFragment
 {
 	GENERATED_BODY()
 
diff --git a/Engine/Plugins/AI/MassAI/Source/MassNavigation/Public/Steering/MassSteeringFragments.h b/Engine/Plugins/AI/MassAI/Source/MassNavigation/Public/Steering/MassSteeringFragments.h
--- a/Engine/Plugins/AI/MassAI/Source/MassNavigation/Public/Steering/MassSteeringFragments.h
+++ b/Engine/Plugins/AI/MassAI/Source/MassNavigation/Public/Steering/MassSteeringFragments.h
@@ -48,7 +48,7 @@ struct MASSNAVIGATION_API FMassStandingSteeringFragment : public FMassFragment
 
 /** Steering related movement parameters. */
 USTRUCT()
-struct MASSNAVIGATION_API FMassMovingSteeringParameters : public FMassSharedFragment
+struct MASSNAVIGATION_API FMassMovingSteeringParameters : public FMassConstSharedFragment
 {
 	GENERATED_BODY()
 
@@ -62,7 +62,7 @@ struct MASSNAVIGATION_API FMassMovingSteeringParameters : public FMassSharedFrag
 };
 
 USTRUCT()
-struct FMassStandingSteeringParameters : public FMassSharedFragment
+struct FMassStandingSteeringParameters : public FMassConstSharedFragment
 {
 	GENERATED_BODY()
 
diff --git a/Engine/Plugins/AI/MassAI/Source/MassZoneGraphNavigation/Public/MassZoneGraphNavigationFragments.h b/Engine/Plugins/AI/MassAI/Source/MassZoneGraphNavigation/Public/MassZoneGraphNavigationFragments.h
--- a/Engine/Plugins/AI/MassAI/Source/MassZoneGraphNavigation/Public/MassZoneGraphNavigationFragments.h
+++ b/Engine/Plugins/AI/MassAI/Source/MassZoneGraphNavigation/Public/MassZoneGraphNavigationFragments.h
@@ -11,7 +11,7 @@
 
 
 USTRUCT()
-struct MASSZONEGRAPHNAVIGATION_API FMassZoneGraphNavigationParameters : public FMassSharedFragment
+struct MASSZONEGRAPHNAVIGATION_API FMassZoneGraphNavigationParameters : public FMassConstSharedFragment
 {
 	GENERATED_BODY()
 
diff --git a/Engine/Plugins/Experimental/TypedElementsDataStorage/Source/TypedElementsDataStorage/Private/TypedElementDatabaseCommandBuffer.cpp b/Engine/Plugins/Experimental/TypedElementsDataStorage/Source/TypedElementsDataStorage/Private/TypedElementDatabaseCommandBuffer.cpp
--- a/Engine/Plugins/Experimental/TypedElementsDataStorage/Source/TypedElementsDataStorage/Private/TypedElementDatabaseCommandBuffer.cpp
+++ b/Engine/Plugins/Experimental/TypedElementsDataStorage/Source/TypedElementsDataStorage/Private/TypedElementDatabaseCommandBuffer.cpp
@@ -236,7 +236,7 @@ void FTypedElementDatabaseCommandBuffer::Execute_AddColumnsCommand(
 	FMassTagBitSet TagsToAdd)
 {
 	FMassArchetypeCompositionDescriptor AddComposition(
-		MoveTemp(FragmentsToAdd), MoveTemp(TagsToAdd), FMassChunkFragmentBitSet(), FMassSharedFragmentBitSet());
+		MoveTemp(FragmentsToAdd), MoveTemp(TagsToAdd), FMassChunkFragmentBitSet(), FMassSharedFragmentBitSet(), FMassConstSharedFragmentBitSet());
 	MassEntityManager.AddCompositionToEntity_GetDelta(FMassEntityHandle::FromNumber(Row), AddComposition);
 }
 
@@ -266,7 +266,7 @@ void FTypedElementDatabaseCommandBuffer::Execute_RemoveColumnsCommand(
 	FMassTagBitSet TagsToRemove)
 {
 	FMassArchetypeCompositionDescriptor RemoveComposition(
-		MoveTemp(FragmentsToRemove), MoveTemp(TagsToRemove), FMassChunkFragmentBitSet(), FMassSharedFragmentBitSet());
+		MoveTemp(FragmentsToRemove), MoveTemp(TagsToRemove), FMassChunkFragmentBitSet(), FMassSharedFragmentBitSet(), FMassConstSharedFragmentBitSet());
 	MassEntityManager.RemoveCompositionFromEntity(FMassEntityHandle::FromNumber(Row), RemoveComposition);
 }
 
diff --git a/Engine/Plugins/Runtime/MassEntity/Source/MassEntity/Private/MassArchetypeData.cpp b/Engine/Plugins/Runtime/MassEntity/Source/MassEntity/Private/MassArchetypeData.cpp
--- a/Engine/Plugins/Runtime/MassEntity/Source/MassEntity/Private/MassArchetypeData.cpp
+++ b/Engine/Plugins/Runtime/MassEntity/Source/MassEntity/Private/MassArchetypeData.cpp
@@ -85,6 +85,7 @@ void FMassArchetypeData::Initialize(const FMassArchetypeCompositionDescriptor& I
 
 	// Share fragments
 	CompositionDescriptor.SharedFragments = InCompositionDescriptor.SharedFragments;
+	CompositionDescriptor.ConstSharedFragments = InCompositionDescriptor.ConstSharedFragments;
 
 	EntityListOffsetWithinChunk = 0;
 }
@@ -987,8 +988,8 @@ void FMassArchetypeData::BatchAddEntities(TConstArrayView<FMassEntityHandle> Ent
 {
 	SCOPE_CYCLE_COUNTER(STAT_Mass_ArchetypeBatchAdd);
 
-	testableCheckfReturn(SharedFragmentValues.HasExactFragmentTypesMatch(GetCompositionDescriptor().SharedFragments)
-		, return, TEXT("%hs parameter SharedFragmentValues doesn't match archetype's composition"), __FUNCTION__);
+	testableCheckfReturn(SharedFragmentValues.HasExactSharedFragmentTypesMatch(GetCompositionDescriptor().SharedFragments), return, TEXT("%hs parameter SharedFragmentValues doesn't match archetype's composition"), __FUNCTION__);
+	testableCheckfReturn(SharedFragmentValues.HasExactConstSharedFragmentTypesMatch(GetCompositionDescriptor().ConstSharedFragments), return, TEXT("%hs parameter ConstSharedFragmentValues doesn't match archetype's composition"), __FUNCTION__);
 
 	FMassArchetypeEntityCollection::FArchetypeEntityRange ResultSubchunk;
 	ResultSubchunk.ChunkIndex = 0;
@@ -1014,25 +1015,31 @@ void FMassArchetypeData::BatchAddEntities(TConstArrayView<FMassEntityHandle> Ent
 void FMassArchetypeData::BatchMoveEntitiesToAnotherArchetype(const FMassArchetypeEntityCollection& EntityCollection
 	, FMassArchetypeData& NewArchetype, TArray<FMassEntityHandle>& OutEntitiesBeingMoved
 	, TArray<FMassArchetypeEntityCollection::FArchetypeEntityRange>* OutNewRanges, const FMassArchetypeSharedFragmentValues* SharedFragmentValuesToAdd
-	, const FMassSharedFragmentBitSet* SharedFragmentToRemoveBitSet)
+	, const FMassSharedFragmentBitSet* SharedFragmentToRemoveBitSet, const FMassConstSharedFragmentBitSet* ConstSharedFragmentToRemoveBitSet)
 {
 	check(&NewArchetype != this);
 
-	if (SharedFragmentValuesToAdd || SharedFragmentToRemoveBitSet)
+	// verify the new archetype's shared fragment composition matches current archetype's composition modified as requested
+	if (SharedFragmentValuesToAdd)
 	{
-		// verify the new archetype's shared fragment composition matches current archetype's composition modified as requested
-		FMassSharedFragmentBitSet NewSharedFragmentsBitset = GetSharedFragmentBitSet();
+		bool bIsValidArchetype = true;
 		if (SharedFragmentToRemoveBitSet)
 		{
+			FMassSharedFragmentBitSet NewSharedFragmentsBitset = GetSharedFragmentBitSet();
 			NewSharedFragmentsBitset -= *SharedFragmentToRemoveBitSet;
+			NewSharedFragmentsBitset += SharedFragmentValuesToAdd->GetSharedFragmentBitSet();
+			bIsValidArchetype = NewArchetype.GetCompositionDescriptor().SharedFragments == NewSharedFragmentsBitset;
 		}
-		if (SharedFragmentValuesToAdd)
+
+		if (bIsValidArchetype && ConstSharedFragmentToRemoveBitSet)
 		{
-			NewSharedFragmentsBitset += SharedFragmentValuesToAdd->GetSharedFragmentBitSet();
+			FMassConstSharedFragmentBitSet NewConstSharedFragmentsBitset = GetConstSharedFragmentBitSet();
+			NewConstSharedFragmentsBitset -= *ConstSharedFragmentToRemoveBitSet;
+			NewConstSharedFragmentsBitset += SharedFragmentValuesToAdd->GetConstSharedFragmentBitSet();
+			bIsValidArchetype = NewArchetype.GetCompositionDescriptor().ConstSharedFragments == NewConstSharedFragmentsBitset;
 		}
-		
-		testableCheckfReturn(NewArchetype.GetCompositionDescriptor().SharedFragments == NewSharedFragmentsBitset
-			, return, TEXT("%hs parameter SharedFragmentValues doesn't match archetype's composition"), __FUNCTION__);
+
+		testableCheckfReturn(bIsValidArchetype, return, TEXT("%hs parameter SharedFragmentValues doesn't match archetype's composition"), __FUNCTION__);
 	}
 
 	TArray<FMassArchetypeEntityCollection::FArchetypeEntityRange> Subchunks(EntityCollection.GetRanges());
@@ -1079,6 +1086,10 @@ void FMassArchetypeData::BatchMoveEntitiesToAnotherArchetype(const FMassArchetyp
 				{
 					NewSharedValues.Remove(*SharedFragmentToRemoveBitSet);
 				}
+				if (ConstSharedFragmentToRemoveBitSet)
+				{
+					NewSharedValues.Remove(*ConstSharedFragmentToRemoveBitSet);
+				}
 				if (SharedFragmentValuesToAdd)
 				{
 					NewSharedValues.Append(*SharedFragmentValuesToAdd);
diff --git a/Engine/Plugins/Runtime/MassEntity/Source/MassEntity/Private/MassArchetypeData.h b/Engine/Plugins/Runtime/MassEntity/Source/MassEntity/Private/MassArchetypeData.h
--- a/Engine/Plugins/Runtime/MassEntity/Source/MassEntity/Private/MassArchetypeData.h
+++ b/Engine/Plugins/Runtime/MassEntity/Source/MassEntity/Private/MassArchetypeData.h
@@ -196,6 +196,7 @@ public:
 	const FMassTagBitSet& GetTagBitSet() const { return CompositionDescriptor.Tags; }
 	const FMassChunkFragmentBitSet& GetChunkFragmentBitSet() const { return CompositionDescriptor.ChunkFragments; }
 	const FMassSharedFragmentBitSet& GetSharedFragmentBitSet() const { return CompositionDescriptor.SharedFragments; }
+	const FMassConstSharedFragmentBitSet& GetConstSharedFragmentBitSet() const { return CompositionDescriptor.ConstSharedFragments; }
 
 	const FMassArchetypeCompositionDescriptor& GetCompositionDescriptor() const { return CompositionDescriptor; }
 	FORCEINLINE const FMassArchetypeSharedFragmentValues& GetSharedFragmentValues(int32 EntityIndex) const
@@ -372,7 +373,8 @@ public:
 	void BatchMoveEntitiesToAnotherArchetype(const FMassArchetypeEntityCollection& EntityCollection, FMassArchetypeData& NewArchetype
 		, TArray<FMassEntityHandle>& OutEntitiesBeingMoved, TArray<FMassArchetypeEntityCollection::FArchetypeEntityRange>* OutNewChunks = nullptr
 		, const FMassArchetypeSharedFragmentValues* SharedFragmentValuesToAdd = nullptr
-		, const FMassSharedFragmentBitSet* SharedFragmentToRemoveBitSet = nullptr);
+		, const FMassSharedFragmentBitSet* SharedFragmentToRemoveBitSet = nullptr
+		, const FMassConstSharedFragmentBitSet* ConstSharedFragmentToRemoveBitSet = nullptr);
 	void BatchSetFragmentValues(TConstArrayView<FMassArchetypeEntityCollection::FArchetypeEntityRange> EntityCollection, const FMassGenericPayloadViewSlice& Payload);
 
 protected:
diff --git a/Engine/Plugins/Runtime/MassEntity/Source/MassEntity/Private/MassEntityManager.cpp b/Engine/Plugins/Runtime/MassEntity/Source/MassEntity/Private/MassEntityManager.cpp
--- a/Engine/Plugins/Runtime/MassEntity/Source/MassEntity/Private/MassEntityManager.cpp
+++ b/Engine/Plugins/Runtime/MassEntity/Source/MassEntity/Private/MassEntityManager.cpp
@@ -14,7 +14,6 @@
 #include "Misc/CoreDelegates.h"
 #include "Algo/Find.h"
 
-
 const FMassEntityHandle FMassEntityManager::InvalidEntity;
 
 namespace UE::Mass::Private
@@ -357,18 +356,26 @@ FMassArchetypeHandle FMassEntityManager::CreateArchetype(const TSharedPtr<FMassA
 	check(SourceArchetype.IsValid());
 	checkf(AddedFragments.IsEmpty() == false, TEXT("%hs Adding an empty fragment list to an archetype is not supported."), __FUNCTION__);
 
-	const FMassArchetypeCompositionDescriptor Composition(AddedFragments + SourceArchetype->GetFragmentBitSet(), SourceArchetype->GetTagBitSet(), SourceArchetype->GetChunkFragmentBitSet(), SourceArchetype->GetSharedFragmentBitSet());
+	const FMassArchetypeCompositionDescriptor Composition(AddedFragments + SourceArchetype->GetFragmentBitSet()
+		, SourceArchetype->GetTagBitSet()
+		, SourceArchetype->GetChunkFragmentBitSet()
+		, SourceArchetype->GetSharedFragmentBitSet()
+		, SourceArchetype->GetConstSharedFragmentBitSet());
 	return CreateArchetype(Composition, CreationParams);
 }
 
-FMassArchetypeHandle FMassEntityManager::GetOrCreateSuitableArchetype(const FMassArchetypeHandle& ArchetypeHandle, const FMassSharedFragmentBitSet& SharedFragmentBitSet
+FMassArchetypeHandle FMassEntityManager::GetOrCreateSuitableArchetype(const FMassArchetypeHandle& ArchetypeHandle
+	, const FMassSharedFragmentBitSet& SharedFragmentBitSet
+	, const FMassConstSharedFragmentBitSet& ConstSharedFragmentBitSet
 	, const FMassArchetypeCreationParams& CreationParams)
 {
-	FMassArchetypeData& ArchetypeData = FMassArchetypeHelper::ArchetypeDataFromHandleChecked(ArchetypeHandle);
-	if (SharedFragmentBitSet != ArchetypeData.GetSharedFragmentBitSet())
+	const FMassArchetypeData& ArchetypeData = FMassArchetypeHelper::ArchetypeDataFromHandleChecked(ArchetypeHandle);
+	if (SharedFragmentBitSet != ArchetypeData.GetSharedFragmentBitSet()
+		|| ConstSharedFragmentBitSet != ArchetypeData.GetConstSharedFragmentBitSet())
 	{
 		FMassArchetypeCompositionDescriptor NewDescriptor = ArchetypeData.GetCompositionDescriptor();
 		NewDescriptor.SharedFragments = SharedFragmentBitSet;
+		NewDescriptor.ConstSharedFragments = ConstSharedFragmentBitSet;
 		return CreateArchetype(NewDescriptor);
 	}
 	return ArchetypeHandle;
@@ -430,7 +437,11 @@ FMassArchetypeHandle FMassEntityManager::InternalCreateSimilarArchetype(const TS
 {
 	checkSlow(SourceArchetype.IsValid());
 	const FMassArchetypeData& SourceArchetypeRef = *SourceArchetype.Get();
-	FMassArchetypeCompositionDescriptor NewComposition(SourceArchetypeRef.GetFragmentBitSet(), OverrideTags, SourceArchetypeRef.GetChunkFragmentBitSet(), SourceArchetypeRef.GetSharedFragmentBitSet());
+	FMassArchetypeCompositionDescriptor NewComposition(SourceArchetypeRef.GetFragmentBitSet()
+		, OverrideTags
+		, SourceArchetypeRef.GetChunkFragmentBitSet()
+		, SourceArchetypeRef.GetSharedFragmentBitSet()
+		, SourceArchetypeRef.GetConstSharedFragmentBitSet());
 	return InternalCreateSimilarArchetype(SourceArchetypeRef, MoveTemp(NewComposition));
 }
 
@@ -438,7 +449,11 @@ FMassArchetypeHandle FMassEntityManager::InternalCreateSimilarArchetype(const TS
 {
 	checkSlow(SourceArchetype.IsValid());
 	const FMassArchetypeData& SourceArchetypeRef = *SourceArchetype.Get();
-	FMassArchetypeCompositionDescriptor NewComposition(OverrideFragments, SourceArchetypeRef.GetTagBitSet(), SourceArchetypeRef.GetChunkFragmentBitSet(), SourceArchetypeRef.GetSharedFragmentBitSet());
+	FMassArchetypeCompositionDescriptor NewComposition(OverrideFragments
+		, SourceArchetypeRef.GetTagBitSet()
+		, SourceArchetypeRef.GetChunkFragmentBitSet()
+		, SourceArchetypeRef.GetSharedFragmentBitSet()
+		, SourceArchetypeRef.GetConstSharedFragmentBitSet());
 	return InternalCreateSimilarArchetype(SourceArchetypeRef, MoveTemp(NewComposition));
 }
 
@@ -563,7 +578,9 @@ FMassEntityHandle FMassEntityManager::CreateEntity(const FMassArchetypeHandle& A
 	check(ArchetypeHandle.IsValid());
 
 	const FMassEntityHandle Entity = ReserveEntity();
-	InternalBuildEntity(Entity, GetOrCreateSuitableArchetype(ArchetypeHandle, SharedFragmentValues.GetSharedFragmentBitSet()), SharedFragmentValues);
+	InternalBuildEntity(Entity
+		, GetOrCreateSuitableArchetype(ArchetypeHandle, SharedFragmentValues.GetSharedFragmentBitSet(), SharedFragmentValues.GetConstSharedFragmentBitSet())
+		, SharedFragmentValues);
 
 	return Entity;
 }
@@ -574,7 +591,7 @@ FMassEntityHandle FMassEntityManager::CreateEntity(TConstArrayView<FInstancedStr
 	check(FragmentInstanceList.Num() > 0);
 
 	const FMassArchetypeHandle& ArchetypeHandle = CreateArchetype(FMassArchetypeCompositionDescriptor(FragmentInstanceList,
-		FMassTagBitSet(), FMassChunkFragmentBitSet(), FMassSharedFragmentBitSet()), CreationParams);
+		FMassTagBitSet(), FMassChunkFragmentBitSet(), FMassSharedFragmentBitSet(), FMassConstSharedFragmentBitSet()), CreationParams);
 	check(ArchetypeHandle.IsValid());
 
 	const FMassEntityHandle Entity = ReserveEntity();
@@ -617,7 +634,7 @@ void FMassEntityManager::BuildEntity(FMassEntityHandle Entity, TConstArrayView<F
 	checkf(!IsEntityBuilt(Entity), TEXT("Expecting an entity that is not already built"));
 
 	checkf(SharedFragmentValues.IsSorted(), TEXT("Expecting shared fragment values to be previously sorted"));
-	FMassArchetypeCompositionDescriptor Composition(FragmentInstanceList, FMassTagBitSet(), FMassChunkFragmentBitSet(), FMassSharedFragmentBitSet());
+	FMassArchetypeCompositionDescriptor Composition(FragmentInstanceList, FMassTagBitSet(), FMassChunkFragmentBitSet(), FMassSharedFragmentBitSet(), FMassConstSharedFragmentBitSet());
 	for (const FConstSharedStruct& SharedFragment : SharedFragmentValues.GetConstSharedFragments())
 	{
 		Composition.SharedFragments.Add(*SharedFragment.GetScriptStruct());
@@ -654,7 +671,7 @@ void FMassEntityManager::BatchBuildEntities(const FMassArchetypeEntityCollection
 	checkf(IsProcessing() == false, TEXT("Synchronous API function %hs called during mass processing. Use asynchronous API instead."), __FUNCTION__);
 	check(SharedFragmentValues.IsSorted());
 
-	FMassArchetypeCompositionDescriptor Composition(FragmentsAffected, FMassTagBitSet(), FMassChunkFragmentBitSet(), FMassSharedFragmentBitSet());
+	FMassArchetypeCompositionDescriptor Composition(FragmentsAffected, FMassTagBitSet(), FMassChunkFragmentBitSet(), FMassSharedFragmentBitSet(), FMassConstSharedFragmentBitSet());
 	for (const FConstSharedStruct& SharedFragment : SharedFragmentValues.GetConstSharedFragments())
 	{
 		Composition.SharedFragments.Add(*SharedFragment.GetScriptStruct());
@@ -729,7 +746,8 @@ TSharedRef<FMassEntityManager::FEntityCreationContext> FMassEntityManager::Batch
 	checkf(IsProcessing() == false, TEXT("Synchronous API function %hs called during mass processing. Use asynchronous API instead."), __FUNCTION__);
 	checkf(!ReservedEntities.IsEmpty(), TEXT("No reserved entities given to batch create."));
 
-	return InternalBatchCreateReservedEntities(GetOrCreateSuitableArchetype(ArchetypeHandle, SharedFragmentValues.GetSharedFragmentBitSet())
+	return InternalBatchCreateReservedEntities(
+		GetOrCreateSuitableArchetype(ArchetypeHandle, SharedFragmentValues.GetSharedFragmentBitSet(), SharedFragmentValues.GetConstSharedFragmentBitSet())
 		, SharedFragmentValues, ReservedEntities);
 }
 
@@ -744,7 +762,8 @@ TSharedRef<FMassEntityManager::FEntityCreationContext> FMassEntityManager::Batch
 
 	TConstArrayView<FMassEntityHandle> ReservedEntities = BatchReserveEntities(Count, InOutEntities);
 	
-	return InternalBatchCreateReservedEntities(GetOrCreateSuitableArchetype(ArchetypeHandle, SharedFragmentValues.GetSharedFragmentBitSet())
+	return InternalBatchCreateReservedEntities(
+		GetOrCreateSuitableArchetype(ArchetypeHandle, SharedFragmentValues.GetSharedFragmentBitSet(), SharedFragmentValues.GetConstSharedFragmentBitSet())
 		, SharedFragmentValues, ReservedEntities);
 }
 
@@ -1055,7 +1074,11 @@ void FMassEntityManager::RemoveFragmentListFromEntity(FMassEntityHandle Entity,
 	if (OldArchetype->GetFragmentBitSet().HasAny(FragmentsToRemove))
 	{
 		// If all the fragments got removed this will result in fetching of the empty archetype
-		const FMassArchetypeCompositionDescriptor NewComposition(OldArchetype->GetFragmentBitSet() - FragmentsToRemove, OldArchetype->GetTagBitSet(), OldArchetype->GetChunkFragmentBitSet(), OldArchetype->GetSharedFragmentBitSet());
+		const FMassArchetypeCompositionDescriptor NewComposition(OldArchetype->GetFragmentBitSet() - FragmentsToRemove
+			, OldArchetype->GetTagBitSet()
+			, OldArchetype->GetChunkFragmentBitSet()
+			, OldArchetype->GetSharedFragmentBitSet()
+			, OldArchetype->GetConstSharedFragmentBitSet());
 		const FMassArchetypeHandle NewArchetypeHandle = CreateArchetype(NewComposition);
 
 		FMassArchetypeCompositionDescriptor CompositionDelta;
@@ -1171,7 +1194,7 @@ bool FMassEntityManager::AddConstSharedFragmentToEntity(const FMassEntityHandle
 
 	const UScriptStruct* StructType = InConstSharedFragment.GetScriptStruct();
 	CA_ASSUME(StructType);
-	if (CurrentArchetype->GetCompositionDescriptor().SharedFragments.Contains(*StructType))
+	if (CurrentArchetype->GetCompositionDescriptor().ConstSharedFragments.Contains(*StructType))
 	{
 		const FMassArchetypeSharedFragmentValues& SharedFragmentValues = CurrentArchetype->GetSharedFragmentValues(Entity);
 		FConstSharedStruct ExistingConstSharedStruct = SharedFragmentValues.GetConstSharedFragmentStruct(StructType);
@@ -1185,7 +1208,7 @@ bool FMassEntityManager::AddConstSharedFragmentToEntity(const FMassEntityHandle
 	}
 	
 	FMassArchetypeCompositionDescriptor NewComposition(CurrentArchetype->GetCompositionDescriptor());
-	NewComposition.SharedFragments.Add(*StructType);
+	NewComposition.ConstSharedFragments.Add(*StructType);
 	const FMassArchetypeHandle NewArchetypeHandle = CreateArchetype(NewComposition);
 	check(NewArchetypeHandle.IsValid());
 	FMassArchetypeData* NewArchetype = NewArchetypeHandle.DataPtr.Get();
@@ -1397,6 +1420,7 @@ void FMassEntityManager::BatchAddSharedFragmentsForEntities(TConstArrayView<FMas
 
 		FMassArchetypeCompositionDescriptor NewComposition(CurrentArchetype->GetCompositionDescriptor());
 		NewComposition.SharedFragments += AddedFragmentValues.GetSharedFragmentBitSet();
+		NewComposition.ConstSharedFragments += AddedFragmentValues.GetConstSharedFragmentBitSet();
 
 		const FMassArchetypeHandle NewArchetypeHandle = CreateArchetype(NewComposition);
 		check(NewArchetypeHandle.IsValid());
@@ -1716,6 +1740,31 @@ void FMassEntityManager::GetMatchingArchetypes(const FMassFragmentRequirements&
 			continue;
 		}
 
+		if (Archetype.GetConstSharedFragmentBitSet().HasAll(Requirements.GetRequiredAllConstSharedFragments()) == false)
+		{
+			// missing some required fragments, skip.
+#if WITH_MASSENTITY_DEBUG
+			const FMassConstSharedFragmentBitSet UnsatisfiedFragments = Requirements.GetRequiredAllConstSharedFragments() - Archetype.GetConstSharedFragmentBitSet();
+			FStringOutputDevice Description;
+			UnsatisfiedFragments.DebugGetStringDesc(Description);
+			UE_LOG(LogMass, VeryVerbose, TEXT("Archetype did not match due to missing Const Shared Fragments: %s")
+				, *Description);
+#endif // WITH_MASSENTITY_DEBUG
+			continue;
+		}
+
+		if (Archetype.GetConstSharedFragmentBitSet().HasNone(Requirements.GetRequiredNoneConstSharedFragments()) == false)
+		{
+			// has some Fragments required to be absent
+#if WITH_MASSENTITY_DEBUG
+			const FMassConstSharedFragmentBitSet UnwantedFragments = Requirements.GetRequiredNoneConstSharedFragments().GetOverlap(Archetype.GetConstSharedFragmentBitSet());
+			FStringOutputDevice Description;
+			UnwantedFragments.DebugGetStringDesc(Description);
+			UE_LOG(LogMass, VeryVerbose, TEXT("Archetype has Const Shared Fragments required absent: %s")
+				, *Description);
+#endif // WITH_MASSENTITY_DEBUG
+			continue;
+		}
 
 		OutValidArchetypes.Add(ArchetypePtr);
 	}
diff --git a/Engine/Plugins/Runtime/MassEntity/Source/MassEntity/Private/MassEntityTypes.cpp b/Engine/Plugins/Runtime/MassEntity/Source/MassEntity/Private/MassEntityTypes.cpp
--- a/Engine/Plugins/Runtime/MassEntity/Source/MassEntity/Private/MassEntityTypes.cpp
+++ b/Engine/Plugins/Runtime/MassEntity/Source/MassEntity/Private/MassEntityTypes.cpp
@@ -11,6 +11,7 @@ DEFINE_TYPEBITSET(FMassFragmentBitSet);
 DEFINE_TYPEBITSET(FMassTagBitSet);
 DEFINE_TYPEBITSET(FMassChunkFragmentBitSet);
 DEFINE_TYPEBITSET(FMassSharedFragmentBitSet);
+DEFINE_TYPEBITSET(FMassConstSharedFragmentBitSet);
 DEFINE_TYPEBITSET(FMassExternalSubsystemBitSet);
 
 //-----------------------------------------------------------------------------
@@ -28,7 +29,7 @@ FConstSharedStruct FMassArchetypeSharedFragmentValues::AddConstSharedFragment(co
 	}
 
 	check(StructType);
-	SharedFragmentBitSet.Add(*StructType);
+	ConstSharedFragmentBitSet.Add(*StructType);
 	FConstSharedStruct& StructInstance = ConstSharedFragments.Add_GetRef(Fragment);
 	DirtyHashCache();
 	return StructInstance;
@@ -133,7 +134,8 @@ namespace UE::Mass::Private
 
 bool FMassArchetypeSharedFragmentValues::HasSameValues(const FMassArchetypeSharedFragmentValues& Other) const
 {
-	if (SharedFragmentBitSet.IsEquivalent(Other.SharedFragmentBitSet) == false)
+	if (SharedFragmentBitSet.IsEquivalent(Other.SharedFragmentBitSet) == false
+		|| ConstSharedFragmentBitSet.IsEquivalent(Other.ConstSharedFragmentBitSet) == false)
 	{
 		return false;
 	}
@@ -171,7 +173,7 @@ int32 FMassArchetypeSharedFragmentValues::Append(const FMassArchetypeSharedFragm
 	{
 		const UScriptStruct* StructType = SharedStruct.GetScriptStruct();
 		check(StructType);
-		if (SharedFragmentBitSet.Contains(*StructType))
+		if (ConstSharedFragmentBitSet.Contains(*StructType))
 		{
 			const int32 FragmentIndex = ConstSharedFragments.IndexOfByPredicate(FStructTypeEqualOperator(StructType));
 			if (testableEnsureMsgf(FragmentIndex != INDEX_NONE, TEXT("%hs trying to switch non-const fragment to const, type %s")
@@ -187,8 +189,9 @@ int32 FMassArchetypeSharedFragmentValues::Append(const FMassArchetypeSharedFragm
 			++AddedOrModifiedCount;
 		}
 	}
-	
+
 	SharedFragmentBitSet += Other.SharedFragmentBitSet;
+	ConstSharedFragmentBitSet += Other.ConstSharedFragmentBitSet;
 	DirtyHashCache();
 
 	return AddedOrModifiedCount;
@@ -204,20 +207,39 @@ int32 FMassArchetypeSharedFragmentValues::Remove(const FMassSharedFragmentBitSet
 		const UScriptStruct* StructType = CommonFragments.GetTypeAtIndex(*It);
 		check(StructType);
 
-		const int32 RecularFragmentIndex = SharedFragments.IndexOfByPredicate(FStructTypeEqualOperator(StructType));
-		if (RecularFragmentIndex != INDEX_NONE)
+		const int32 RegularFragmentIndex = SharedFragments.IndexOfByPredicate(FStructTypeEqualOperator(StructType));
+		if (RegularFragmentIndex != INDEX_NONE)
 		{
-			SharedFragments[RecularFragmentIndex].Reset();
+			SharedFragments[RegularFragmentIndex].Reset();
 			++RemovedCount;
 		}
-		else
+
+		++It;
+	}
+
+	if (RemovedCount)
+	{
+		SharedFragmentBitSet -= CommonFragments;
+		DirtyHashCache();
+	}
+	return RemovedCount;
+}
+
+int32 FMassArchetypeSharedFragmentValues::Remove(const FMassConstSharedFragmentBitSet& ConstSharedFragmentToRemoveBitSet)
+{
+	int32 RemovedCount = 0;
+	FMassConstSharedFragmentBitSet CommonFragments = (ConstSharedFragmentBitSet & ConstSharedFragmentToRemoveBitSet);
+	FMassConstSharedFragmentBitSet::FIndexIterator It = CommonFragments.GetIndexIterator();
+	while(It)
+	{
+		const UScriptStruct* StructType = CommonFragments.GetTypeAtIndex(*It);
+		check(StructType);
+
+		const int32 RegularFragmentIndex = ConstSharedFragments.IndexOfByPredicate(FStructTypeEqualOperator(StructType));
+		if (RegularFragmentIndex != INDEX_NONE)
 		{
-			const int32 ConstFragmentIndex = ConstSharedFragments.IndexOfByPredicate(FStructTypeEqualOperator(StructType));
-			if (ConstFragmentIndex != INDEX_NONE)
-			{
-				ConstSharedFragments[ConstFragmentIndex].Reset();
-				++RemovedCount;
-			}
+			ConstSharedFragments[RegularFragmentIndex].Reset();
+			++RemovedCount;
 		}
 
 		++It;
@@ -225,7 +247,7 @@ int32 FMassArchetypeSharedFragmentValues::Remove(const FMassSharedFragmentBitSet
 
 	if (RemovedCount)
 	{
-		SharedFragmentBitSet -= CommonFragments;
+		ConstSharedFragmentBitSet -= CommonFragments;
 		DirtyHashCache();
 	}
 	return RemovedCount;
diff --git a/Engine/Plugins/Runtime/MassEntity/Source/MassEntity/Private/MassProcessorDependencySolver.cpp b/Engine/Plugins/Runtime/MassEntity/Source/MassEntity/Private/MassProcessorDependencySolver.cpp
--- a/Engine/Plugins/Runtime/MassEntity/Source/MassEntity/Private/MassProcessorDependencySolver.cpp
+++ b/Engine/Plugins/Runtime/MassEntity/Source/MassEntity/Private/MassProcessorDependencySolver.cpp
@@ -51,6 +51,8 @@ void FMassExecutionRequirements::Append(const FMassExecutionRequirements& Other)
 		SharedFragments[i] += Other.SharedFragments[i];
 		RequiredSubsystems[i] += Other.RequiredSubsystems[i];
 	}
+	ConstSharedFragments.Read += Other.ConstSharedFragments.Read;
+
 	RequiredAllTags += Other.RequiredAllTags;
 	RequiredAnyTags += Other.RequiredAnyTags;
 	RequiredNoneTags += Other.RequiredNoneTags;
@@ -61,7 +63,7 @@ void FMassExecutionRequirements::Append(const FMassExecutionRequirements& Other)
 
 void FMassExecutionRequirements::CountResourcesUsed()
 {
-	ResourcesUsedCount = 0;
+	ResourcesUsedCount = ConstSharedFragments.Read.CountStoredTypes();
 
 	for (int i = 0; i < EMassAccessOperation::MAX; ++i)
 	{
@@ -83,7 +85,7 @@ int32 FMassExecutionRequirements::GetTotalBitsUsedCount()
 bool FMassExecutionRequirements::IsEmpty() const
 {
 	return Fragments.IsEmpty() && ChunkFragments.IsEmpty() 
-		&& SharedFragments.IsEmpty() && RequiredSubsystems.IsEmpty()
+		&& SharedFragments.IsEmpty() && ConstSharedFragments.IsEmpty() && RequiredSubsystems.IsEmpty()
 		&& RequiredAllTags.IsEmpty() && RequiredAnyTags.IsEmpty() && RequiredNoneTags.IsEmpty();
 }
 
@@ -92,7 +94,8 @@ FMassArchetypeCompositionDescriptor FMassExecutionRequirements::AsCompositionDes
 	return FMassArchetypeCompositionDescriptor(Fragments.Read + Fragments.Write
 		, RequiredAllTags + RequiredAnyTags
 		, ChunkFragments.Read + ChunkFragments.Write
-		, SharedFragments.Read + SharedFragments.Write);
+		, SharedFragments.Read + SharedFragments.Write
+		, ConstSharedFragments.Read);
 }
 
 //----------------------------------------------------------------------//
@@ -229,6 +232,7 @@ bool FMassProcessorDependencySolver::FResourceUsage::HasArchetypeConflict(TMassE
 
 bool FMassProcessorDependencySolver::FResourceUsage::CanAccessRequirements(const FMassExecutionRequirements& TestedRequirements, const TArray<FMassArchetypeHandle>& InArchetypes) const
 {
+	// note that on purpose we're not checking ConstSharedFragments - those are always only read, no danger of conflicting access
 	bool bCanAccess = (CanAccess<FMassFragmentBitSet>(Requirements.Fragments, TestedRequirements.Fragments) || !HasArchetypeConflict(FragmentsAccess, InArchetypes))
 		&& (CanAccess<FMassChunkFragmentBitSet>(Requirements.ChunkFragments, TestedRequirements.ChunkFragments) || !HasArchetypeConflict(ChunkFragmentsAccess, InArchetypes))
 		&& (CanAccess<FMassSharedFragmentBitSet>(Requirements.SharedFragments, TestedRequirements.SharedFragments) || !HasArchetypeConflict(SharedFragmentsAccess, InArchetypes))
@@ -243,6 +247,8 @@ void FMassProcessorDependencySolver::FResourceUsage::SubmitNode(const int32 Node
 	HandleElementType<FMassChunkFragmentBitSet>(ChunkFragmentsAccess, InOutNode.Requirements.ChunkFragments, InOutNode, NodeIndex);
 	HandleElementType<FMassSharedFragmentBitSet>(SharedFragmentsAccess, InOutNode.Requirements.SharedFragments, InOutNode, NodeIndex);
 	HandleElementType<FMassExternalSubsystemBitSet>(RequiredSubsystemsAccess, InOutNode.Requirements.RequiredSubsystems, InOutNode, NodeIndex);
+	// note that on purpose we're not pushing ConstSharedFragments - those are always only read, no danger of conflicting access
+	// so there's no point in tracking them
 
 	Requirements.Append(InOutNode.Requirements);
 }
diff --git a/Engine/Plugins/Runtime/MassEntity/Source/MassEntity/Private/MassRequirements.cpp b/Engine/Plugins/Runtime/MassEntity/Source/MassEntity/Private/MassRequirements.cpp
--- a/Engine/Plugins/Runtime/MassEntity/Source/MassEntity/Private/MassRequirements.cpp
+++ b/Engine/Plugins/Runtime/MassEntity/Source/MassEntity/Private/MassRequirements.cpp
@@ -32,6 +32,23 @@ namespace UE::Mass::Private
 			}
 		}
 	}
+
+	template<>
+	void ExportRequirements<FMassConstSharedFragmentBitSet>(TConstArrayView<FMassFragmentRequirementDescription> Requirements
+		, TMassExecutionAccess<FMassConstSharedFragmentBitSet>& Out)
+	{
+		for (const FMassFragmentRequirementDescription& Requirement : Requirements)
+		{
+			if (Requirement.Presence != EMassFragmentPresence::None)
+			{
+				check(Requirement.StructType);
+				if (ensureMsgf(Requirement.AccessMode == EMassFragmentAccess::ReadOnly, TEXT("ReadOnly is the only supported AccessMode for ConstSharedFragments")))
+				{
+					Out.Read.Add(*Requirement.StructType);
+				}
+			}
+		}
+	}
 }
 
 //////////////////////////////////////////////////////////////////////
@@ -121,7 +138,9 @@ bool FMassFragmentRequirements::DoesArchetypeMatchRequirements(const FMassArchet
 		&& ArchetypeComposition.ChunkFragments.HasAll(RequiredAllChunkFragments)
 		&& ArchetypeComposition.ChunkFragments.HasNone(RequiredNoneChunkFragments)
 		&& ArchetypeComposition.SharedFragments.HasAll(RequiredAllSharedFragments)
-		&& ArchetypeComposition.SharedFragments.HasNone(RequiredNoneSharedFragments);
+		&& ArchetypeComposition.SharedFragments.HasNone(RequiredNoneSharedFragments)
+		&& ArchetypeComposition.ConstSharedFragments.HasAll(RequiredAllConstSharedFragments)
+		&& ArchetypeComposition.ConstSharedFragments.HasNone(RequiredNoneConstSharedFragments);
 }
 
 void FMassFragmentRequirements::ExportRequirements(FMassExecutionRequirements& OutRequirements) const
@@ -129,8 +148,8 @@ void FMassFragmentRequirements::ExportRequirements(FMassExecutionRequirements& O
 	using UE::Mass::Private::ExportRequirements;
 	ExportRequirements<FMassFragmentBitSet>(FragmentRequirements, OutRequirements.Fragments);
 	ExportRequirements<FMassChunkFragmentBitSet>(ChunkFragmentRequirements, OutRequirements.ChunkFragments);
-	ExportRequirements<FMassSharedFragmentBitSet>(ConstSharedFragmentRequirements, OutRequirements.SharedFragments);
 	ExportRequirements<FMassSharedFragmentBitSet>(SharedFragmentRequirements, OutRequirements.SharedFragments);
+	ExportRequirements<FMassConstSharedFragmentBitSet>(ConstSharedFragmentRequirements, OutRequirements.ConstSharedFragments);
 
 	OutRequirements.RequiredAllTags = RequiredAllTags;
 	OutRequirements.RequiredAnyTags = RequiredAnyTags;
@@ -156,6 +175,9 @@ void FMassFragmentRequirements::Reset()
 	RequiredAllSharedFragments.Reset();
 	RequiredOptionalSharedFragments.Reset();
 	RequiredNoneSharedFragments.Reset();
+	RequiredAllConstSharedFragments.Reset();
+	RequiredOptionalConstSharedFragments.Reset();
+	RequiredNoneConstSharedFragments.Reset();
 
 	IncrementalChangesCount = 0;
 }
diff --git a/Engine/Plugins/Runtime/MassEntity/Source/MassEntity/Public/MassEntityManager.h b/Engine/Plugins/Runtime/MassEntity/Source/MassEntity/Public/MassEntityManager.h
--- a/Engine/Plugins/Runtime/MassEntity/Source/MassEntity/Public/MassEntityManager.h
+++ b/Engine/Plugins/Runtime/MassEntity/Source/MassEntity/Public/MassEntityManager.h
@@ -143,7 +143,9 @@ public:
 	 * @param SharedFragmentBitSet indicates which shared fragments we want the target archetype to have. If ArchetypeHandle 
 	 *	doesn't have these a new archetype will be created.
 	 */
-	FMassArchetypeHandle GetOrCreateSuitableArchetype(const FMassArchetypeHandle& ArchetypeHandle, const FMassSharedFragmentBitSet& SharedFragmentBitSet
+	FMassArchetypeHandle GetOrCreateSuitableArchetype(const FMassArchetypeHandle& ArchetypeHandle
+		, const FMassSharedFragmentBitSet& SharedFragmentBitSet
+		, const FMassConstSharedFragmentBitSet& ConstSharedFragmentBitSet
 		, const FMassArchetypeCreationParams& CreationParams = FMassArchetypeCreationParams());
 
 	/** Fetches the archetype for a given Entity. If Entity is not valid it will still return a handle, just with an invalid archetype */
@@ -375,7 +377,7 @@ public:
 	template<typename ConstSharedFragmentType>
 	ConstSharedFragmentType* GetConstSharedFragmentDataPtr(FMassEntityHandle Entity) const
 	{
-		static_assert(TIsDerivedFrom<ConstSharedFragmentType, FMassSharedFragment>::IsDerived, "Given struct doesn't represent a valid shared fragment type. Make sure to inherit from FMassSharedFragment or one of its child-types.");
+		static_assert(TIsDerivedFrom<ConstSharedFragmentType, FMassConstSharedFragment>::IsDerived, "Given struct doesn't represent a valid const shared fragment type. Make sure to inherit from FMassConstSharedFragment or one of its child-types.");
 		const FConstSharedStruct* ConstSharedStruct = InternalGetConstSharedFragmentPtr(Entity, ConstSharedFragmentType::StaticStruct());
 		return (ConstSharedFragmentType*)(ConstSharedStruct ? ConstSharedStruct->GetMemory() : nullptr);
 	}
@@ -390,7 +392,7 @@ public:
 
 	FConstStructView GetConstSharedFragmentDataStruct(FMassEntityHandle Entity, const UScriptStruct* ConstSharedFragmentType) const
 	{
-		checkf((ConstSharedFragmentType != nullptr) && ConstSharedFragmentType->IsChildOf(FMassSharedFragment::StaticStruct())
+		checkf((ConstSharedFragmentType != nullptr) && ConstSharedFragmentType->IsChildOf(FMassConstSharedFragment::StaticStruct())
 			, TEXT("GetConstSharedFragmentDataStruct called with an invalid fragment type '%s'"), *GetPathNameSafe(ConstSharedFragmentType));
 		const FConstSharedStruct* ConstSharedStruct = InternalGetConstSharedFragmentPtr(Entity, ConstSharedFragmentType);
 		return ConstSharedStruct
@@ -475,7 +477,7 @@ public:
 	template<typename T>
 	const FConstSharedStruct& GetOrCreateConstSharedFragmentByHash(const uint32 Hash, const T& Fragment)
 	{
-		static_assert(TIsDerivedFrom<T, FMassSharedFragment>::IsDerived, "Given struct doesn't represent a valid shared fragment type. Make sure to inherit from FMassSharedFragment or one of its child-types.");
+		static_assert(TIsDerivedFrom<T, FMassConstSharedFragment>::IsDerived, "Given struct doesn't represent a valid const shared fragment type. Make sure to inherit from FMassConstSharedFragment or one of its child-types.");
 		int32& Index = ConstSharedFragmentsMap.FindOrAddByHash(Hash, Hash, INDEX_NONE);
 		if (Index == INDEX_NONE)
 		{
diff --git a/Engine/Plugins/Runtime/MassEntity/Source/MassEntity/Public/MassEntityTypes.h b/Engine/Plugins/Runtime/MassEntity/Source/MassEntity/Public/MassEntityTypes.h
--- a/Engine/Plugins/Runtime/MassEntity/Source/MassEntity/Public/MassEntityTypes.h
+++ b/Engine/Plugins/Runtime/MassEntity/Source/MassEntity/Public/MassEntityTypes.h
@@ -58,6 +58,15 @@ struct FMassSharedFragment
 	FMassSharedFragment() {}
 };
 
+USTRUCT()
+struct FMassConstSharedFragment
+{
+	GENERATED_BODY()
+
+	FMassConstSharedFragment() {}
+};
+
+
 // A handle to a lightweight entity.  An entity is used in conjunction with the FMassEntityManager
 // for the current world and can contain lightweight fragments.
 USTRUCT()
@@ -135,6 +144,7 @@ DECLARE_STRUCTTYPEBITSET_EXPORTED(MASSENTITY_API, FMassFragmentBitSet, FMassFrag
 DECLARE_STRUCTTYPEBITSET_EXPORTED(MASSENTITY_API, FMassTagBitSet, FMassTag);
 DECLARE_STRUCTTYPEBITSET_EXPORTED(MASSENTITY_API, FMassChunkFragmentBitSet, FMassChunkFragment);
 DECLARE_STRUCTTYPEBITSET_EXPORTED(MASSENTITY_API, FMassSharedFragmentBitSet, FMassSharedFragment);
+DECLARE_STRUCTTYPEBITSET_EXPORTED(MASSENTITY_API, FMassConstSharedFragmentBitSet, FMassConstSharedFragment);
 DECLARE_CLASSTYPEBITSET_EXPORTED(MASSENTITY_API, FMassExternalSubsystemBitSet, USubsystem);
 
 /** The type summarily describing a composition of an entity or an archetype. It contains information on both the
@@ -142,26 +152,44 @@ DECLARE_CLASSTYPEBITSET_EXPORTED(MASSENTITY_API, FMassExternalSubsystemBitSet, U
 struct FMassArchetypeCompositionDescriptor
 {
 	FMassArchetypeCompositionDescriptor() = default;
-	FMassArchetypeCompositionDescriptor(const FMassFragmentBitSet& InFragments, const FMassTagBitSet& InTags, const FMassChunkFragmentBitSet& InChunkFragments, const FMassSharedFragmentBitSet& InSharedFragments)
+	FMassArchetypeCompositionDescriptor(const FMassFragmentBitSet& InFragments,
+		const FMassTagBitSet& InTags,
+		const FMassChunkFragmentBitSet& InChunkFragments,
+		const FMassSharedFragmentBitSet& InSharedFragments,
+		const FMassConstSharedFragmentBitSet& InConstSharedFragments)
 		: Fragments(InFragments)
 		, Tags(InTags)
 		, ChunkFragments(InChunkFragments)
 		, SharedFragments(InSharedFragments)
+		, ConstSharedFragments(InConstSharedFragments)
 	{}
 
-	FMassArchetypeCompositionDescriptor(TConstArrayView<const UScriptStruct*> InFragments, const FMassTagBitSet& InTags, const FMassChunkFragmentBitSet& InChunkFragments, const FMassSharedFragmentBitSet& InSharedFragments)
-		: FMassArchetypeCompositionDescriptor(FMassFragmentBitSet(InFragments), InTags, InChunkFragments, InSharedFragments)
+	FMassArchetypeCompositionDescriptor(TConstArrayView<const UScriptStruct*> InFragments,
+		const FMassTagBitSet& InTags,
+		const FMassChunkFragmentBitSet& InChunkFragments,
+		const FMassSharedFragmentBitSet& InSharedFragments,
+		const FMassConstSharedFragmentBitSet& InConstSharedFragments)
+		: FMassArchetypeCompositionDescriptor(FMassFragmentBitSet(InFragments), InTags, InChunkFragments, InSharedFragments, InConstSharedFragments)
 	{}
 
-	FMassArchetypeCompositionDescriptor(TConstArrayView<FInstancedStruct> InFragmentInstances, const FMassTagBitSet& InTags, const FMassChunkFragmentBitSet& InChunkFragments, const FMassSharedFragmentBitSet& InSharedFragments)
-		: FMassArchetypeCompositionDescriptor(FMassFragmentBitSet(InFragmentInstances), InTags, InChunkFragments, InSharedFragments)
+	FMassArchetypeCompositionDescriptor(TConstArrayView<FInstancedStruct> InFragmentInstances,
+		const FMassTagBitSet& InTags,
+		const FMassChunkFragmentBitSet& InChunkFragments,
+		const FMassSharedFragmentBitSet& InSharedFragments,
+		const FMassConstSharedFragmentBitSet& InConstSharedFragments)
+		: FMassArchetypeCompositionDescriptor(FMassFragmentBitSet(InFragmentInstances), InTags, InChunkFragments, InSharedFragments, InConstSharedFragments)
 	{}
 
-	FMassArchetypeCompositionDescriptor(FMassFragmentBitSet&& InFragments, FMassTagBitSet&& InTags, FMassChunkFragmentBitSet&& InChunkFragments, FMassSharedFragmentBitSet&& InSharedFragments)
+	FMassArchetypeCompositionDescriptor(FMassFragmentBitSet&& InFragments,
+		FMassTagBitSet&& InTags,
+		FMassChunkFragmentBitSet&& InChunkFragments,
+		FMassSharedFragmentBitSet&& InSharedFragments,
+		FMassConstSharedFragmentBitSet&& InConstSharedFragments)
 		: Fragments(MoveTemp(InFragments))
 		, Tags(MoveTemp(InTags))
 		, ChunkFragments(MoveTemp(InChunkFragments))
 		, SharedFragments(MoveTemp(InSharedFragments))
+		, ConstSharedFragments(MoveTemp(InConstSharedFragments))
 	{}
 
 	FMassArchetypeCompositionDescriptor(FMassFragmentBitSet&& InFragments)
@@ -178,6 +206,7 @@ struct FMassArchetypeCompositionDescriptor
 		Tags.Reset();
 		ChunkFragments.Reset();
 		SharedFragments.Reset();
+		ConstSharedFragments.Reset();
 	}
 
 	bool IsEquivalent(const FMassArchetypeCompositionDescriptor& OtherDescriptor) const
@@ -185,7 +214,8 @@ struct FMassArchetypeCompositionDescriptor
 		return Fragments.IsEquivalent(OtherDescriptor.Fragments) &&
 			Tags.IsEquivalent(OtherDescriptor.Tags) &&
 			ChunkFragments.IsEquivalent(OtherDescriptor.ChunkFragments) &&
-			SharedFragments.IsEquivalent(OtherDescriptor.SharedFragments);
+			SharedFragments.IsEquivalent(OtherDescriptor.SharedFragments) &&
+			ConstSharedFragments.IsEquivalent(OtherDescriptor.ConstSharedFragments);
 	}
 
 	bool IsEmpty() const 
@@ -193,7 +223,8 @@ struct FMassArchetypeCompositionDescriptor
 		return Fragments.IsEmpty() &&
 			Tags.IsEmpty() &&
 			ChunkFragments.IsEmpty() &&
-			SharedFragments.IsEmpty();
+			SharedFragments.IsEmpty() &&
+			ConstSharedFragments.IsEmpty();
 	}
 
 	bool HasAll(const FMassArchetypeCompositionDescriptor& OtherDescriptor) const
@@ -201,26 +232,34 @@ struct FMassArchetypeCompositionDescriptor
 		return Fragments.HasAll(OtherDescriptor.Fragments) &&
 			Tags.HasAll(OtherDescriptor.Tags) &&
 			ChunkFragments.HasAll(OtherDescriptor.ChunkFragments) &&
-			SharedFragments.HasAll(OtherDescriptor.SharedFragments);
+			SharedFragments.HasAll(OtherDescriptor.SharedFragments) &&
+			ConstSharedFragments.HasAll(OtherDescriptor.ConstSharedFragments);
 	}
 
-	static uint32 CalculateHash(const FMassFragmentBitSet& InFragments, const FMassTagBitSet& InTags, const FMassChunkFragmentBitSet& InChunkFragments, const FMassSharedFragmentBitSet& InSharedFragmentBitSet)
+	static uint32 CalculateHash(const FMassFragmentBitSet& InFragments, const FMassTagBitSet& InTags
+		, const FMassChunkFragmentBitSet& InChunkFragments, const FMassSharedFragmentBitSet& InSharedFragmentBitSet
+		, const FMassConstSharedFragmentBitSet& InConstSharedFragmentBitSet)
 	{
 		const uint32 FragmentsHash = GetTypeHash(InFragments);
 		const uint32 TagsHash = GetTypeHash(InTags);
 		const uint32 ChunkFragmentsHash = GetTypeHash(InChunkFragments);
 		const uint32 SharedFragmentsHash = GetTypeHash(InSharedFragmentBitSet);
-		return HashCombine(HashCombine(HashCombine(FragmentsHash, TagsHash), ChunkFragmentsHash), SharedFragmentsHash);
+		const uint32 ConstSharedFragmentsHash = GetTypeHash(InConstSharedFragmentBitSet);
+		return HashCombine(HashCombine(HashCombine(HashCombine(FragmentsHash, TagsHash), ChunkFragmentsHash), SharedFragmentsHash), ConstSharedFragmentsHash);
 	}	
 
 	uint32 CalculateHash() const 
 	{
-		return CalculateHash(Fragments, Tags, ChunkFragments, SharedFragments);
+		return CalculateHash(Fragments, Tags, ChunkFragments, SharedFragments, ConstSharedFragments);
 	}
 
 	int32 CountStoredTypes() const
 	{
-		return Fragments.CountStoredTypes() + Tags.CountStoredTypes() + ChunkFragments.CountStoredTypes() + SharedFragments.CountStoredTypes();
+		return Fragments.CountStoredTypes()
+			+ Tags.CountStoredTypes()
+			+ ChunkFragments.CountStoredTypes()
+			+ SharedFragments.CountStoredTypes()
+			+ ConstSharedFragments.CountStoredTypes();
 	}
 
 	void DebugOutputDescription(FOutputDevice& Ar) const
@@ -245,6 +284,28 @@ struct FMassArchetypeCompositionDescriptor
 	FMassTagBitSet Tags;
 	FMassChunkFragmentBitSet ChunkFragments;
 	FMassSharedFragmentBitSet SharedFragments;
+	FMassConstSharedFragmentBitSet ConstSharedFragments;
+
+	UE_DEPRECATED(5.5, "This FMassArchetypeCompositionDescriptor constructor is deprecated. Please explicitly provide FConstSharedFragmentBitSet.")
+	FMassArchetypeCompositionDescriptor(const FMassFragmentBitSet& InFragments, const FMassTagBitSet& InTags, const FMassChunkFragmentBitSet& InChunkFragments, const FMassSharedFragmentBitSet& InSharedFragments)
+		: FMassArchetypeCompositionDescriptor(InFragments, InTags, InChunkFragments, InSharedFragments, FMassConstSharedFragmentBitSet())
+	{}
+
+	UE_DEPRECATED(5.5, "This FMassArchetypeCompositionDescriptor constructor is deprecated. Please explicitly provide FConstSharedFragmentBitSet.")
+	FMassArchetypeCompositionDescriptor(TConstArrayView<const UScriptStruct*> InFragments, const FMassTagBitSet& InTags, const FMassChunkFragmentBitSet& InChunkFragments, const FMassSharedFragmentBitSet& InSharedFragments)
+		: FMassArchetypeCompositionDescriptor(FMassFragmentBitSet(InFragments), InTags, InChunkFragments, InSharedFragments, FMassConstSharedFragmentBitSet())
+	{}
+
+	UE_DEPRECATED(5.5, "This FMassArchetypeCompositionDescriptor constructor is deprecated. Please explicitly provide FConstSharedFragmentBitSet.")
+	FMassArchetypeCompositionDescriptor(TConstArrayView<FInstancedStruct> InFragmentInstances, const FMassTagBitSet& InTags, const FMassChunkFragmentBitSet& InChunkFragments, const FMassSharedFragmentBitSet& InSharedFragments)
+		: FMassArchetypeCompositionDescriptor(FMassFragmentBitSet(InFragmentInstances), InTags, InChunkFragments, InSharedFragments, FMassConstSharedFragmentBitSet())
+	{}
+
+	UE_DEPRECATED(5.5, "This FMassArchetypeCompositionDescriptor constructor is deprecated. Please explicitly provide FConstSharedFragmentBitSet.")
+	FMassArchetypeCompositionDescriptor(FMassFragmentBitSet&& InFragments, FMassTagBitSet&& InTags, FMassChunkFragmentBitSet&& InChunkFragments, FMassSharedFragmentBitSet&& InSharedFragments)
+	{
+		ensureMsgf(false, TEXT("This constructor is defunct. Please update your implementation based on deprecation warning."));
+	}
 };
 
 /** 
@@ -260,16 +321,32 @@ struct MASSENTITY_API FMassArchetypeSharedFragmentValues
 	FMassArchetypeSharedFragmentValues& operator=(const FMassArchetypeSharedFragmentValues& OtherFragmentValues) = default;
 	FMassArchetypeSharedFragmentValues& operator=(FMassArchetypeSharedFragmentValues&& OtherFragmentValues) = default;
 
-	FORCEINLINE bool HasExactFragmentTypesMatch(const FMassSharedFragmentBitSet& InSharedFragmentBitSet) const
+	FORCEINLINE bool HasExactFragmentTypesMatch(const FMassSharedFragmentBitSet& InSharedFragmentBitSet, const FMassConstSharedFragmentBitSet& InConstSharedFragmentBitSet) const
+	{
+		return HasExactSharedFragmentTypesMatch(InSharedFragmentBitSet)
+			&& HasExactConstSharedFragmentTypesMatch(InConstSharedFragmentBitSet);
+	}
+
+	FORCEINLINE bool HasExactSharedFragmentTypesMatch(const FMassSharedFragmentBitSet& InSharedFragmentBitSet) const
 	{
 		return SharedFragmentBitSet == InSharedFragmentBitSet;
 	}
 
-	FORCEINLINE bool HasAllRequiredFragmentTypes(const FMassSharedFragmentBitSet& InSharedFragmentBitSet) const
+	FORCEINLINE bool HasAllRequiredSharedFragmentTypes(const FMassSharedFragmentBitSet& InSharedFragmentBitSet) const
 	{
 		return SharedFragmentBitSet.HasAll(InSharedFragmentBitSet);
 	}
 
+	FORCEINLINE bool HasExactConstSharedFragmentTypesMatch(const FMassConstSharedFragmentBitSet& InConstSharedFragmentBitSet) const
+	{
+		return ConstSharedFragmentBitSet == InConstSharedFragmentBitSet;
+	}
+
+	FORCEINLINE bool HasAllRequiredConstSharedFragmentTypes(const FMassConstSharedFragmentBitSet& InConstSharedFragmentBitSet) const
+	{
+		return ConstSharedFragmentBitSet.HasAll(InConstSharedFragmentBitSet);
+	}
+
 	FORCEINLINE bool IsEquivalent(const FMassArchetypeSharedFragmentValues& OtherSharedFragmentValues) const
 	{
 		return GetTypeHash(*this) == GetTypeHash(OtherSharedFragmentValues);
@@ -285,13 +362,37 @@ struct MASSENTITY_API FMassArchetypeSharedFragmentValues
 
 	FORCEINLINE bool ContainsType(const UScriptStruct* FragmentType) const
 	{
-		return FragmentType != nullptr && SharedFragmentBitSet.Contains(*FragmentType);
+		if (FragmentType)
+		{
+			if (FragmentType->IsChildOf(FMassSharedFragment::StaticStruct()))
+			{
+				return SharedFragmentBitSet.Contains(*FragmentType);
+			}
+
+			if (FragmentType->IsChildOf(FMassConstSharedFragment::StaticStruct()))
+			{
+				return ConstSharedFragmentBitSet.Contains(*FragmentType);
+			}
+		}
+
+		return false;
 	}
 
 	template<typename T>
 	FORCEINLINE bool ContainsType() const
 	{
-		return SharedFragmentBitSet.Contains(*T::StaticStruct());
+		if constexpr (TIsDerivedFrom<T, FMassConstSharedFragment>::IsDerived)
+		{
+			return ConstSharedFragmentBitSet.Contains(*T::StaticStruct());
+		}
+		else if constexpr (TIsDerivedFrom<T, FMassSharedFragment>::IsDerived)
+		{
+			return SharedFragmentBitSet.Contains(*T::StaticStruct());
+		}
+		else
+		{
+			return false;
+		}
 	}
 
 	/** 
@@ -322,6 +423,12 @@ struct MASSENTITY_API FMassArchetypeSharedFragmentValues
 	 */
 	int32 Remove(const FMassSharedFragmentBitSet& SharedFragmentToRemoveBitSet);
 
+	/** 
+	 * Note that the function removes the const shared fragments by type
+	 * @return number of fragments types removed
+	 */
+	int32 Remove(const FMassConstSharedFragmentBitSet& ConstSharedFragmentToRemoveBitSet);
+
 	FORCEINLINE const TArray<FConstSharedStruct>& GetConstSharedFragments() const
 	{
 		return ConstSharedFragments;
@@ -354,6 +461,11 @@ struct MASSENTITY_API FMassArchetypeSharedFragmentValues
 		return SharedFragmentBitSet;
 	}
 
+	const FMassConstSharedFragmentBitSet& GetConstSharedFragmentBitSet() const
+	{
+		return ConstSharedFragmentBitSet;
+	}
+
 	FORCEINLINE void DirtyHashCache()
 	{
 		HashCache = UINT32_MAX;
@@ -403,8 +515,19 @@ protected:
 	mutable bool bSorted = true; 
 	
 	FMassSharedFragmentBitSet SharedFragmentBitSet;
+	FMassConstSharedFragmentBitSet ConstSharedFragmentBitSet;
 	TArray<FConstSharedStruct> ConstSharedFragments;
 	TArray<FSharedStruct> SharedFragments;
+
+public:
+	//-----------------------------------------------------------------------------
+	// DEPRECATED
+	//-----------------------------------------------------------------------------
+	UE_DEPRECATED(5.5, "HasExactFragmentTypesMatch is deprecated. Use HasExactSharedFragmentTypesMatch or the two-parameter version of HasExactFragmentTypesMatch.")
+	FORCEINLINE bool HasExactFragmentTypesMatch(const FMassSharedFragmentBitSet& InSharedFragmentBitSet) const
+	{
+		return HasExactSharedFragmentTypesMatch(InSharedFragmentBitSet);
+	}
 };
 
 UENUM()
diff --git a/Engine/Plugins/Runtime/MassEntity/Source/MassEntity/Public/MassEntityView.h b/Engine/Plugins/Runtime/MassEntity/Source/MassEntity/Public/MassEntityView.h
--- a/Engine/Plugins/Runtime/MassEntity/Source/MassEntity/Public/MassEntityView.h
+++ b/Engine/Plugins/Runtime/MassEntity/Source/MassEntity/Public/MassEntityView.h
@@ -78,8 +78,8 @@ struct MASSENTITY_API FMassEntityView
 	template<typename T>
 	const T* GetConstSharedFragmentDataPtr() const
 	{
-		static_assert(TIsDerivedFrom<T, FMassSharedFragment>::IsDerived,
-			"Given struct doesn't represent a valid shared fragment type. Make sure to inherit from FMassSharedFragment or one of its child-types.");
+		static_assert(TIsDerivedFrom<T, FMassConstSharedFragment>::IsDerived,
+			"Given struct doesn't represent a valid shared fragment type. Make sure to inherit from FMassConstSharedFragment or one of its child-types.");
 
 		return (const T*)GetConstSharedFragmentPtr(*T::StaticStruct());
 	}
@@ -88,38 +88,54 @@ struct MASSENTITY_API FMassEntityView
 	template<typename T>
 	const T& GetConstSharedFragmentData() const
 	{
-		static_assert(TIsDerivedFrom<T, FMassSharedFragment>::IsDerived,
-			"Given struct doesn't represent a valid const shared fragment type. Make sure to inherit from FMassSharedFragment or one of its child-types.");
+		static_assert(TIsDerivedFrom<T, FMassConstSharedFragment>::IsDerived,
+			"Given struct doesn't represent a valid const shared fragment type. Make sure to inherit from FMassConstSharedFragment or one of its child-types.");
 
 		return *((const T*)GetConstSharedFragmentPtrChecked(*T::StaticStruct()));
 	}
 
 	FConstStructView GetConstSharedFragmentDataStruct(const UScriptStruct* FragmentType) const
 	{
-		check(FragmentType);
+		check(FragmentType && FragmentType->IsChildOf(FMassConstSharedFragment::StaticStruct()));
 		return FConstStructView(FragmentType, static_cast<const uint8*>(GetConstSharedFragmentPtr(*FragmentType)));
 	}
 
 	/** will fail a check if the viewed entity doesn't have the given shared fragment */
-	template<typename T>
+	template<typename T, TEMPLATE_REQUIRES(TIsDerivedFrom<T, FMassSharedFragment>::Value)>
 	T& GetSharedFragmentData() const
 	{
-		static_assert(TIsDerivedFrom<T, FMassSharedFragment>::IsDerived,
-			"Given struct doesn't represent a valid shared fragment type. Make sure to inherit from FMassSharedFragment or one of its child-types.");
-
 		return *((T*)GetSharedFragmentPtrChecked(*T::StaticStruct()));
 	}
 
 	/** if the viewed entity doesn't have the given shared fragment the function will return null */
-	template<typename T>
+	template<typename T, TEMPLATE_REQUIRES(TIsDerivedFrom<T, FMassSharedFragment>::Value)>
 	T* GetSharedFragmentDataPtr() const
 	{
-		static_assert(TIsDerivedFrom<T, FMassSharedFragment>::IsDerived,
-			"Given struct doesn't represent a valid shared fragment type. Make sure to inherit from FMassSharedFragment or one of its child-types.");
-
 		return (T*)GetSharedFragmentPtr(*T::StaticStruct());
 	}
 
+	template<typename T, TEMPLATE_REQUIRES(TIsDerivedFrom<T, FMassConstSharedFragment>::Value)>
+	UE_DEPRECATED(5.5, "Using GetSharedFragmentDataPtr with const shared fragments is deprecated. Use GetConstSharedFragmentDataPtr instead")
+	T* GetSharedFragmentDataPtr() const
+	{
+		return const_cast<T*>(GetConstSharedFragmentDataPtr<T>());
+	}
+
+	template<typename T, TEMPLATE_REQUIRES(TIsDerivedFrom<T, FMassConstSharedFragment>::Value)>
+	UE_DEPRECATED(5.5, "Using GetSharedFragmentDataPtr with const shared fragments is deprecated. Use GetConstSharedFragmentData instead")
+	T& GetSharedFragmentData() const
+	{
+		/*static T DummyInstance;
+		return DummyInstance;*/
+		return const_cast<T&>(GetConstSharedFragmentData<T>());
+	}
+
+	FStructView GetSharedFragmentDataStruct(const UScriptStruct* FragmentType) const
+	{
+		check(FragmentType && FragmentType->IsChildOf(FMassSharedFragment::StaticStruct()));
+		return FStructView(FragmentType, static_cast<uint8*>(GetSharedFragmentPtr(*FragmentType)));
+	}
+
 	template<typename T>
 	bool HasTag() const
 	{
diff --git a/Engine/Plugins/Runtime/MassEntity/Source/MassEntity/Public/MassExecutionContext.h b/Engine/Plugins/Runtime/MassEntity/Source/MassEntity/Public/MassExecutionContext.h
--- a/Engine/Plugins/Runtime/MassEntity/Source/MassEntity/Public/MassExecutionContext.h
+++ b/Engine/Plugins/Runtime/MassEntity/Source/MassEntity/Public/MassExecutionContext.h
@@ -198,7 +198,7 @@ public:
 	template<typename T>
 	const T* GetConstSharedFragmentPtr() const
 	{
-		static_assert(TIsDerivedFrom<T, FMassSharedFragment>::IsDerived, "Given struct doesn't represent a valid Shared fragment type. Make sure to inherit from FMassSharedFragment or one of its child-types.");
+		static_assert(TIsDerivedFrom<T, FMassConstSharedFragment>::IsDerived, "Given struct doesn't represent a valid const shared fragment type. Make sure to inherit from FMassConstSharedFragment or one of its child-types.");
 
 		const FConstSharedFragmentView* FoundSharedFragmentData = ConstSharedFragmentViews.FindByPredicate([](const FConstSharedFragmentView& Element) { return Element.Requirement.StructType == T::StaticStruct(); });
 		return FoundSharedFragmentData ? FoundSharedFragmentData->FragmentView.GetPtr<const T>() : static_cast<const T*>(nullptr);
diff --git a/Engine/Plugins/Runtime/MassEntity/Source/MassEntity/Public/MassProcessorDependencySolver.h b/Engine/Plugins/Runtime/MassEntity/Source/MassEntity/Public/MassProcessorDependencySolver.h
--- a/Engine/Plugins/Runtime/MassEntity/Source/MassEntity/Public/MassProcessorDependencySolver.h
+++ b/Engine/Plugins/Runtime/MassEntity/Source/MassEntity/Public/MassProcessorDependencySolver.h
@@ -40,6 +40,24 @@ struct TMassExecutionAccess
 	bool IsEmpty() const { return Read.IsEmpty() && Write.IsEmpty(); }
 };
 
+/** 
+ * TMassExecutionAccess specialization for FMassConstSharedFragmentBitSet to enforce lack of access (not needed) and
+ * no "Write" component (conceptually doesn't make sense).
+ */
+template<>
+struct TMassExecutionAccess<FMassConstSharedFragmentBitSet>
+{
+	FMassConstSharedFragmentBitSet Read;
+	TConstArrayView<FMassConstSharedFragmentBitSet> AsArrayView() const 
+	{ 
+		return MakeArrayView(&Read, 1); 
+	}
+	bool IsEmpty() const 
+	{ 
+		return Read.IsEmpty(); 
+	}
+};
+
 struct MASSENTITY_API FMassExecutionRequirements
 {
 	void Append(const FMassExecutionRequirements& Other);
@@ -51,6 +69,7 @@ struct MASSENTITY_API FMassExecutionRequirements
 	TMassExecutionAccess<FMassFragmentBitSet> Fragments;
 	TMassExecutionAccess<FMassChunkFragmentBitSet> ChunkFragments;
 	TMassExecutionAccess<FMassSharedFragmentBitSet> SharedFragments;
+	TMassExecutionAccess<FMassConstSharedFragmentBitSet> ConstSharedFragments;
 	TMassExecutionAccess<FMassExternalSubsystemBitSet> RequiredSubsystems;
 	FMassTagBitSet RequiredAllTags;
 	FMassTagBitSet RequiredAnyTags;
diff --git a/Engine/Plugins/Runtime/MassEntity/Source/MassEntity/Public/MassRequirements.h b/Engine/Plugins/Runtime/MassEntity/Source/MassEntity/Public/MassRequirements.h
--- a/Engine/Plugins/Runtime/MassEntity/Source/MassEntity/Public/MassRequirements.h
+++ b/Engine/Plugins/Runtime/MassEntity/Source/MassEntity/Public/MassRequirements.h
@@ -307,7 +307,7 @@ public:
 	template<typename T>
 	FMassFragmentRequirements& AddConstSharedRequirement(const EMassFragmentPresence Presence = EMassFragmentPresence::All)
 	{
-		static_assert(TIsDerivedFrom<T, FMassSharedFragment>::IsDerived, "Given struct doesn't represent a valid shared fragment type. Make sure to inherit from FMassSharedFragment or one of its child-types.");
+		static_assert(TIsDerivedFrom<T, FMassConstSharedFragment>::IsDerived, "Given struct doesn't represent a valid shared fragment type. Make sure to inherit from FMassConstSharedFragment or one of its child-types.");
 		checkf(ConstSharedFragmentRequirements.FindByPredicate([](const FMassFragmentRequirementDescription& Item) { return Item.StructType == T::StaticStruct(); }) == nullptr
 			, TEXT("Duplicated requirements are not supported. %s already present"), *T::StaticStruct()->GetName());
 		checkf(Presence != EMassFragmentPresence::Any, TEXT("\'Any\' is not a valid Presence value for AddConstSharedRequirement."));
@@ -315,15 +315,15 @@ public:
 		switch (Presence)
 		{
 		case EMassFragmentPresence::All:
-			RequiredAllSharedFragments.Add<T>();
+			RequiredAllConstSharedFragments.Add<T>();
 			ConstSharedFragmentRequirements.Emplace(T::StaticStruct(), EMassFragmentAccess::ReadOnly, Presence);
 			break;
 		case EMassFragmentPresence::Optional:
-			RequiredOptionalSharedFragments.Add<T>();
+			RequiredOptionalConstSharedFragments.Add<T>();
 			ConstSharedFragmentRequirements.Emplace(T::StaticStruct(), EMassFragmentAccess::ReadOnly, Presence);
 			break;
 		case EMassFragmentPresence::None:
-			RequiredNoneSharedFragments.Add<T>();
+			RequiredNoneConstSharedFragments.Add<T>();
 			break;
 		}
 		IncrementChangeCounter();
@@ -332,8 +332,8 @@ public:
 
 	FMassFragmentRequirements& AddConstSharedRequirement(const UScriptStruct* FragmentType, const EMassFragmentPresence Presence = EMassFragmentPresence::All)
 	{
-		if (!ensureMsgf(FragmentType->IsChildOf(FMassSharedFragment::StaticStruct())
-			, TEXT("Given struct doesn't represent a valid shared fragment type. Make sure to inherit from FMassSharedFragment or one of its child-types.")))
+		if (!ensureMsgf(FragmentType->IsChildOf(FMassConstSharedFragment::StaticStruct())
+			, TEXT("Given struct doesn't represent a valid shared fragment type. Make sure to inherit from FMassConstSharedFragment or one of its child-types.")))
 		{
 			return *this;
 		}
@@ -345,15 +345,15 @@ public:
 		switch (Presence)
 		{
 		case EMassFragmentPresence::All:
-			RequiredAllSharedFragments.Add(*FragmentType);
+			RequiredAllConstSharedFragments.Add(*FragmentType);
 			ConstSharedFragmentRequirements.Emplace(FragmentType, EMassFragmentAccess::ReadOnly, Presence);
 			break;
 		case EMassFragmentPresence::Optional:
-			RequiredOptionalSharedFragments.Add(*FragmentType);
+			RequiredOptionalConstSharedFragments.Add(*FragmentType);
 			ConstSharedFragmentRequirements.Emplace(FragmentType, EMassFragmentAccess::ReadOnly, Presence);
 			break;
 		case EMassFragmentPresence::None:
-			RequiredNoneSharedFragments.Add(*FragmentType);
+			RequiredNoneConstSharedFragments.Add(*FragmentType);
 			break;
 		}
 		IncrementChangeCounter();
@@ -421,6 +421,9 @@ public:
 	const FMassSharedFragmentBitSet& GetRequiredAllSharedFragments() const { return RequiredAllSharedFragments; }
 	const FMassSharedFragmentBitSet& GetRequiredOptionalSharedFragments() const { return RequiredOptionalSharedFragments; }
 	const FMassSharedFragmentBitSet& GetRequiredNoneSharedFragments() const { return RequiredNoneSharedFragments; }
+	const FMassConstSharedFragmentBitSet& GetRequiredAllConstSharedFragments() const { return RequiredAllConstSharedFragments; }
+	const FMassConstSharedFragmentBitSet& GetRequiredOptionalConstSharedFragments() const { return RequiredOptionalConstSharedFragments; }
+	const FMassConstSharedFragmentBitSet& GetRequiredNoneConstSharedFragments() const { return RequiredNoneConstSharedFragments; }
 
 	bool IsEmpty() const;
 
@@ -455,6 +458,9 @@ protected:
 	FMassSharedFragmentBitSet RequiredAllSharedFragments;
 	FMassSharedFragmentBitSet RequiredOptionalSharedFragments;
 	FMassSharedFragmentBitSet RequiredNoneSharedFragments;
+	FMassConstSharedFragmentBitSet RequiredAllConstSharedFragments;
+	FMassConstSharedFragmentBitSet RequiredOptionalConstSharedFragments;
+	FMassConstSharedFragmentBitSet RequiredNoneConstSharedFragments;
 
 private:
 	mutable uint16 bValidityIsCached : 1 = false;
diff --git a/Engine/Plugins/Runtime/MassEntity/Source/MassEntityDebugger/Private/SMassQuery.cpp b/Engine/Plugins/Runtime/MassEntity/Source/MassEntityDebugger/Private/SMassQuery.cpp
--- a/Engine/Plugins/Runtime/MassEntity/Source/MassEntityDebugger/Private/SMassQuery.cpp
+++ b/Engine/Plugins/Runtime/MassEntity/Source/MassEntityDebugger/Private/SMassQuery.cpp
@@ -66,6 +66,7 @@ void SMassQuery::Construct(const FArguments& InArgs, TSharedPtr<FMassDebuggerQue
 	AddBitSet(Box, ExecutionRequirements.RequiredNoneTags,TEXT("RequiredNoneTags"), ReadOnlyBrush);
 	AddBitSetPair(Box, ExecutionRequirements.ChunkFragments, TEXT("Chunk Fragments"));
 	AddBitSetPair(Box, ExecutionRequirements.SharedFragments, TEXT("Shared Fragments"));
+	AddBitSetPair(Box, ExecutionRequirements.ConstSharedFragments, TEXT("Const Shared Fragments"));
 	AddBitSetPair(Box, ExecutionRequirements.RequiredSubsystems, TEXT("Required Subsystems"));
 
 	ChildSlot
diff --git a/Engine/Plugins/Runtime/MassEntity/Source/MassEntityTestSuite/Private/MassSharedFragmentsTest.cpp b/Engine/Plugins/Runtime/MassEntity/Source/MassEntityTestSuite/Private/MassSharedFragmentsTest.cpp
--- a/Engine/Plugins/Runtime/MassEntity/Source/MassEntityTestSuite/Private/MassSharedFragmentsTest.cpp
+++ b/Engine/Plugins/Runtime/MassEntity/Source/MassEntityTestSuite/Private/MassSharedFragmentsTest.cpp
@@ -67,12 +67,12 @@ struct FSharedFragmentValues_CreateConst : FExecutionTestBase
 		FMassArchetypeSharedFragmentValues Values;
 
 		{
-			FConstSharedStruct SharedFragmentInstance = FConstSharedStruct::Make<FTestSharedFragment_Int>(TestIntValue);
+			FConstSharedStruct SharedFragmentInstance = FConstSharedStruct::Make<FTestConstSharedFragment_Int>(TestIntValue);
 			Values.AddConstSharedFragment(SharedFragmentInstance);
 		}
 
-		const FTestSharedFragment_Int* ConstInstance = GetConstSharedFragmentPtr<FTestSharedFragment_Int>(Values);
-		FTestSharedFragment_Int* NonConstInstance = GetMutableSharedFragmentPtr<FTestSharedFragment_Int>(Values);
+		const FTestConstSharedFragment_Int* ConstInstance = GetConstSharedFragmentPtr<FTestConstSharedFragment_Int>(Values);
+		FTestConstSharedFragment_Int* NonConstInstance = GetMutableSharedFragmentPtr<FTestConstSharedFragment_Int>(Values);
 
 		AITEST_NULL("Fetching fragment as a shared fragment should fail", NonConstInstance);
 		AITEST_NOT_NULL("Fetching fragment as a const shared fragment should not fail", ConstInstance);
@@ -97,8 +97,8 @@ struct FSharedFragmentValues_Contains : FExecutionTestBase
 			, Values.ContainsType<FTestSharedFragment_Int>());
 
 		{
-			FConstSharedStruct SharedFragmentInstance = FConstSharedStruct::Make<FTestSharedFragment_Int>(TestIntValue);
-			Values.AddConstSharedFragment(SharedFragmentInstance);
+			FSharedStruct SharedFragmentInstance = FSharedStruct::Make<FTestSharedFragment_Int>(TestIntValue);
+			Values.AddSharedFragment(SharedFragmentInstance);
 		}
 
 		AITEST_TRUE("Empty FMassArchetypeSharedFragmentValues should fail ContainsType tests"
@@ -107,8 +107,8 @@ struct FSharedFragmentValues_Contains : FExecutionTestBase
 			, Values.ContainsType<FTestSharedFragment_Int>());
 
 		{
-			FConstSharedStruct SharedFragmentInstance = FConstSharedStruct::Make<FTestSharedFragment_Float>(TestFloatValue);
-			Values.AddConstSharedFragment(SharedFragmentInstance);
+			FSharedStruct SharedFragmentInstance = FSharedStruct::Make<FTestSharedFragment_Float>(TestFloatValue);
+			Values.AddSharedFragment(SharedFragmentInstance);
 		}
 		AITEST_TRUE("Empty FMassArchetypeSharedFragmentValues should fail ContainsType tests"
 			, Values.ContainsType(FTestSharedFragment_Float::StaticStruct()));
@@ -152,12 +152,12 @@ struct FSharedFragmentValues_Append : FExecutionTestBase
 		}
 
 		FMassArchetypeSharedFragmentValues ValuesConstInt;
-		ValuesConstInt.AddConstSharedFragment(FConstSharedStruct::Make<FTestSharedFragment_Int>(TestIntValue));
+		ValuesConstInt.AddConstSharedFragment(FConstSharedStruct::Make<FTestConstSharedFragment_Int>(TestIntValue));
 		FMassArchetypeSharedFragmentValues ValuesConstFloat;
-		ValuesConstFloat.AddConstSharedFragment(FConstSharedStruct::Make<FTestSharedFragment_Float>(TestFloatValue));
+		ValuesConstFloat.AddConstSharedFragment(FConstSharedStruct::Make<FTestConstSharedFragment_Float>(TestFloatValue));
 		FMassArchetypeSharedFragmentValues ValuestConstIntFloat;
-		ValuestConstIntFloat.AddConstSharedFragment(FConstSharedStruct::Make<FTestSharedFragment_Int>(TestIntValue));
-		ValuestConstIntFloat.AddConstSharedFragment(FConstSharedStruct::Make<FTestSharedFragment_Float>(TestFloatValue));
+		ValuestConstIntFloat.AddConstSharedFragment(FConstSharedStruct::Make<FTestConstSharedFragment_Int>(TestIntValue));
+		ValuestConstIntFloat.AddConstSharedFragment(FConstSharedStruct::Make<FTestConstSharedFragment_Float>(TestFloatValue));
 
 		{
 			FMassArchetypeSharedFragmentValues Values;
@@ -174,19 +174,6 @@ struct FSharedFragmentValues_Append : FExecutionTestBase
 			AITEST_TRUE("#8 Append results should match expectations", Values.HasSameValues(ValuestConstIntFloat));
 		}
 
-		// test mismatching "mode" (const vs non-const) failure
-		AITEST_SCOPED_CHECK("trying to switch", 2);
-		{
-			FMassArchetypeSharedFragmentValues Values;
-			Values.Append(ValuesNonConstInt);
-			AITEST_EQUAL("#1 Adding mismatching mode should fail", Values.Append(ValuesConstInt), 0);
-		}
-		{
-			FMassArchetypeSharedFragmentValues Values;
-			Values.Append(ValuesConstInt);
-			AITEST_EQUAL("#2 Adding mismatching mode should fail", Values.Append(ValuesNonConstInt), 0);
-		}
-
 		return true;
 	}
 };
@@ -225,25 +212,25 @@ struct FSharedFragmentValues_Remove : FExecutionTestBase
 		}
 		{
 			FMassArchetypeSharedFragmentValues ValuesConstInt;
-			ValuesConstInt.AddConstSharedFragment(FConstSharedStruct::Make<FTestSharedFragment_Int>(TestIntValue));
+			ValuesConstInt.AddConstSharedFragment(FConstSharedStruct::Make<FTestConstSharedFragment_Int>(TestIntValue));
 			FMassArchetypeSharedFragmentValues ValuesConstFloat;
-			ValuesConstFloat.AddConstSharedFragment(FConstSharedStruct::Make<FTestSharedFragment_Float>(TestFloatValue));
+			ValuesConstFloat.AddConstSharedFragment(FConstSharedStruct::Make<FTestConstSharedFragment_Float>(TestFloatValue));
 			FMassArchetypeSharedFragmentValues ValuestConstIntFloat;
-			ValuestConstIntFloat.AddConstSharedFragment(FConstSharedStruct::Make<FTestSharedFragment_Int>(TestIntValue));
-			ValuestConstIntFloat.AddConstSharedFragment(FConstSharedStruct::Make<FTestSharedFragment_Float>(TestFloatValue));
+			ValuestConstIntFloat.AddConstSharedFragment(FConstSharedStruct::Make<FTestConstSharedFragment_Int>(TestIntValue));
+			ValuestConstIntFloat.AddConstSharedFragment(FConstSharedStruct::Make<FTestConstSharedFragment_Float>(TestFloatValue));
 
 			{
 				FMassArchetypeSharedFragmentValues Values = ValuestConstIntFloat;
 				AITEST_TRUE("Assignment should result in same values", Values.HasSameValues(ValuestConstIntFloat));
 
 				// removing just the Int shared fragment
-				Values.Remove(ValuesConstInt.GetSharedFragmentBitSet());
+				Values.Remove(ValuesConstInt.GetConstSharedFragmentBitSet());
 				AITEST_TRUE("#3 Removal results should match expectations", Values.HasSameValues(ValuesConstFloat));
 			}
 			{
 				FMassArchetypeSharedFragmentValues Values = ValuestConstIntFloat;
 				// removing just the Float shared fragment
-				Values.Remove(ValuesConstFloat.GetSharedFragmentBitSet());
+				Values.Remove(ValuesConstFloat.GetConstSharedFragmentBitSet());
 				AITEST_TRUE("#4 Removal results should match expectations", Values.HasSameValues(ValuesConstInt));
 			}
 		}
@@ -336,7 +323,7 @@ struct FSharedFragmentBase : public FEntityTestBase
 	}
 };
 
-template<typename TSharedStruct>
+template<typename TSharedStruct, typename TSharedFragment>
 struct FSharedFragment_CreateEntitesWithSharedFragment : public FSharedFragmentBase
 {
 	virtual bool InstantTest() override
@@ -345,8 +332,8 @@ struct FSharedFragment_CreateEntitesWithSharedFragment : public FSharedFragmentB
 		constexpr int32 TestIntValueB = 63;
 		FMassEntityHandle EntityA, EntityB;
 
-		CreateEntity<TSharedStruct, FTestSharedFragment_Int>(EntityA, TestIntValueA);
-		CreateEntity<TSharedStruct, FTestSharedFragment_Int>(EntityB, TestIntValueB);
+		CreateEntity<TSharedStruct, TSharedFragment>(EntityA, TestIntValueA);
+		CreateEntity<TSharedStruct, TSharedFragment>(EntityB, TestIntValueB);
 
 		AITEST_EQUAL("Both entities should end up in the same archetype", EntityManager->GetArchetypeForEntityUnsafe(EntityA)
 			, EntityManager->GetArchetypeForEntityUnsafe(EntityB));
@@ -362,27 +349,27 @@ struct FSharedFragment_CreateEntitesWithSharedFragment : public FSharedFragmentB
 			SharedFragmentA = EntityManager->GetConstSharedFragmentDataStruct(EntityA, FTestSharedFragment_Int::StaticStruct());
 			SharedFragmentB = EntityManager->GetConstSharedFragmentDataStruct(EntityB, FTestSharedFragment_Int::StaticStruct());
 		}*/
-		FConstStructView SharedFragmentA = GetSharedFragmentView<TSharedStruct, FTestSharedFragment_Int>(EntityA);
-		FConstStructView SharedFragmentB = GetSharedFragmentView<TSharedStruct, FTestSharedFragment_Int>(EntityB);
+		FConstStructView SharedFragmentA = GetSharedFragmentView<TSharedStruct, TSharedFragment>(EntityA);
+		FConstStructView SharedFragmentB = GetSharedFragmentView<TSharedStruct, TSharedFragment>(EntityB);
 
 		AITEST_TRUE("SharedFragmentA should be valid", SharedFragmentA.IsValid());
 		AITEST_TRUE("SharedFragmentB should be valid", SharedFragmentB.IsValid());
-		AITEST_EQUAL("SharedFragmentA should be of expected type", SharedFragmentA.GetScriptStruct(), FTestSharedFragment_Int::StaticStruct());
-		AITEST_EQUAL("SharedFragmentB should be of expected type", SharedFragmentB.GetScriptStruct(), FTestSharedFragment_Int::StaticStruct());
+		AITEST_EQUAL("SharedFragmentA should be of expected type", SharedFragmentA.GetScriptStruct(), TSharedFragment::StaticStruct());
+		AITEST_EQUAL("SharedFragmentB should be of expected type", SharedFragmentB.GetScriptStruct(), TSharedFragment::StaticStruct());
 		AITEST_NOT_EQUAL("SharedFragmentA and SharedFragmentB should be different instanceS", SharedFragmentA, SharedFragmentB);
 		AITEST_NOT_EQUAL("SharedFragmentA and SharedFragmentB should be distinct", SharedFragmentA, SharedFragmentB);
-		AITEST_EQUAL("SharedFragmentA's value should match the expected value", SharedFragmentA.Get<const FTestSharedFragment_Int>().Value, TestIntValueA);
-		AITEST_EQUAL("SharedFragmentB's value should match the expected value", SharedFragmentB.Get<const FTestSharedFragment_Int>().Value, TestIntValueB);
+		AITEST_EQUAL("SharedFragmentA's value should match the expected value", SharedFragmentA.Get<const TSharedFragment>().Value, TestIntValueA);
+		AITEST_EQUAL("SharedFragmentB's value should match the expected value", SharedFragmentB.Get<const TSharedFragment>().Value, TestIntValueB);
 
 		return true;
 	}
 };
-using FSharedFragment_CreateEntitesWithNonConstSharedFragment = FSharedFragment_CreateEntitesWithSharedFragment<FSharedStruct>;
+using FSharedFragment_CreateEntitesWithNonConstSharedFragment = FSharedFragment_CreateEntitesWithSharedFragment<FSharedStruct, FTestSharedFragment_Int>;
 IMPLEMENT_AI_INSTANT_TEST(FSharedFragment_CreateEntitesWithNonConstSharedFragment, "System.Mass.SharedFragments.CreateEntities");
-using FSharedFragment_CreateEntitesWithConstSharedFragment = FSharedFragment_CreateEntitesWithSharedFragment<FConstSharedStruct>;
+using FSharedFragment_CreateEntitesWithConstSharedFragment = FSharedFragment_CreateEntitesWithSharedFragment<FConstSharedStruct, FTestConstSharedFragment_Int>;
 IMPLEMENT_AI_INSTANT_TEST(FSharedFragment_CreateEntitesWithConstSharedFragment, "System.Mass.SharedFragments.CreateEntitiesConst");
 
-template<typename TSharedStruct>
+template<typename TSharedStruct, typename TSharedFragment>
 struct FSharedFragment_BatchCreateEntitesWithSharedFragment : public FSharedFragmentBase
 {
 	virtual bool InstantTest() override
@@ -397,8 +384,8 @@ struct FSharedFragment_BatchCreateEntitesWithSharedFragment : public FSharedFrag
 		TArray<FMassEntityHandle> EntitiesA;
 		TArray<FMassEntityHandle> EntitiesB;
 
-		CreateEntities<TSharedStruct, FTestSharedFragment_Int>(EntitiesA, EntitiesToCreateNumA, TestIntValueA);
-		CreateEntities<TSharedStruct, FTestSharedFragment_Int>(EntitiesB, EntitiesToCreateNumB, TestIntValueB);
+		CreateEntities<TSharedStruct, TSharedFragment>(EntitiesA, EntitiesToCreateNumA, TestIntValueA);
+		CreateEntities<TSharedStruct, TSharedFragment>(EntitiesB, EntitiesToCreateNumB, TestIntValueB);
 
 		FMassArchetypeHandle CommonArchetype = EntityManager->GetArchetypeForEntityUnsafe(EntitiesA[0]);
 		AITEST_EQUAL("All the entities should end up in the same archetype"
@@ -408,23 +395,23 @@ struct FSharedFragment_BatchCreateEntitesWithSharedFragment : public FSharedFrag
 
 		for (FMassEntityHandle EntityHandle : EntitiesA)
 		{
-			FConstStructView SharedFragment = GetSharedFragmentView<TSharedStruct, FTestSharedFragment_Int>(EntityHandle);
+			FConstStructView SharedFragment = GetSharedFragmentView<TSharedStruct, TSharedFragment>(EntityHandle);
 			AITEST_TRUE("SharedFragment for entity type A should be valid", SharedFragment.IsValid());
-			AITEST_EQUAL("SharedFragment for entity type A  should be of expected type", SharedFragment.GetScriptStruct(), FTestSharedFragment_Int::StaticStruct());
-			AITEST_EQUAL("SharedFragment's value for entity type A should match the expected value", SharedFragment.Get<const FTestSharedFragment_Int>().Value, TestIntValueA);
+			AITEST_EQUAL("SharedFragment for entity type A  should be of expected type", SharedFragment.GetScriptStruct(), TSharedFragment::StaticStruct());
+			AITEST_EQUAL("SharedFragment's value for entity type A should match the expected value", SharedFragment.Get<const TSharedFragment>().Value, TestIntValueA);
 		}
 
-		FConstStructView SharedFragment = GetSharedFragmentView<TSharedStruct, FTestSharedFragment_Int>(EntitiesB[0]);
+		FConstStructView SharedFragment = GetSharedFragmentView<TSharedStruct, TSharedFragment>(EntitiesB[0]);
 		AITEST_TRUE("SharedFragment for entity type B should be valid", SharedFragment.IsValid());
-		AITEST_EQUAL("SharedFragment for entity type B  should be of expected type", SharedFragment.GetScriptStruct(), FTestSharedFragment_Int::StaticStruct());
-		AITEST_EQUAL("SharedFragment's value for entity type B should match the expected value", SharedFragment.Get<const FTestSharedFragment_Int>().Value, TestIntValueB);
+		AITEST_EQUAL("SharedFragment for entity type B  should be of expected type", SharedFragment.GetScriptStruct(), TSharedFragment::StaticStruct());
+		AITEST_EQUAL("SharedFragment's value for entity type B should match the expected value", SharedFragment.Get<const TSharedFragment>().Value, TestIntValueB);
 
 		return true;
 	}
 };
-using FSharedFragment_BatchCreateEntitesWithNonConstSharedFragment = FSharedFragment_BatchCreateEntitesWithSharedFragment<FSharedStruct>;
+using FSharedFragment_BatchCreateEntitesWithNonConstSharedFragment = FSharedFragment_BatchCreateEntitesWithSharedFragment<FSharedStruct, FTestSharedFragment_Int>;
 IMPLEMENT_AI_INSTANT_TEST(FSharedFragment_BatchCreateEntitesWithNonConstSharedFragment, "System.Mass.SharedFragments.BatchCreateEntities");
-using FSharedFragment_BatchCreateEntitesWithConstSharedFragment = FSharedFragment_BatchCreateEntitesWithSharedFragment<FConstSharedStruct>;
+using FSharedFragment_BatchCreateEntitesWithConstSharedFragment = FSharedFragment_BatchCreateEntitesWithSharedFragment<FConstSharedStruct, FTestConstSharedFragment_Int>;
 IMPLEMENT_AI_INSTANT_TEST(FSharedFragment_BatchCreateEntitesWithConstSharedFragment, "System.Mass.SharedFragments.BatchCreateEntitiesConst");
 
 
@@ -433,17 +420,17 @@ struct FSharedFragment_AddToEntity : FEntityTestBase
 	virtual bool InstantTest() override
 	{
 		constexpr int32 TestIntValue = 1023;
-		FTestSharedFragment_Int FragmentInstance(TestIntValue);
-		FSharedStruct SharedFragmentInstance = FSharedStruct::Make(FragmentInstance);
+		FTestConstSharedFragment_Int FragmentInstance(TestIntValue);
+		FConstSharedStruct SharedFragmentInstance = FSharedStruct::Make(FragmentInstance);
 
 		const FMassEntityHandle EntityHandle = EntityManager->CreateEntity(FloatsArchetype);
 
-		FTestSharedFragment_Int* EntitySharedFragment = EntityManager->GetSharedFragmentDataPtr<FTestSharedFragment_Int>(EntityHandle);
+		FTestConstSharedFragment_Int* EntitySharedFragment = EntityManager->GetConstSharedFragmentDataPtr<FTestConstSharedFragment_Int>(EntityHandle);
 		AITEST_NULL("Initially the entity is not expected to have the shared fragment", EntitySharedFragment);
 
 		EntityManager->AddConstSharedFragmentToEntity(EntityHandle, SharedFragmentInstance);
 
-		EntitySharedFragment = EntityManager->GetConstSharedFragmentDataPtr<FTestSharedFragment_Int>(EntityHandle);
+		EntitySharedFragment = EntityManager->GetConstSharedFragmentDataPtr<FTestConstSharedFragment_Int>(EntityHandle);
 		AITEST_NOT_NULL("The entity is expected to have the shared fragment after the operation", EntitySharedFragment);
 		AITEST_EQUAL("The the shared fragment is expected to store the configured value", EntitySharedFragment->Value, TestIntValue);
 
@@ -451,7 +438,7 @@ struct FSharedFragment_AddToEntity : FEntityTestBase
 		// now we're going to add it again and test the systems behavior, we'll be adding the same FMasSharedFragment type
 		// in both const and non-const way.
 		constexpr int32 DifferentTestIntValue = TestIntValue + 1;
-		FTestSharedFragment_Int DifferentFragmentInstance(DifferentTestIntValue);
+		FTestConstSharedFragment_Int DifferentFragmentInstance(DifferentTestIntValue);
 		FSharedStruct DifferentSharedFragmentInstance = FSharedStruct::Make(DifferentFragmentInstance);
 		FConstSharedStruct DifferentConstSharedFragmentInstance = FConstSharedStruct::Make(DifferentFragmentInstance);
 
@@ -488,7 +475,7 @@ struct FSharedFragment_BatchAddToEntity : FEntityTestBase
 		FMassArchetypeEntityCollection EntityCollection(InitialArchetype, EntitiesMoved, FMassArchetypeEntityCollection::EDuplicatesHandling::NoDuplicates);
 
 		FMassArchetypeSharedFragmentValues SharedValues;
-		FConstSharedStruct ConstSharedFragment = FConstSharedStruct::Make<FTestSharedFragment_Int>(TestIntValue);
+		FConstSharedStruct ConstSharedFragment = FConstSharedStruct::Make<FTestConstSharedFragment_Int>(TestIntValue);
 		SharedValues.AddConstSharedFragment(ConstSharedFragment);
 		EntityManager->BatchAddSharedFragmentsForEntities(MakeArrayView(&EntityCollection, 1), SharedValues);
 
@@ -497,7 +484,7 @@ struct FSharedFragment_BatchAddToEntity : FEntityTestBase
 		AITEST_EQUAL("Number of entities moves needs to match expectations", EntitiesMovedNum, EntitiesToMoveNum);
 		for (const FMassEntityHandle& EntityHandle : EntitiesMoved)
 		{
-			FTestSharedFragment_Int* SharedFragmentInstance = EntityManager->GetConstSharedFragmentDataPtr<FTestSharedFragment_Int>(EntityHandle);
+			FTestConstSharedFragment_Int* SharedFragmentInstance = EntityManager->GetConstSharedFragmentDataPtr<FTestConstSharedFragment_Int>(EntityHandle);
 			AITEST_NOT_NULL("Every entity moved needs to have a valid shared fragment", SharedFragmentInstance);
 			AITEST_EQUAL("The shared fragment's value needs to match expectations", SharedFragmentInstance->Value, TestIntValue);
 		}
@@ -520,7 +507,7 @@ struct FSharedFragment_BatchSetAttempt : FEntityTestBase
 
 		TArray<FMassEntityHandle> CreatedEntityHandles;
 		FMassArchetypeSharedFragmentValues SharedIntValues;
-		FConstSharedStruct ConstSharedFragment = FConstSharedStruct::Make<FTestSharedFragment_Int>(TestIntValue);
+		FConstSharedStruct ConstSharedFragment = FConstSharedStruct::Make<FTestConstSharedFragment_Int>(TestIntValue);
 		SharedIntValues.AddConstSharedFragment(ConstSharedFragment);
 
 		TSharedRef<FMassEntityManager::FEntityCreationContext> CreationContext = EntityManager->BatchCreateEntities(FloatsArchetype, SharedIntValues, EntitiesToCreateNum, CreatedEntityHandles);
@@ -536,7 +523,7 @@ struct FSharedFragment_BatchSetAttempt : FEntityTestBase
 		}
 		{
 			FMassArchetypeSharedFragmentValues DifferentSharedIntValues;
-			FConstSharedStruct OtherConstSharedFragment = FConstSharedStruct::Make<FTestSharedFragment_Int>(OtherTestIntValue);
+			FConstSharedStruct OtherConstSharedFragment = FConstSharedStruct::Make<FTestConstSharedFragment_Int>(OtherTestIntValue);
 			DifferentSharedIntValues.AddConstSharedFragment(OtherConstSharedFragment);
 
 			AITEST_SCOPED_CHECK("Setting shared fragment values without archetype change is not supported", 1);
@@ -556,8 +543,8 @@ struct FSharedFragment_BatchAddToEmpty : FEntityTestBase
 		constexpr int32 NumToReserve = 32;
 
 		FMassArchetypeSharedFragmentValues SharedIntValues;
-		FConstSharedStruct ConstSharedFragment = FConstSharedStruct::Make<FTestSharedFragment_Int>();
-		SharedIntValues.AddConstSharedFragment(ConstSharedFragment);
+		FSharedStruct SharedFragment = FSharedStruct::Make<FTestSharedFragment_Int>();
+		SharedIntValues.AddSharedFragment(SharedFragment);
 
 		TArray<FMassEntityHandle> ReservedEntityHandles;
 		EntityManager->BatchReserveEntities(NumToReserve, ReservedEntityHandles);
diff --git a/Engine/Plugins/Runtime/MassEntity/Source/MassEntityTestSuite/Public/MassEntityTestTypes.h b/Engine/Plugins/Runtime/MassEntity/Source/MassEntityTestSuite/Public/MassEntityTestTypes.h
--- a/Engine/Plugins/Runtime/MassEntity/Source/MassEntityTestSuite/Public/MassEntityTestTypes.h
+++ b/Engine/Plugins/Runtime/MassEntity/Source/MassEntityTestSuite/Public/MassEntityTestTypes.h
@@ -99,6 +99,19 @@ struct FTestSharedFragment_Int : public FMassSharedFragment
 	FTestSharedFragment_Int(const int32 InValue = 0) : Value(InValue) {}
 };
 
+USTRUCT()
+struct FTestConstSharedFragment_Int : public FMassConstSharedFragment
+{
+	using FValueType = int32;
+
+	GENERATED_BODY()
+
+	UPROPERTY()
+	int32 Value = 0;
+
+	FTestConstSharedFragment_Int(const int32 InValue = 0) : Value(InValue) {}
+};
+
 USTRUCT()
 struct FTestSharedFragment_Float : public FMassSharedFragment
 {
@@ -112,6 +125,19 @@ struct FTestSharedFragment_Float : public FMassSharedFragment
 	FTestSharedFragment_Float(const float InValue = 0) : Value(InValue) {}
 };
 
+USTRUCT()
+struct FTestConstSharedFragment_Float : public FMassConstSharedFragment
+{
+	using FValueType = float;
+
+	GENERATED_BODY()
+
+	UPROPERTY()
+	float Value = 0.f;
+
+	FTestConstSharedFragment_Float(const float InValue = 0) : Value(InValue) {}
+};
+
 /** @todo rename to FTestTag */
 USTRUCT()
 struct FTestFragment_Tag : public FMassTag
diff --git a/Engine/Plugins/Runtime/MassGameplay/Source/MassGameplayTestSuite/Private/MassEntityTemplateTest.cpp b/Engine/Plugins/Runtime/MassGameplay/Source/MassGameplayTestSuite/Private/MassEntityTemplateTest.cpp
--- a/Engine/Plugins/Runtime/MassGameplay/Source/MassGameplayTestSuite/Private/MassEntityTemplateTest.cpp
+++ b/Engine/Plugins/Runtime/MassGameplay/Source/MassGameplayTestSuite/Private/MassEntityTemplateTest.cpp
@@ -69,7 +69,7 @@ struct FEntityTemplate_Composition : FEntityTemplateBase
 		}
 		{
 			FSharedStruct SharedFragment = EntityManager->GetOrCreateSharedFragmentByHash<FTestSharedFragment_Float>(1);
-			TemplateData.AddConstSharedFragment(SharedFragment);
+			TemplateData.AddSharedFragment(SharedFragment);
 		}
 		TemplateData.Sort();
 
diff --git a/Engine/Plugins/Runtime/MassGameplay/Source/MassLOD/Public/MassSimulationLOD.h b/Engine/Plugins/Runtime/MassGameplay/Source/MassLOD/Public/MassSimulationLOD.h
--- a/Engine/Plugins/Runtime/MassGameplay/Source/MassLOD/Public/MassSimulationLOD.h
+++ b/Engine/Plugins/Runtime/MassGameplay/Source/MassLOD/Public/MassSimulationLOD.h
@@ -79,7 +79,7 @@ struct MASSLOD_API FMassSimulationVariableTickChunkFragment : public FMassVariab
 };
 
 USTRUCT()
-struct MASSLOD_API FMassSimulationLODParameters : public FMassSharedFragment
+struct MASSLOD_API FMassSimulationLODParameters : public FMassConstSharedFragment
 {
 	GENERATED_BODY()
 
@@ -103,7 +103,7 @@ struct MASSLOD_API FMassSimulationLODParameters : public FMassSharedFragment
 };
 
 USTRUCT()
-struct MASSLOD_API FMassSimulationVariableTickParameters : public FMassSharedFragment
+struct MASSLOD_API FMassSimulationVariableTickParameters : public FMassConstSharedFragment
 {
 	GENERATED_BODY()
 
diff --git a/Engine/Plugins/Runtime/MassGameplay/Source/MassMovement/Public/MassMovementFragments.h b/Engine/Plugins/Runtime/MassGameplay/Source/MassMovement/Public/MassMovementFragments.h
--- a/Engine/Plugins/Runtime/MassGameplay/Source/MassMovement/Public/MassMovementFragments.h
+++ b/Engine/Plugins/Runtime/MassGameplay/Source/MassMovement/Public/MassMovementFragments.h
@@ -25,7 +25,7 @@ struct MASSMOVEMENT_API FMassForceFragment : public FMassFragment
 };
 
 USTRUCT()
-struct MASSMOVEMENT_API FMassMovementParameters : public FMassSharedFragment
+struct MASSMOVEMENT_API FMassMovementParameters : public FMassConstSharedFragment
 {
 	GENERATED_BODY()
 
diff --git a/Engine/Plugins/Runtime/MassGameplay/Source/MassReplication/Public/MassReplicationFragments.h b/Engine/Plugins/Runtime/MassGameplay/Source/MassReplication/Public/MassReplicationFragments.h
--- a/Engine/Plugins/Runtime/MassGameplay/Source/MassReplication/Public/MassReplicationFragments.h
+++ b/Engine/Plugins/Runtime/MassGameplay/Source/MassReplication/Public/MassReplicationFragments.h
@@ -76,7 +76,7 @@ protected:
 };
 
 USTRUCT()
-struct FMassReplicationParameters : public FMassSharedFragment
+struct FMassReplicationParameters : public FMassConstSharedFragment
 {
 	GENERATED_BODY()
 public:
diff --git a/Engine/Plugins/Runtime/MassGameplay/Source/MassRepresentation/Public/MassRepresentationFragments.h b/Engine/Plugins/Runtime/MassGameplay/Source/MassRepresentation/Public/MassRepresentationFragments.h
--- a/Engine/Plugins/Runtime/MassGameplay/Source/MassRepresentation/Public/MassRepresentationFragments.h
+++ b/Engine/Plugins/Runtime/MassGameplay/Source/MassRepresentation/Public/MassRepresentationFragments.h
@@ -96,7 +96,7 @@ struct TMassSharedFragmentTraits<FMassRepresentationSubsystemSharedFragment> fin
 
 
 USTRUCT()
-struct FMassRepresentationParameters : public FMassSharedFragment
+struct FMassRepresentationParameters : public FMassConstSharedFragment
 {
 	GENERATED_BODY()
 
@@ -191,7 +191,7 @@ inline void FMassRepresentationParameters::ComputeCachedValues() const
 }
 
 USTRUCT()
-struct FMassVisualizationLODParameters : public FMassSharedFragment
+struct FMassVisualizationLODParameters : public FMassConstSharedFragment
 {
 	GENERATED_BODY()
 
@@ -246,7 +246,7 @@ struct FMassVisualizationLODSharedFragment : public FMassSharedFragment
  *	* we do not care about Max Count
  */
 USTRUCT()
-struct FMassDistanceLODParameters : public FMassSharedFragment
+struct FMassDistanceLODParameters : public FMassConstSharedFragment
 {
 	GENERATED_BODY()
 
diff --git a/Engine/Plugins/Runtime/MassGameplay/Source/MassSpawner/Public/MassEntityTemplate.h b/Engine/Plugins/Runtime/MassGameplay/Source/MassSpawner/Public/MassEntityTemplate.h
--- a/Engine/Plugins/Runtime/MassGameplay/Source/MassSpawner/Public/MassEntityTemplate.h
+++ b/Engine/Plugins/Runtime/MassGameplay/Source/MassSpawner/Public/MassEntityTemplate.h
@@ -208,11 +208,11 @@ struct MASSSPAWNER_API FMassEntityTemplateData
 	void AddConstSharedFragment(const FConstSharedStruct& SharedFragment)
 	{
 		const UScriptStruct* FragmentType = SharedFragment.GetScriptStruct();
-		if(ensureMsgf(FragmentType && FragmentType->IsChildOf(FMassSharedFragment::StaticStruct()), TEXT("Given struct doesn't represent a valid shared fragment type. Make sure to inherit from FMassSharedFragment or one of its child-types.")))
+		if(ensureMsgf(FragmentType && FragmentType->IsChildOf(FMassConstSharedFragment::StaticStruct()), TEXT("Given struct doesn't represent a valid const shared fragment type. Make sure to inherit from FMassConstSharedFragment or one of its child-types.")))
 		{
-			if (!Composition.SharedFragments.Contains(*FragmentType))
+			if (!Composition.ConstSharedFragments.Contains(*FragmentType))
 			{
-				Composition.SharedFragments.Add(*FragmentType);
+				Composition.ConstSharedFragments.Add(*FragmentType);
 				SharedFragmentValues.AddConstSharedFragment(SharedFragment);
 			}
 #if DO_ENSURE
