Issue 460075

diff --git a/Engine/Plugins/Runtime/MassEntity/Source/MassEntity/Private/MassArchetypeData.cpp b/Engine/Plugins/Runtime/MassEntity/Source/MassEntity/Private/MassArchetypeData.cpp
--- a/Engine/Plugins/Runtime/MassEntity/Source/MassEntity/Private/MassArchetypeData.cpp
+++ b/Engine/Plugins/Runtime/MassEntity/Source/MassEntity/Private/MassArchetypeData.cpp
@@ -181,7 +181,8 @@ void FMassArchetypeData::AddEntity(FMassEntityHandle Entity, const FMassArchetyp
 int32 FMassArchetypeData::AddEntityInternal(FMassEntityHandle Entity, const FMassArchetypeSharedFragmentValues& SharedFragmentValues)
 {
 	checkf(SharedFragmentValues.IsSorted(), TEXT("Expecting shared fragment values to be previously sorted"));
-	checkfSlow(SharedFragmentValues.HasExactFragmentTypesMatch(CompositionDescriptor.SharedFragments), TEXT("Expecting values for every specified shared fragment in the archetype and only those"))
+	checkfSlow(SharedFragmentValues.HasExactFragmentTypesMatch(CompositionDescriptor.SharedFragments, CompositionDescriptor.ConstSharedFragments)
+		, TEXT("Expecting values for every specified shared fragment in the archetype and only those"))
 
 	int32 IndexWithinChunk = 0;
 	int32 AbsoluteIndex = 0;
@@ -1145,7 +1146,8 @@ void FMassArchetypeData::BatchMoveEntitiesToAnotherArchetype(const FMassArchetyp
 FMassArchetypeEntityCollection::FArchetypeEntityRange FMassArchetypeData::PrepareNextEntitiesSpanInternal(TConstArrayView<FMassEntityHandle> Entities, const FMassArchetypeSharedFragmentValues& SharedFragmentValues, int32 StartingChunk)
 {
 	checkf(SharedFragmentValues.IsSorted(), TEXT("Expecting shared fragment values to be previously sorted"));
-	checkfSlow(SharedFragmentValues.HasExactFragmentTypesMatch(CompositionDescriptor.SharedFragments), TEXT("Expecting values for every specified shared fragment in the archetype and only those"))
+	checkfSlow(SharedFragmentValues.HasExactFragmentTypesMatch(CompositionDescriptor.SharedFragments, CompositionDescriptor.ConstSharedFragments)
+		, TEXT("Expecting values for every specified shared fragment in the archetype and only those"))
 
 	int32 StartIndexWithinChunk = INDEX_NONE;
 	int32 AbsoluteStartIndex = 0;
