Mass: Expose Initialization State through the base subsystem

the pre flight is actually 662bc12f264864c1652cdbe7

diff --git a/Engine/Plugins/Runtime/MassEntity/Source/MassEntity/Public/MassSubsystemBase.h b/Engine/Plugins/Runtime/MassEntity/Source/MassEntity/Public/MassSubsystemBase.h
--- a/Engine/Plugins/Runtime/MassEntity/Source/MassEntity/Public/MassSubsystemBase.h
+++ b/Engine/Plugins/Runtime/MassEntity/Source/MassEntity/Public/MassSubsystemBase.h
@@ -27,6 +27,7 @@ class MASSENTITY_API UMassSubsystemBase : public UWorldSubsystem
 
 public:
 	static bool AreRuntimeMassSubsystemsAllowed(UObject* Outer);
+	UE::Mass::Subsystems::FInitializationState GetInitializationState() const { return InitializationState; }
 
 protected:
 	//~USubsystem interface
@@ -59,6 +60,9 @@ class MASSENTITY_API UMassTickableSubsystemBase : public UTickableWorldSubsystem
 {
 	GENERATED_BODY()
 
+public:
+	UE::Mass::Subsystems::FInitializationState GetInitializationState() const { return InitializationState; }
+
 protected:
 	//~USubsystem interface
 	virtual bool ShouldCreateSubsystem(UObject* Outer) const override;
