Remove extraneous #pragma once in .cpp file

diff --git a/Engine/Plugins/Runtime/MassEntity/Source/MassEntity/Private/MassEntityManagerStorage.cpp b/Engine/Plugins/Runtime/MassEntity/Source/MassEntity/Private/MassEntityManagerStorage.cpp
--- a/Engine/Plugins/Runtime/MassEntity/Source/MassEntity/Private/MassEntityManagerStorage.cpp
+++ b/Engine/Plugins/Runtime/MassEntity/Source/MassEntity/Private/MassEntityManagerStorage.cpp
@@ -1,5 +1,4 @@
 ﻿// Copyright Epic Games, Inc. All Rights Reserved.
-#pragma once
 
 #include "MassEntityManagerStorage.h"
 
