[Mass] fixed some deprecated calls and replaced some checkSlow by regular checks

diff --git a/Engine/Plugins/Runtime/MassEntity/Source/MassEntity/Private/MassArchetypeData.cpp b/Engine/Plugins/Runtime/MassEntity/Source/MassEntity/Private/MassArchetypeData.cpp
--- a/Engine/Plugins/Runtime/MassEntity/Source/MassEntity/Private/MassArchetypeData.cpp
+++ b/Engine/Plugins/Runtime/MassEntity/Source/MassEntity/Private/MassArchetypeData.cpp
@@ -137,7 +137,7 @@ void FMassArchetypeData::ConfigureFragments()
 	for (int32 FragmentIndex = 0; FragmentIndex < SortedFragmentList.Num(); ++FragmentIndex)
 	{
 		const UScriptStruct* FragmentType = SortedFragmentList[FragmentIndex];
-		checkSlow(FragmentType);
+		check(FragmentType);
 		FragmentConfigs[FragmentIndex].FragmentType = FragmentType;
 
 		AlignmentPadding += SIZE_T(FragmentType->GetMinAlignment());
@@ -181,7 +181,7 @@ void FMassArchetypeData::AddEntity(FMassEntityHandle Entity, const FMassArchetyp
 int32 FMassArchetypeData::AddEntityInternal(FMassEntityHandle Entity, const FMassArchetypeSharedFragmentValues& SharedFragmentValues)
 {
 	checkf(SharedFragmentValues.IsSorted(), TEXT("Expecting shared fragment values to be previously sorted"));
-	checkfSlow(SharedFragmentValues.HasExactFragmentTypesMatch(CompositionDescriptor.SharedFragments, CompositionDescriptor.ConstSharedFragments)
+	checkf(SharedFragmentValues.HasExactFragmentTypesMatch(CompositionDescriptor.SharedFragments, CompositionDescriptor.ConstSharedFragments)
 		, TEXT("Expecting values for every specified shared fragment in the archetype and only those"))
 
 	int32 IndexWithinChunk = 0;
@@ -1146,7 +1146,7 @@ void FMassArchetypeData::BatchMoveEntitiesToAnotherArchetype(const FMassArchetyp
 FMassArchetypeEntityCollection::FArchetypeEntityRange FMassArchetypeData::PrepareNextEntitiesSpanInternal(TConstArrayView<FMassEntityHandle> Entities, const FMassArchetypeSharedFragmentValues& SharedFragmentValues, int32 StartingChunk)
 {
 	checkf(SharedFragmentValues.IsSorted(), TEXT("Expecting shared fragment values to be previously sorted"));
-	checkfSlow(SharedFragmentValues.HasExactFragmentTypesMatch(CompositionDescriptor.SharedFragments, CompositionDescriptor.ConstSharedFragments)
+	checkf(SharedFragmentValues.HasExactFragmentTypesMatch(CompositionDescriptor.SharedFragments, CompositionDescriptor.ConstSharedFragments)
 		, TEXT("Expecting values for every specified shared fragment in the archetype and only those"))
 
 	int32 StartIndexWithinChunk = INDEX_NONE;
