[Mass] Fix concurrent storage to return IsValidIndex as false when passed an index that is not within the range of allocate entities but is within the range of possible entities.

diff --git a/Engine/Plugins/Runtime/MassEntity/Source/MassEntity/Private/MassEntityManagerStorage.cpp b/Engine/Plugins/Runtime/MassEntity/Source/MassEntity/Private/MassEntityManagerStorage.cpp
--- a/Engine/Plugins/Runtime/MassEntity/Source/MassEntity/Private/MassEntityManagerStorage.cpp
+++ b/Engine/Plugins/Runtime/MassEntity/Source/MassEntity/Private/MassEntityManagerStorage.cpp
@@ -252,8 +252,12 @@ namespace UE::Mass
 	bool FConcurrentEntityStorage::IsValidIndex(int32 Index) const
 	{
 		// Page Index is which page in the array of pages we need to access
-		const uint32 PageIndex = static_cast<uint32>(Index) >> MaximumEntityCountShift;
-		return PageIndex < PageCount;
+		if (Index >= 0)
+		{
+			const uint32 PageIndex = static_cast<uint32>(Index) >> FMath::FloorLog2(MaxEntitiesPerPage);
+			return PageIndex < PageCount;
+		}
+		return false;
 	}
 
 	SIZE_T FConcurrentEntityStorage::GetAllocatedSize() const
