Mass' chunk sizes aren't used by its derived archetypes.

When a new archetype is created by mutating an existing one, the new archetype doesn't use the chunk size of the original archetype but instead retrieves the size from the settings. This causes unexpected behavior as the intention is the automatically created archetype mirrors the original. This fix makes sure the chunk size is properly inherited.

[RN] Chunk sizes set on archetypes in Mass are now respected by automatically generated derived archetypes.

diff --git a/Engine/Plugins/Runtime/MassEntity/Source/MassEntity/Private/MassEntityManager.cpp b/Engine/Plugins/Runtime/MassEntity/Source/MassEntity/Private/MassEntityManager.cpp
--- a/Engine/Plugins/Runtime/MassEntity/Source/MassEntity/Private/MassEntityManager.cpp
+++ b/Engine/Plugins/Runtime/MassEntity/Source/MassEntity/Private/MassEntityManager.cpp
@@ -342,6 +342,12 @@ FMassArchetypeHandle FMassEntityManager::CreateArchetype(TConstArrayView<const U
 	return CreateArchetype(Composition, CreationParams);
 }
 
+FMassArchetypeHandle FMassEntityManager::CreateArchetype(FMassArchetypeHandle SourceArchetype, TConstArrayView<const UScriptStruct*> FragmentsAndTagsList)
+{
+	const FMassArchetypeData& ArchetypeData = FMassArchetypeHelper::ArchetypeDataFromHandleChecked(SourceArchetype); 
+	return CreateArchetype(SourceArchetype, FragmentsAndTagsList, FMassArchetypeCreationParams(ArchetypeData));
+}
+
 FMassArchetypeHandle FMassEntityManager::CreateArchetype(FMassArchetypeHandle SourceArchetype,
 	TConstArrayView<const UScriptStruct*> FragmentsAndTagsList, const FMassArchetypeCreationParams& CreationParams)
 {
@@ -351,6 +357,11 @@ FMassArchetypeHandle FMassEntityManager::CreateArchetype(FMassArchetypeHandle So
 	return CreateArchetype(Composition, CreationParams);
 }
 
+FMassArchetypeHandle FMassEntityManager::CreateArchetype(const TSharedPtr<FMassArchetypeData>& SourceArchetype, const FMassFragmentBitSet& AddedFragments)
+{
+	return CreateArchetype(SourceArchetype, AddedFragments, FMassArchetypeCreationParams(*SourceArchetype));
+}
+
 FMassArchetypeHandle FMassEntityManager::CreateArchetype(const TSharedPtr<FMassArchetypeData>& SourceArchetype, const FMassFragmentBitSet& AddedFragments, const FMassArchetypeCreationParams& CreationParams)
 {
 	check(SourceArchetype.IsValid());
@@ -479,7 +490,7 @@ FMassArchetypeHandle FMassEntityManager::InternalCreateSimilarArchetype(const FM
 		++ArchetypeDataVersion;
 
 		// Create a new archetype
-		FMassArchetypeData* NewArchetype = new FMassArchetypeData();
+		FMassArchetypeData* NewArchetype = new FMassArchetypeData(FMassArchetypeCreationParams(SourceArchetypeRef));
 		NewArchetype->InitializeWithSimilar(SourceArchetypeRef, MoveTemp(NewComposition), ArchetypeDataVersion);
 		NewArchetype->CopyDebugNamesFrom(SourceArchetypeRef);
 
@@ -931,7 +942,7 @@ void FMassEntityManager::AddCompositionToEntity_GetDelta(FMassEntityHandle Entit
 		NewDescriptor.Fragments += InDescriptor.Fragments;
 		NewDescriptor.Tags += InDescriptor.Tags;
 
-		const FMassArchetypeHandle NewArchetypeHandle = CreateArchetype(NewDescriptor);
+		const FMassArchetypeHandle NewArchetypeHandle = CreateArchetype(NewDescriptor, FMassArchetypeCreationParams(*OldArchetype));
 
 		if (ensure(NewArchetypeHandle.DataPtr.Get() != OldArchetype))
 		{
@@ -968,7 +979,7 @@ void FMassEntityManager::RemoveCompositionFromEntity(FMassEntityHandle Entity, c
 			ensureMsgf(OldArchetype->GetCompositionDescriptor().HasAll(InDescriptor), TEXT("Some of the elements being removed are already missing from entity\'s composition."));
 			ObserverManager.OnPreCompositionRemoved(Entity, InDescriptor);
 
-			const FMassArchetypeHandle NewArchetypeHandle = CreateArchetype(NewDescriptor);
+			const FMassArchetypeHandle NewArchetypeHandle = CreateArchetype(NewDescriptor, FMassArchetypeCreationParams(*OldArchetype));
 
 			if (ensure(NewArchetypeHandle.DataPtr.Get() != OldArchetype))
 			{
@@ -1079,10 +1090,10 @@ void FMassEntityManager::RemoveFragmentListFromEntity(FMassEntityHandle Entity,
 			, OldArchetype->GetChunkFragmentBitSet()
 			, OldArchetype->GetSharedFragmentBitSet()
 			, OldArchetype->GetConstSharedFragmentBitSet());
-		const FMassArchetypeHandle NewArchetypeHandle = CreateArchetype(NewComposition);
+		const FMassArchetypeHandle NewArchetypeHandle = CreateArchetype(NewComposition, FMassArchetypeCreationParams(*OldArchetype));
 
 		FMassArchetypeCompositionDescriptor CompositionDelta;
-		// Find overlap.  It isn't guaranteed that the old archtype has all of the fragments being removed.
+		// Find overlap.  It isn't guaranteed that the old archetype has all of the fragments being removed.
 		CompositionDelta.Fragments = OldArchetype->GetFragmentBitSet().GetOverlap(FragmentsToRemove);
 		ObserverManager.OnPreCompositionRemoved(Entity, CompositionDelta);
 
@@ -1209,7 +1220,7 @@ bool FMassEntityManager::AddConstSharedFragmentToEntity(const FMassEntityHandle
 	
 	FMassArchetypeCompositionDescriptor NewComposition(CurrentArchetype->GetCompositionDescriptor());
 	NewComposition.ConstSharedFragments.Add(*StructType);
-	const FMassArchetypeHandle NewArchetypeHandle = CreateArchetype(NewComposition);
+	const FMassArchetypeHandle NewArchetypeHandle = CreateArchetype(NewComposition, FMassArchetypeCreationParams(*CurrentArchetype));
 	check(NewArchetypeHandle.IsValid());
 	FMassArchetypeData* NewArchetype = NewArchetypeHandle.DataPtr.Get();
 	check(NewArchetype);
@@ -1422,7 +1433,7 @@ void FMassEntityManager::BatchAddSharedFragmentsForEntities(TConstArrayView<FMas
 		NewComposition.SharedFragments += AddedFragmentValues.GetSharedFragmentBitSet();
 		NewComposition.ConstSharedFragments += AddedFragmentValues.GetConstSharedFragmentBitSet();
 
-		const FMassArchetypeHandle NewArchetypeHandle = CreateArchetype(NewComposition);
+		const FMassArchetypeHandle NewArchetypeHandle = CreateArchetype(NewComposition, FMassArchetypeCreationParams(*CurrentArchetype));
 		check(NewArchetypeHandle.IsValid());
 		FMassArchetypeData* NewArchetype = NewArchetypeHandle.DataPtr.Get();
 		check(NewArchetype);
diff --git a/Engine/Plugins/Runtime/MassEntity/Source/MassEntity/Private/MassEntityTypes.cpp b/Engine/Plugins/Runtime/MassEntity/Source/MassEntity/Private/MassEntityTypes.cpp
--- a/Engine/Plugins/Runtime/MassEntity/Source/MassEntity/Private/MassEntityTypes.cpp
+++ b/Engine/Plugins/Runtime/MassEntity/Source/MassEntity/Private/MassEntityTypes.cpp
@@ -1,6 +1,8 @@
 // Copyright Epic Games, Inc. All Rights Reserved.
 
 #include "MassEntityTypes.h"
+
+#include "MassArchetypeData.h"
 #include "StructUtilsTypes.h"
 
 #include UE_INLINE_GENERATED_CPP_BY_NAME(MassEntityTypes)
@@ -287,4 +289,12 @@ void FMassGenericPayloadView::SwapElementsToEnd(const int32 StartIndex, int32 Nu
 		FMemory::Memmove(ViewData + MovedStartOffset, ViewData + MovedStartOffset + MovedSize, MoveOffset);
 		FMemory::Memcpy(ViewData + MovedStartOffset + MoveOffset, MovedElements.GetData(), MovedSize);
 	}
+}
+
+//-----------------------------------------------------------------------------
+// FMassArchetypeCreationParams
+//-----------------------------------------------------------------------------
+FMassArchetypeCreationParams::FMassArchetypeCreationParams(const FMassArchetypeData& Archetype)
+	: ChunkMemorySize(static_cast<int32>(Archetype.GetChunkAllocSize()))
+{
 }
\ No newline at end of file
diff --git a/Engine/Plugins/Runtime/MassEntity/Source/MassEntity/Public/MassEntityManager.h b/Engine/Plugins/Runtime/MassEntity/Source/MassEntity/Public/MassEntityManager.h
--- a/Engine/Plugins/Runtime/MassEntity/Source/MassEntity/Public/MassEntityManager.h
+++ b/Engine/Plugins/Runtime/MassEntity/Source/MassEntity/Public/MassEntityManager.h
@@ -109,9 +109,23 @@ public:
 	 * A special, relaxed but slower version of CreateArchetype functions that allows FragmentAngTagsList to contain
 	 * both fragments and tags. This version takes an original archetype and copies it layout, then appends any fragments and tags from the
 	 * provided list if they're not already in the original archetype.
+	 * 
+	 * @param SourceArchetype The archetype where the composition will be copied from.
+	 * @param FragmentsAndTagsList The list of fragments and tags to add to the copied composition.
 	 */
-	FMassArchetypeHandle CreateArchetype(FMassArchetypeHandle SourceArchetype, 
-		TConstArrayView<const UScriptStruct*> FragmentsAndTagsList, const FMassArchetypeCreationParams& CreationParams = FMassArchetypeCreationParams());
+	FMassArchetypeHandle CreateArchetype(FMassArchetypeHandle SourceArchetype, TConstArrayView<const UScriptStruct*> FragmentsAndTagsList);
+	
+	/**
+	 * A special, relaxed but slower version of CreateArchetype functions that allows FragmentAngTagsList to contain
+	 * both fragments and tags. This version takes an original archetype and copies it layout, then appends any fragments and tags from the
+	 * provided list if they're not already in the original archetype.
+	 * 
+	 * @param SourceArchetype The archetype where the composition will be copied from.
+	 * @param FragmentsAndTagsList The list of fragments and tags to add to the copied composition.
+	 * @param CreationParams Additional arguments used to create the new archetype.
+	 */
+	FMassArchetypeHandle CreateArchetype(FMassArchetypeHandle SourceArchetype, TConstArrayView<const UScriptStruct*> FragmentsAndTagsList, 
+		const FMassArchetypeCreationParams& CreationParams);
 
 	/**
 	 * CreateArchetype from a composition descriptor and initial values
@@ -122,6 +136,18 @@ public:
 	 */
 	FMassArchetypeHandle CreateArchetype(const FMassArchetypeCompositionDescriptor& Composition, const FMassArchetypeCreationParams& CreationParams = FMassArchetypeCreationParams());
 
+	/**
+	 *  Creates an archetype like SourceArchetype + InFragments.
+	 *  @param SourceArchetype the archetype used to initially populate the list of fragments of the archetype being created.
+	 *  @param InFragments list of unique fragments to add to fragments fetched from SourceArchetype. Note that
+	 *   adding an empty list is not supported and doing so will result in failing a `check`
+	 *  @return a handle of a new archetype
+	 *  @note it's caller's responsibility to ensure that NewFragmentList is not empty and contains only fragment
+	 *   types that SourceArchetype doesn't already have. If the caller cannot guarantee it use of AddFragment functions
+	 *   family is recommended.
+	 */
+	FMassArchetypeHandle CreateArchetype(const TSharedPtr<FMassArchetypeData>& SourceArchetype, const FMassFragmentBitSet& InFragments);
+	
 	/** 
 	 *  Creates an archetype like SourceArchetype + InFragments. 
 	 *  @param SourceArchetype the archetype used to initially populate the list of fragments of the archetype being created. 
@@ -133,8 +159,8 @@ public:
 	 *   types that SourceArchetype doesn't already have. If the caller cannot guarantee it use of AddFragment functions
 	 *   family is recommended.
 	 */
-	FMassArchetypeHandle CreateArchetype(const TSharedPtr<FMassArchetypeData>& SourceArchetype, const FMassFragmentBitSet& InFragments
-		, const FMassArchetypeCreationParams& CreationParams = FMassArchetypeCreationParams());
+	FMassArchetypeHandle CreateArchetype(const TSharedPtr<FMassArchetypeData>& SourceArchetype, const FMassFragmentBitSet& InFragments, 
+		const FMassArchetypeCreationParams& CreationParams);
 
 	/** 
 	 * A helper function to be used when creating entities with shared fragments provided, or when adding shared fragments
diff --git a/Engine/Plugins/Runtime/MassEntity/Source/MassEntity/Public/MassEntityTypes.h b/Engine/Plugins/Runtime/MassEntity/Source/MassEntity/Public/MassEntityTypes.h
--- a/Engine/Plugins/Runtime/MassEntity/Source/MassEntity/Public/MassEntityTypes.h
+++ b/Engine/Plugins/Runtime/MassEntity/Source/MassEntity/Public/MassEntityTypes.h
@@ -749,6 +749,9 @@ namespace UE::Mass
 
 struct FMassArchetypeCreationParams
 {
+	FMassArchetypeCreationParams() = default;
+	explicit FMassArchetypeCreationParams(const struct FMassArchetypeData& Archetype);
+
 	/** Created archetype will have chunks of this size. 0 denotes "use default" (see UE::Mass::ChunkSize) */
 	int32 ChunkMemorySize = 0;
 
