Moved MassArchetypeData.h over to Internal/ folder to remove MassEntityTestSuite's dependency on MassEntity module's Private folder

diff --git a/Engine/Plugins/Runtime/MassEntity/Source/MassEntity/Internal/MassArchetypeData.h b/Engine/Plugins/Runtime/MassEntity/Source/MassEntity/Internal/MassArchetypeData.h
new file mode 100644
--- /dev/null
+++ b/Engine/Plugins/Runtime/MassEntity/Source/MassEntity/Internal/MassArchetypeData.h
@@ -0,0 +1,421 @@
+// Copyright Epic Games, Inc. All Rights Reserved.
+
+#pragma once
+
+#include "MassEntityManager.h"
+#include "MassArchetypeTypes.h"
+
+
+struct FMassEntityQuery;
+struct FMassExecutionContext;
+class FOutputDevice;
+struct FMassArchetypeEntityCollection;
+
+namespace UE::Mass
+{
+	int32 SanitizeChunkMemorySize(const int32 InChunkMemorySize, const bool bLogMismatch = true);
+}
+
+// This is one chunk within an archetype
+struct FMassArchetypeChunk
+{
+private:
+	uint8* RawMemory = nullptr;
+	SIZE_T AllocSize = 0;
+	int32 NumInstances = 0;
+	int32 SerialModificationNumber = 0;
+	TArray<FInstancedStruct> ChunkFragmentData;
+	FMassArchetypeSharedFragmentValues SharedFragmentValues;
+
+public:
+	explicit FMassArchetypeChunk(const SIZE_T InAllocSize, TConstArrayView<FInstancedStruct> InChunkFragmentTemplates, FMassArchetypeSharedFragmentValues InSharedFragmentValues)
+		: AllocSize(InAllocSize)
+		, ChunkFragmentData(InChunkFragmentTemplates)
+		, SharedFragmentValues(InSharedFragmentValues)
+	{
+		RawMemory = (uint8*)FMemory::Malloc(AllocSize);
+	}
+
+	~FMassArchetypeChunk()
+	{
+		// Only release memory if it was not done already.
+		if (RawMemory != nullptr)
+		{
+			FMemory::Free(RawMemory);
+			RawMemory = nullptr;
+		}
+	}
+
+	// Returns the Entity array element at the specified index
+	FMassEntityHandle& GetEntityArrayElementRef(int32 ChunkBase, int32 IndexWithinChunk)
+	{
+		return ((FMassEntityHandle*)(RawMemory + ChunkBase))[IndexWithinChunk];
+	}
+
+	uint8* GetRawMemory() const
+	{
+		return RawMemory;
+	}
+
+	int32 GetNumInstances() const
+	{
+		return NumInstances;
+	}
+
+	void AddMultipleInstances(uint32 Count)
+	{
+		NumInstances += Count;
+		SerialModificationNumber++;
+	}
+
+	void RemoveMultipleInstances(uint32 Count)
+	{
+		NumInstances -= Count;
+		check(NumInstances >= 0);
+		SerialModificationNumber++;
+
+		// Because we only remove trailing chunks to avoid messing up the absolute indices in the entities map,
+		// We are freeing the memory here to save memory
+		if (NumInstances == 0)
+		{
+			FMemory::Free(RawMemory);
+			RawMemory = nullptr;
+		}
+	}
+
+	void AddInstance()
+	{
+		AddMultipleInstances(1);
+	}
+
+	void RemoveInstance()
+	{
+		RemoveMultipleInstances(1);
+	}
+
+	int32 GetSerialModificationNumber() const
+	{
+		return SerialModificationNumber;
+	}
+
+	FStructView GetMutableChunkFragmentViewChecked(const int32 Index) { return FStructView(ChunkFragmentData[Index]); }
+
+	FInstancedStruct* FindMutableChunkFragment(const UScriptStruct* Type)
+	{
+		return ChunkFragmentData.FindByPredicate([Type](const FInstancedStruct& Element)
+			{
+				return Element.GetScriptStruct()->IsChildOf(Type);
+			});
+	}
+
+	void Recycle(TConstArrayView<FInstancedStruct> InChunkFragmentsTemplate, const FMassArchetypeSharedFragmentValues& InSharedFragmentValues)
+	{
+		checkf(NumInstances == 0, TEXT("Recycling a chunk that is not empty."));
+		SerialModificationNumber++;
+		ChunkFragmentData = InChunkFragmentsTemplate;
+		SharedFragmentValues = InSharedFragmentValues;
+		
+		// If this chunk previously had entity and it does not anymore, we might have to reallocate the memory as it was freed to save memory
+		if (RawMemory == nullptr)
+		{
+			RawMemory = (uint8*)FMemory::Malloc(AllocSize);
+		}
+	}
+
+	bool IsValidSubChunk(const int32 StartIndex, const int32 Length) const
+	{
+		return StartIndex >= 0 && StartIndex < NumInstances && (StartIndex + Length) <= NumInstances;
+	}
+
+#if WITH_MASSENTITY_DEBUG
+	int32 DebugGetChunkFragmentCount() const { return ChunkFragmentData.Num(); }
+#endif // WITH_MASSENTITY_DEBUG
+
+	FMassArchetypeSharedFragmentValues& GetMutableSharedFragmentValues() { return SharedFragmentValues; }
+	const FMassArchetypeSharedFragmentValues& GetSharedFragmentValues() const { return SharedFragmentValues; }
+};
+
+// Information for a single fragment type in an archetype
+struct FMassArchetypeFragmentConfig
+{
+	const UScriptStruct* FragmentType = nullptr;
+	int32 ArrayOffsetWithinChunk = 0;
+
+	void* GetFragmentData(uint8* ChunkBase, int32 IndexWithinChunk) const
+	{
+		return ChunkBase + ArrayOffsetWithinChunk + (IndexWithinChunk * FragmentType->GetStructureSize());
+	}
+};
+
+// An archetype is defined by a collection of unique fragment types (no duplicates).
+// Order doesn't matter, there will only ever be one FMassArchetypeData per unique set of fragment types per entity manager subsystem
+struct FMassArchetypeData
+{
+private:
+	// One-stop-shop variable describing the archetype's fragment and tag composition 
+	FMassArchetypeCompositionDescriptor CompositionDescriptor;
+
+	// Pre-created default chunk fragment templates
+	TArray<FInstancedStruct> ChunkFragmentsTemplate;
+
+	TArray<FMassArchetypeFragmentConfig, TInlineAllocator<16>> FragmentConfigs;
+	
+	TArray<FMassArchetypeChunk> Chunks;
+
+	// Entity ID to index within archetype
+	//@TODO: Could be folded into FEntityData in the entity manager at the expense of a bit
+	// of loss of encapsulation and extra complexity during archetype changes
+	TMap<int32, int32> EntityMap;
+	
+	TMap<const UScriptStruct*, int32> FragmentIndexMap;
+
+	int32 NumEntitiesPerChunk;
+	SIZE_T TotalBytesPerEntity = 0;
+	int32 EntityListOffsetWithinChunk;
+
+	// Archetype version at which this archetype was created, useful for query to do incremental archetype matching
+	uint32 CreatedArchetypeDataVersion = 0;
+
+#if WITH_MASSENTITY_DEBUG
+	// Arrays of names the archetype is referred as.
+	TArray<FName> DebugNames;
+#endif // WITH_MASSENTITY_DEBUG
+
+	// Defaults to UMassEntitySettings.ChunkMemorySize. In near future will support being set via constructor.
+	const SIZE_T ChunkMemorySize = 0;
+	
+	friend FMassEntityQuery;
+	friend FMassArchetypeEntityCollection;
+	friend FMassDebugger;
+
+public:
+	explicit FMassArchetypeData(const FMassArchetypeCreationParams& CreationParams = FMassArchetypeCreationParams());
+
+	TConstArrayView<FMassArchetypeFragmentConfig> GetFragmentConfigs() const { return FragmentConfigs; }
+	const FMassFragmentBitSet& GetFragmentBitSet() const { return CompositionDescriptor.Fragments; }
+	const FMassTagBitSet& GetTagBitSet() const { return CompositionDescriptor.Tags; }
+	const FMassChunkFragmentBitSet& GetChunkFragmentBitSet() const { return CompositionDescriptor.ChunkFragments; }
+	const FMassSharedFragmentBitSet& GetSharedFragmentBitSet() const { return CompositionDescriptor.SharedFragments; }
+	const FMassConstSharedFragmentBitSet& GetConstSharedFragmentBitSet() const { return CompositionDescriptor.ConstSharedFragments; }
+
+	const FMassArchetypeCompositionDescriptor& GetCompositionDescriptor() const { return CompositionDescriptor; }
+	FORCEINLINE const FMassArchetypeSharedFragmentValues& GetSharedFragmentValues(int32 EntityIndex) const
+	{ 
+		const int32 AbsoluteIndex = EntityMap.FindChecked(EntityIndex);
+		const int32 ChunkIndex = AbsoluteIndex / NumEntitiesPerChunk;
+
+		return Chunks[ChunkIndex].GetSharedFragmentValues();
+	}
+	FORCEINLINE const FMassArchetypeSharedFragmentValues& GetSharedFragmentValues(FMassEntityHandle Entity) const
+	{
+		return GetSharedFragmentValues(Entity.Index);
+	}
+
+	/** Method to iterate on all the fragment types */
+	void ForEachFragmentType(TFunction< void(const UScriptStruct* /*FragmentType*/)> Function) const;
+	bool HasFragmentType(const UScriptStruct* FragmentType) const;
+	bool HasTagType(const UScriptStruct* FragmentType) const { check(FragmentType); return CompositionDescriptor.Tags.Contains(*FragmentType); }
+
+	bool IsEquivalent(const FMassArchetypeCompositionDescriptor& OtherCompositionDescriptor) const
+	{
+		return CompositionDescriptor.IsEquivalent(OtherCompositionDescriptor);
+	}
+
+	void Initialize(const FMassArchetypeCompositionDescriptor& InCompositionDescriptor, const uint32 ArchetypeDataVersion);
+
+	/** 
+	 * A special way of initializing an archetype resulting in a copy of BaseArchetype's setup with OverrideTags
+	 * replacing original tags of BaseArchetype
+	 */
+	void InitializeWithSimilar(const FMassArchetypeData& BaseArchetype, FMassArchetypeCompositionDescriptor&& NewComposition, const uint32 ArchetypeDataVersion);
+
+	void AddEntity(FMassEntityHandle Entity, const FMassArchetypeSharedFragmentValues& InSharedFragmentValues);
+	void RemoveEntity(FMassEntityHandle Entity);
+
+	bool HasFragmentDataForEntity(const UScriptStruct* FragmentType, int32 EntityIndex) const;
+	void* GetFragmentDataForEntityChecked(const UScriptStruct* FragmentType, int32 EntityIndex) const;
+	void* GetFragmentDataForEntity(const UScriptStruct* FragmentType, int32 EntityIndex) const;
+
+	FORCEINLINE const int32* GetInternalIndexForEntity(const int32 EntityIndex) const { return EntityMap.Find(EntityIndex); }
+	FORCEINLINE int32 GetInternalIndexForEntityChecked(const int32 EntityIndex) const { return EntityMap.FindChecked(EntityIndex); }
+	int32 GetNumEntitiesPerChunk() const { return NumEntitiesPerChunk; }
+	SIZE_T GetBytesPerEntity() const { return TotalBytesPerEntity; }
+
+	int32 GetNumEntities() const { return EntityMap.Num(); }
+
+	SIZE_T GetChunkAllocSize() const { return ChunkMemorySize; }
+
+	int32 GetChunkCount() const { return Chunks.Num(); }
+	int32 GetNonEmptyChunkCount() const;
+
+	uint32 GetCreatedArchetypeDataVersion() const { return CreatedArchetypeDataVersion; }
+
+	void ExecuteFunction(FMassExecutionContext& RunContext, const FMassExecuteFunction& Function, const FMassQueryRequirementIndicesMapping& RequirementMapping
+		, FMassArchetypeEntityCollection::FConstEntityRangeArrayView EntityRangeContainer, const FMassChunkConditionFunction& ChunkCondition);
+	void ExecuteFunction(FMassExecutionContext& RunContext, const FMassExecuteFunction& Function, const FMassQueryRequirementIndicesMapping& RequirementMapping
+		, const FMassChunkConditionFunction& ChunkCondition);
+
+	void ExecutionFunctionForChunk(FMassExecutionContext& RunContext, const FMassExecuteFunction& Function, const FMassQueryRequirementIndicesMapping& RequirementMapping
+		, const FMassArchetypeEntityCollection::FArchetypeEntityRange& EntityRange, const FMassChunkConditionFunction& ChunkCondition = FMassChunkConditionFunction());
+
+	/**
+	 * Compacts entities to fill up chunks as much as possible
+	 * @return number of entities moved around
+	 */
+	int32 CompactEntities(const double TimeAllowed);
+
+	/**
+	 * Moves the entity from this archetype to another, will only copy all matching fragment types
+	 * @param Entity is the entity to move
+	 * @param NewArchetype the archetype to move to
+	 * @param SharedFragmentValuesOverride if provided will override all given Entity's shared fragment values
+	 */
+	void MoveEntityToAnotherArchetype(const FMassEntityHandle Entity, FMassArchetypeData& NewArchetype, const FMassArchetypeSharedFragmentValues* SharedFragmentValuesOverride = nullptr);
+
+	/**
+	 * Set all fragment sources data on specified entity, will check if there are fragment sources type that does not exist in the archetype
+	 * @param Entity is the entity to set the data of all fragments
+	 * @param FragmentSources are the fragments to copy the data from
+	 */
+	void SetFragmentsData(const FMassEntityHandle Entity, TArrayView<const FInstancedStruct> FragmentSources);
+
+	/** For all entities indicated by EntityCollection the function sets the value of fragment of type
+	 *  FragmentSource.GetScriptStruct to the value represented by FragmentSource.GetMemory */
+	void SetFragmentData(FMassArchetypeEntityCollection::FConstEntityRangeArrayView EntityRangeContainer, const FInstancedStruct& FragmentSource);
+
+	/** Returns conversion from given Requirements to archetype's fragment indices */
+	void GetRequirementsFragmentMapping(TConstArrayView<FMassFragmentRequirementDescription> Requirements, FMassFragmentIndicesMapping& OutFragmentIndices) const;
+
+	/** Returns conversion from given ChunkRequirements to archetype's chunk fragment indices */
+	void GetRequirementsChunkFragmentMapping(TConstArrayView<FMassFragmentRequirementDescription> ChunkRequirements, FMassFragmentIndicesMapping& OutFragmentIndices) const;
+
+	/** Returns conversion from given const shared requirements to archetype's const shared fragment indices */
+	void GetRequirementsConstSharedFragmentMapping(TConstArrayView<FMassFragmentRequirementDescription> Requirements, FMassFragmentIndicesMapping& OutFragmentIndices) const;
+
+	/** Returns conversion from given shared requirements to archetype's shared fragment indices */
+	void GetRequirementsSharedFragmentMapping(TConstArrayView<FMassFragmentRequirementDescription> Requirements, FMassFragmentIndicesMapping& OutFragmentIndices) const;
+
+	SIZE_T GetAllocatedSize() const;
+
+	// Converts the list of fragments into a user-readable debug string
+	FString DebugGetDescription() const;
+
+	/** Copies debug names from another archetype data. */
+	void CopyDebugNamesFrom(const FMassArchetypeData& Other)
+	{ 
+#if WITH_MASSENTITY_DEBUG
+		DebugNames = Other.DebugNames; 
+#endif // WITH_MASSENTITY_DEBUG
+	}
+
+#if WITH_MASSENTITY_DEBUG
+	/** Fetches how much memory is allocated for active chunks, and how much of that memory is actually occupied */
+	void DebugGetEntityMemoryNumbers(SIZE_T& OutActiveChunksMemorySize, SIZE_T& OutActiveEntitiesMemorySize) const;
+
+	/** Adds new debug name associated with the archetype. */
+	void AddUniqueDebugName(const FName& Name) { DebugNames.AddUnique(Name); }
+	
+	/** @return array of debug names associated with this archetype. */
+	const TConstArrayView<FName> GetDebugNames() const { return DebugNames; }
+	
+	/** @return string of all debug names combined */
+	FString GetCombinedDebugNamesAsString() const;
+
+	/**
+	 * Prints out debug information about the archetype
+	 */
+	void DebugPrintArchetype(FOutputDevice& Ar);
+
+	/**
+	 * Prints out fragment's values for the specified entity. 
+	 * @param Entity The entity for which we want to print fragment values
+	 * @param Ar The output device
+	 * @param InPrefix Optional prefix to remove from fragment names
+	 */
+	void DebugPrintEntity(FMassEntityHandle Entity, FOutputDevice& Ar, const TCHAR* InPrefix = TEXT("")) const;
+#endif // WITH_MASSENTITY_DEBUG
+
+	void REMOVEME_GetArrayViewForFragmentInChunk(int32 ChunkIndex, const UScriptStruct* FragmentType, void*& OutChunkBase, int32& OutNumEntities);
+
+	//////////////////////////////////////////////////////////////////////
+	// low level api
+	FORCEINLINE const int32* GetFragmentIndex(const UScriptStruct* FragmentType) const { return FragmentIndexMap.Find(FragmentType); }
+	FORCEINLINE int32 GetFragmentIndexChecked(const UScriptStruct* FragmentType) const { return FragmentIndexMap.FindChecked(FragmentType); }
+
+	FORCEINLINE void* GetFragmentData(const int32 FragmentIndex, const FMassRawEntityInChunkData EntityIndex) const
+	{
+		return FragmentConfigs[FragmentIndex].GetFragmentData(EntityIndex.ChunkRawMemory, EntityIndex.IndexWithinChunk);
+	}
+
+	FORCEINLINE FMassRawEntityInChunkData MakeEntityHandle(int32 EntityIndex) const
+	{
+		const int32 AbsoluteIndex = EntityMap.FindChecked(EntityIndex);
+		const int32 ChunkIndex = AbsoluteIndex / NumEntitiesPerChunk;
+	
+		return FMassRawEntityInChunkData(Chunks[ChunkIndex].GetRawMemory(), AbsoluteIndex % NumEntitiesPerChunk); 
+	}
+
+	FORCEINLINE FMassRawEntityInChunkData MakeEntityHandle(FMassEntityHandle Entity) const
+	{
+		return MakeEntityHandle(Entity.Index); 
+	}
+
+	bool IsInitialized() const { return TotalBytesPerEntity > 0 && FragmentConfigs.IsEmpty() == false; }
+
+	//////////////////////////////////////////////////////////////////////
+	// batched api
+	void BatchDestroyEntityChunks(FMassArchetypeEntityCollection::FConstEntityRangeArrayView EntityRangeContainer, TArray<FMassEntityHandle>& OutEntitiesRemoved);
+	void BatchAddEntities(TConstArrayView<FMassEntityHandle> Entities, const FMassArchetypeSharedFragmentValues& SharedFragmentValues
+		, TArray<FMassArchetypeEntityCollection::FArchetypeEntityRange>& OutNewRanges);
+	/** 
+	 * @param SharedFragmentValuesOverride if provided will override shared fragment values for the entities being moved
+	 */
+	void BatchMoveEntitiesToAnotherArchetype(const FMassArchetypeEntityCollection& EntityCollection, FMassArchetypeData& NewArchetype
+		, TArray<FMassEntityHandle>& OutEntitiesBeingMoved, TArray<FMassArchetypeEntityCollection::FArchetypeEntityRange>* OutNewChunks = nullptr
+		, const FMassArchetypeSharedFragmentValues* SharedFragmentValuesToAdd = nullptr
+		, const FMassSharedFragmentBitSet* SharedFragmentToRemoveBitSet = nullptr
+		, const FMassConstSharedFragmentBitSet* ConstSharedFragmentToRemoveBitSet = nullptr);
+	void BatchSetFragmentValues(TConstArrayView<FMassArchetypeEntityCollection::FArchetypeEntityRange> EntityCollection, const FMassGenericPayloadViewSlice& Payload);
+
+protected:
+	FMassArchetypeEntityCollection::FArchetypeEntityRange PrepareNextEntitiesSpanInternal(TConstArrayView<FMassEntityHandle> Entities, const FMassArchetypeSharedFragmentValues& InSharedFragmentValues, const int32 StartingChunk = 0);
+	void BatchRemoveEntitiesInternal(const int32 ChunkIndex, const int32 StartIndexWithinChunk, const int32 NumberToRemove);
+
+	struct FTransientChunkLocation
+	{
+		uint8* RawChunkMemory;
+		int32 IndexWithinChunk;
+	};
+	void MoveFragmentsToAnotherArchetypeInternal(FMassArchetypeData& TargetArchetype, FTransientChunkLocation Target, const FTransientChunkLocation Source, const int32 ElementsNum);
+	void MoveFragmentsToNewLocationInternal(FTransientChunkLocation Target, const FTransientChunkLocation Source, const int32 NumberToMove);
+	void ConfigureFragments();
+
+	FORCEINLINE void* GetFragmentData(const int32 FragmentIndex, uint8* ChunkRawMemory, const int32 IndexWithinChunk) const
+	{
+		return FragmentConfigs[FragmentIndex].GetFragmentData(ChunkRawMemory, IndexWithinChunk);
+	}
+
+	void BindEntityRequirements(FMassExecutionContext& RunContext, const FMassFragmentIndicesMapping& EntityFragmentsMapping, FMassArchetypeChunk& Chunk, const int32 SubchunkStart, const int32 SubchunkLength);
+	void BindChunkFragmentRequirements(FMassExecutionContext& RunContext, const FMassFragmentIndicesMapping& ChunkFragmentsMapping, FMassArchetypeChunk& Chunk);
+	void BindConstSharedFragmentRequirements(FMassExecutionContext& RunContext, const FMassArchetypeSharedFragmentValues& SharedFragmentValues, const FMassFragmentIndicesMapping& ChunkFragmentsMapping);
+	void BindSharedFragmentRequirements(FMassExecutionContext& RunContext, FMassArchetypeSharedFragmentValues& SharedFragmentValues, const FMassFragmentIndicesMapping& ChunkFragmentsMapping);
+
+private:
+	int32 AddEntityInternal(FMassEntityHandle Entity, const FMassArchetypeSharedFragmentValues& InSharedFragmentValues);
+	void RemoveEntityInternal(const int32 AbsoluteIndex);
+};
+
+
+struct FMassArchetypeHelper
+{
+	FORCEINLINE static FMassArchetypeData* ArchetypeDataFromHandle(const FMassArchetypeHandle& ArchetypeHandle) { return ArchetypeHandle.DataPtr.Get(); }
+	FORCEINLINE static FMassArchetypeData& ArchetypeDataFromHandleChecked(const FMassArchetypeHandle& ArchetypeHandle)
+	{
+		check(ArchetypeHandle.IsValid());
+		return *ArchetypeHandle.DataPtr.Get();
+	}
+	FORCEINLINE static FMassArchetypeHandle ArchetypeHandleFromData(const TSharedPtr<FMassArchetypeData>& Archetype)
+	{
+		return FMassArchetypeHandle(Archetype);
+	}
+};
diff --git a/Engine/Plugins/Runtime/MassEntity/Source/MassEntity/Private/MassArchetypeData.h b/Engine/Plugins/Runtime/MassEntity/Source/MassEntity/Private/MassArchetypeData.h
deleted file mode 100644
diff --git a/Engine/Plugins/Runtime/MassEntity/Source/MassEntityTestSuite/MassEntityTestSuite.Build.cs b/Engine/Plugins/Runtime/MassEntity/Source/MassEntityTestSuite/MassEntityTestSuite.Build.cs
--- a/Engine/Plugins/Runtime/MassEntity/Source/MassEntityTestSuite/MassEntityTestSuite.Build.cs
+++ b/Engine/Plugins/Runtime/MassEntity/Source/MassEntityTestSuite/MassEntityTestSuite.Build.cs
@@ -13,8 +13,6 @@ namespace UnrealBuildTool.Rules
 
 			UnsafeTypeCastWarningLevel = WarningLevel.Warning;
 
-			PrivateIncludePaths.Add(Path.Combine(GetModuleDirectory("MassEntity"), "Private")); 
-
 			PublicDependencyModuleNames.AddRange(
 				new string[] {
 					"Core",
