New Mass debug capability - determinig why a given entity has not been considered by EntityQueries being run. The information is displayed in Processor View of the MassDebugger

diff --git a/Engine/Plugins/Runtime/MassEntity/Source/MassEntity/Internal/MassArchetypeData.h b/Engine/Plugins/Runtime/MassEntity/Source/MassEntity/Internal/MassArchetypeData.h
--- a/Engine/Plugins/Runtime/MassEntity/Source/MassEntity/Internal/MassArchetypeData.h
+++ b/Engine/Plugins/Runtime/MassEntity/Source/MassEntity/Internal/MassArchetypeData.h
@@ -418,4 +418,16 @@ struct FMassArchetypeHelper
 	{
 		return FMassArchetypeHandle(Archetype);
 	}
+
+	/**
+	 * Determines whether given Archetype matches given Requirements. In case of failure to match and if WITH_MASSENTITY_DEBUG
+	 * the function will also log the reasons for said failure (at VeryVerbose level).
+	 * @param bBailOutOnFirstFail if true will skip the remaining tests as soon as a single mismatch is detected. This option
+	 *	is used when looking for matching archetypes. For debugging purposes use `false` to list all the mismatching elements.
+	 */
+	MASSENTITY_API static bool DoesArchetypeMatchRequirements(const FMassArchetypeData& Archetype, const FMassFragmentRequirements& Requirements
+#if WITH_MASSENTITY_DEBUG
+		, const bool bBailOutOnFirstFail = true, FOutputDevice* OutputDevice = nullptr
+#endif // WITH_MASSENTITY_DEBUG
+		);
 };
diff --git a/Engine/Plugins/Runtime/MassEntity/Source/MassEntity/Private/MassArchetypeData.cpp b/Engine/Plugins/Runtime/MassEntity/Source/MassEntity/Private/MassArchetypeData.cpp
--- a/Engine/Plugins/Runtime/MassEntity/Source/MassEntity/Private/MassArchetypeData.cpp
+++ b/Engine/Plugins/Runtime/MassEntity/Source/MassEntity/Private/MassArchetypeData.cpp
@@ -1318,3 +1318,232 @@ void FMassArchetypeData::BatchSetFragmentValues(TConstArrayView<FMassArchetypeEn
 		EntitiesHandled += EntityRange.Length;
 	}
 }
+
+//-----------------------------------------------------------------------------
+// FMassArchetypeHelper
+//-----------------------------------------------------------------------------
+bool FMassArchetypeHelper::DoesArchetypeMatchRequirements(const FMassArchetypeData& Archetype, const FMassFragmentRequirements& Requirements
+#if WITH_MASSENTITY_DEBUG
+	, const bool bBailOutOnFirstFail, FOutputDevice* OutputDevice
+#endif // WITH_MASSENTITY_DEBUG
+	)
+{
+	bool bResult = true;
+	if (Archetype.GetTagBitSet().HasAll(Requirements.GetRequiredAllTags()) == false)
+	{
+		// missing some required tags, skip.
+#if WITH_MASSENTITY_DEBUG
+		if (OutputDevice)
+		{
+			const FMassTagBitSet UnsatisfiedTags = Requirements.GetRequiredAllTags() - Archetype.GetTagBitSet();
+			FStringOutputDevice Description;
+			UnsatisfiedTags.DebugGetStringDesc(Description);
+			OutputDevice->Logf(TEXT("Archetype did not match due to missing tags: %s"), *Description);
+		}
+		if (!bBailOutOnFirstFail) bResult = false; else
+#endif // WITH_MASSENTITY_DEBUG
+		{
+			return false;
+		}
+	}
+
+	if (Archetype.GetTagBitSet().HasNone(Requirements.GetRequiredNoneTags()) == false)
+	{
+		// has some tags required to be absent
+#if WITH_MASSENTITY_DEBUG
+		if (OutputDevice)
+		{
+			const FMassTagBitSet UnwantedTags = Requirements.GetRequiredAllTags().GetOverlap(Archetype.GetTagBitSet());
+			FStringOutputDevice Description;
+			UnwantedTags.DebugGetStringDesc(Description);
+			OutputDevice->Logf(TEXT("Archetype has tags required absent: %s"), *Description);
+		}
+		if (!bBailOutOnFirstFail) bResult = false; else
+#endif // WITH_MASSENTITY_DEBUG
+		{
+			return false;
+		}
+	}
+
+	if (Requirements.GetRequiredAnyTags().IsEmpty() == false
+		&& Archetype.GetTagBitSet().HasAny(Requirements.GetRequiredAnyTags()) == false)
+	{
+#if WITH_MASSENTITY_DEBUG
+		if (OutputDevice)
+		{
+			FStringOutputDevice Description;
+			Requirements.GetRequiredAnyTags().DebugGetStringDesc(Description);
+			OutputDevice->Logf(TEXT("Archetype did not match due to missing \'any\' tags: %s"), *Description);
+		}
+		if (!bBailOutOnFirstFail) bResult = false; else
+#endif // WITH_MASSENTITY_DEBUG
+		{
+			return false;
+		}
+	}
+
+	if (Archetype.GetFragmentBitSet().HasAll(Requirements.GetRequiredAllFragments()) == false)
+	{
+		// missing some required fragments, skip.
+#if WITH_MASSENTITY_DEBUG
+		if (OutputDevice)
+		{
+			const FMassFragmentBitSet UnsatisfiedFragments = Requirements.GetRequiredAllFragments() - Archetype.GetFragmentBitSet();
+			FStringOutputDevice Description;
+			UnsatisfiedFragments.DebugGetStringDesc(Description);
+			OutputDevice->Logf(TEXT("Archetype did not match due to missing Fragments: %s"), *Description);
+		}
+		if (!bBailOutOnFirstFail) bResult = false; else
+#endif // WITH_MASSENTITY_DEBUG
+		{
+			return false;
+		}
+	}
+
+	if (Archetype.GetFragmentBitSet().HasNone(Requirements.GetRequiredNoneFragments()) == false)
+	{
+		// has some Fragments required to be absent
+#if WITH_MASSENTITY_DEBUG
+		if (OutputDevice)
+		{
+			const FMassFragmentBitSet UnwantedFragments = Requirements.GetRequiredAllFragments().GetOverlap(Archetype.GetFragmentBitSet());
+			FStringOutputDevice Description;
+			UnwantedFragments.DebugGetStringDesc(Description);
+			OutputDevice->Logf(TEXT("Archetype has Fragments required absent: %s"), *Description);
+		}
+		if (!bBailOutOnFirstFail) bResult = false; else
+#endif // WITH_MASSENTITY_DEBUG
+		{
+			return false;
+		}
+	}
+
+	if (Requirements.GetRequiredAnyFragments().IsEmpty() == false
+		&& Archetype.GetFragmentBitSet().HasAny(Requirements.GetRequiredAnyFragments()) == false)
+	{
+#if WITH_MASSENTITY_DEBUG
+		if (OutputDevice)
+		{
+			FStringOutputDevice Description;
+			Requirements.GetRequiredAnyFragments().DebugGetStringDesc(Description);
+			OutputDevice->Logf(TEXT("Archetype did not match due to missing \'any\' fragments: %s"), *Description);
+		}
+		if (!bBailOutOnFirstFail) bResult = false; else
+#endif // WITH_MASSENTITY_DEBUG
+		{
+			return false;
+		}
+	}
+
+	if (Archetype.GetChunkFragmentBitSet().HasAll(Requirements.GetRequiredAllChunkFragments()) == false)
+	{
+		// missing some required fragments, skip.
+#if WITH_MASSENTITY_DEBUG
+		if (OutputDevice)
+		{
+			const FMassChunkFragmentBitSet UnsatisfiedFragments = Requirements.GetRequiredAllChunkFragments() - Archetype.GetChunkFragmentBitSet();
+			FStringOutputDevice Description;
+			UnsatisfiedFragments.DebugGetStringDesc(Description);
+			OutputDevice->Logf(TEXT("Archetype did not match due to missing Chunk Fragments: %s"), *Description);
+		}
+		if (!bBailOutOnFirstFail) bResult = false; else
+#endif // WITH_MASSENTITY_DEBUG
+		{
+			return false;
+		}
+	}
+
+	if (Archetype.GetChunkFragmentBitSet().HasNone(Requirements.GetRequiredNoneChunkFragments()) == false)
+	{
+		// has some Fragments required to be absent
+#if WITH_MASSENTITY_DEBUG
+		if (OutputDevice)
+		{
+			const FMassChunkFragmentBitSet UnwantedFragments = Requirements.GetRequiredNoneChunkFragments().GetOverlap(Archetype.GetChunkFragmentBitSet());
+			FStringOutputDevice Description;
+			UnwantedFragments.DebugGetStringDesc(Description);
+			OutputDevice->Logf(TEXT("Archetype has Chunk Fragments required absent: %s"), *Description);
+		}
+		if (!bBailOutOnFirstFail) bResult = false; else
+#endif // WITH_MASSENTITY_DEBUG
+		{
+			return false;
+		}
+	}
+
+	if (Archetype.GetSharedFragmentBitSet().HasAll(Requirements.GetRequiredAllSharedFragments()) == false)
+	{
+		// missing some required fragments, skip.
+#if WITH_MASSENTITY_DEBUG
+		if (OutputDevice)
+		{
+			const FMassSharedFragmentBitSet UnsatisfiedFragments = Requirements.GetRequiredAllSharedFragments() - Archetype.GetSharedFragmentBitSet();
+			FStringOutputDevice Description;
+			UnsatisfiedFragments.DebugGetStringDesc(Description);
+			OutputDevice->Logf(TEXT("Archetype did not match due to missing Shared Fragments: %s"), *Description);
+		}
+		if (!bBailOutOnFirstFail) bResult = false; else
+#endif // WITH_MASSENTITY_DEBUG
+		{
+			return false;
+		}
+	}
+
+	if (Archetype.GetSharedFragmentBitSet().HasNone(Requirements.GetRequiredNoneSharedFragments()) == false)
+	{
+		// has some Fragments required to be absent
+#if WITH_MASSENTITY_DEBUG
+		if (OutputDevice)
+		{
+			const FMassSharedFragmentBitSet UnwantedFragments = Requirements.GetRequiredNoneSharedFragments().GetOverlap(Archetype.GetSharedFragmentBitSet());
+			FStringOutputDevice Description;
+			UnwantedFragments.DebugGetStringDesc(Description);
+			OutputDevice->Logf(TEXT("Archetype has Shared Fragments required absent: %s"), *Description);
+		}
+		if (!bBailOutOnFirstFail) bResult = false; else
+#endif // WITH_MASSENTITY_DEBUG
+		{
+			return false;
+		}
+	}
+
+	if (Archetype.GetConstSharedFragmentBitSet().HasAll(Requirements.GetRequiredAllConstSharedFragments()) == false)
+	{
+		// missing some required fragments, skip.
+#if WITH_MASSENTITY_DEBUG
+		if (OutputDevice)
+		{
+			const FMassConstSharedFragmentBitSet UnsatisfiedFragments = Requirements.GetRequiredAllConstSharedFragments() - Archetype.GetConstSharedFragmentBitSet();
+			FStringOutputDevice Description;
+			UnsatisfiedFragments.DebugGetStringDesc(Description);
+			OutputDevice->Logf(TEXT("Archetype did not match due to missing Const Shared Fragments: %s"), *Description);
+		}
+		if (!bBailOutOnFirstFail) bResult = false; else
+#endif // WITH_MASSENTITY_DEBUG
+		{
+			return false;
+		}
+	}
+
+	if (Archetype.GetConstSharedFragmentBitSet().HasNone(Requirements.GetRequiredNoneConstSharedFragments()) == false)
+	{
+		// has some Fragments required to be absent
+#if WITH_MASSENTITY_DEBUG
+		if (OutputDevice)
+		{
+			const FMassConstSharedFragmentBitSet UnwantedFragments = Requirements.GetRequiredNoneConstSharedFragments().GetOverlap(Archetype.GetConstSharedFragmentBitSet());
+			FStringOutputDevice Description;
+			UnwantedFragments.DebugGetStringDesc(Description);
+			OutputDevice->Logf(TEXT("Archetype has Const Shared Fragments required absent: %s"), *Description);
+		}
+		// could skip the test as it's the final check we're performing, but leaving it in in case more cases are added
+		// later without checking the existing code
+		if (!bBailOutOnFirstFail) bResult = false; else
+#endif // WITH_MASSENTITY_DEBUG
+		{
+			return false;
+		}
+	}
+
+	return bResult;
+}
diff --git a/Engine/Plugins/Runtime/MassEntity/Source/MassEntity/Private/MassDebugger.cpp b/Engine/Plugins/Runtime/MassEntity/Source/MassEntity/Private/MassDebugger.cpp
--- a/Engine/Plugins/Runtime/MassEntity/Source/MassEntity/Private/MassDebugger.cpp
+++ b/Engine/Plugins/Runtime/MassEntity/Source/MassEntity/Private/MassDebugger.cpp
@@ -22,13 +22,18 @@ namespace UE::Mass::Debug
 {
 	bool bAllowProceduralDebuggedEntitySelection = false;
 	bool bAllowBreakOnDebuggedEntity = false;
+	bool bTestSelectedEntityAgainstProcessorQueries = true;
 
 	FAutoConsoleVariableRef CVars[] =
 	{
-		FAutoConsoleVariableRef(TEXT("mass.debug.AllowProceduralDebuggedEntitySelection"), bAllowProceduralDebuggedEntitySelection
-			, TEXT("Guards whether MASS_SET_ENTITY_DEBUGGED calls take effect."), ECVF_Cheat)
-		, FAutoConsoleVariableRef(TEXT("mass.debug.AllowBreakOnDebuggedEntity"), bAllowBreakOnDebuggedEntity
-			, TEXT("Guards whether MASS_BREAK_IF_ENTITY_DEBUGGED calls take effect."), ECVF_Cheat)
+		{ TEXT("mass.debug.AllowProceduralDebuggedEntitySelection"), bAllowProceduralDebuggedEntitySelection
+			, TEXT("Guards whether MASS_SET_ENTITY_DEBUGGED calls take effect."), ECVF_Cheat}
+		, {TEXT("mass.debug.AllowBreakOnDebuggedEntity"), bAllowBreakOnDebuggedEntity
+			, TEXT("Guards whether MASS_BREAK_IF_ENTITY_DEBUGGED calls take effect."), ECVF_Cheat}
+		, {	TEXT("mass.debug.TestSelectedEntityAgainstProcessorQueries"), bTestSelectedEntityAgainstProcessorQueries
+			, TEXT("Enabling will result in testing all processors' queries against SelectedEntity (as indicated by")
+			TEXT("mass.debug.DebugEntity or the gameplay debugger) and storing potential failure results to be viewed in MassDebugger")
+			, ECVF_Cheat }
 	};
 	
 
@@ -660,4 +665,13 @@ void FMassDebugger::UnregisterEntityManager(FMassEntityManager& EntityManager)
 	OnEntityManagerDeinitialized.Broadcast(EntityManager);
 }
 
+bool FMassDebugger::DoesArchetypeMatchRequirements(const FMassArchetypeHandle& ArchetypeHandle, const FMassFragmentRequirements& Requirements, FOutputDevice& OutputDevice)
+{
+	if (const FMassArchetypeData* Archetype = FMassArchetypeHelper::ArchetypeDataFromHandle(ArchetypeHandle))
+	{
+		return FMassArchetypeHelper::DoesArchetypeMatchRequirements(*Archetype, Requirements, /*bBailOutOnFirstFail=*/false, &OutputDevice);
+	}
+	return false;
+}
+
 #endif // WITH_MASSENTITY_DEBUG
\ No newline at end of file
diff --git a/Engine/Plugins/Runtime/MassEntity/Source/MassEntity/Private/MassEntityManager.cpp b/Engine/Plugins/Runtime/MassEntity/Source/MassEntity/Private/MassEntityManager.cpp
--- a/Engine/Plugins/Runtime/MassEntity/Source/MassEntity/Private/MassEntityManager.cpp
+++ b/Engine/Plugins/Runtime/MassEntity/Source/MassEntity/Private/MassEntityManager.cpp
@@ -1612,6 +1612,11 @@ void FMassEntityManager::GetMatchingArchetypes(const FMassFragmentRequirements&
 		}
 	}
 
+#if WITH_MASSENTITY_DEBUG
+	FStringOutputDevice LogOutput;
+	LogOutput.SetAutoEmitLineTerminator(true);
+#endif // WITH_MASSENTITY_DEBUG
+
 	// Then verify that they contain *all* required fragments
 	for (TSharedPtr<FMassArchetypeData>& ArchetypePtr : AnyArchetypes)
 	{
@@ -1623,161 +1628,21 @@ void FMassEntityManager::GetMatchingArchetypes(const FMassFragmentRequirements&
 			continue;
 		}
 
-		if (Archetype.GetTagBitSet().HasAll(Requirements.GetRequiredAllTags()) == false)
-		{
-			// missing some required tags, skip.
-#if WITH_MASSENTITY_DEBUG
-			const FMassTagBitSet UnsatisfiedTags = Requirements.GetRequiredAllTags() - Archetype.GetTagBitSet();
-			FStringOutputDevice Description;
-			UnsatisfiedTags.DebugGetStringDesc(Description);
-			UE_LOG(LogMass, VeryVerbose, TEXT("Archetype did not match due to missing tags: %s")
-				, *Description);
-#endif // WITH_MASSENTITY_DEBUG
-			continue;
-		}
-
-		if (Archetype.GetTagBitSet().HasNone(Requirements.GetRequiredNoneTags()) == false)
-		{
-			// has some tags required to be absent
-#if WITH_MASSENTITY_DEBUG
-			const FMassTagBitSet UnwantedTags = Requirements.GetRequiredAllTags().GetOverlap(Archetype.GetTagBitSet());
-			FStringOutputDevice Description;
-			UnwantedTags.DebugGetStringDesc(Description);
-			UE_LOG(LogMass, VeryVerbose, TEXT("Archetype has tags required absent: %s")
-				, *Description);
-#endif // WITH_MASSENTITY_DEBUG
-			continue;
-		}
-
-		if (Requirements.GetRequiredAnyTags().IsEmpty() == false 
-			&& Archetype.GetTagBitSet().HasAny(Requirements.GetRequiredAnyTags()) == false)
-		{
-#if WITH_MASSENTITY_DEBUG
-			FStringOutputDevice Description;
-			Requirements.GetRequiredAnyTags().DebugGetStringDesc(Description);
-			UE_LOG(LogMass, VeryVerbose, TEXT("Archetype did not match due to missing \'any\' tags: %s")
-				, *Description);
-#endif // WITH_MASSENTITY_DEBUG
-			continue;
-		}
-		
-		if (Archetype.GetFragmentBitSet().HasAll(Requirements.GetRequiredAllFragments()) == false)
-		{
-			// missing some required fragments, skip.
-#if WITH_MASSENTITY_DEBUG
-			const FMassFragmentBitSet UnsatisfiedFragments = Requirements.GetRequiredAllFragments() - Archetype.GetFragmentBitSet();
-			FStringOutputDevice Description;
-			UnsatisfiedFragments.DebugGetStringDesc(Description);
-			UE_LOG(LogMass, VeryVerbose, TEXT("Archetype did not match due to missing Fragments: %s")
-				, *Description);
-#endif // WITH_MASSENTITY_DEBUG
-			continue;
-		}
-
-		if (Archetype.GetFragmentBitSet().HasNone(Requirements.GetRequiredNoneFragments()) == false)
-		{
-			// has some Fragments required to be absent
-#if WITH_MASSENTITY_DEBUG
-			const FMassFragmentBitSet UnwantedFragments = Requirements.GetRequiredAllFragments().GetOverlap(Archetype.GetFragmentBitSet());
-			FStringOutputDevice Description;
-			UnwantedFragments.DebugGetStringDesc(Description);
-			UE_LOG(LogMass, VeryVerbose, TEXT("Archetype has Fragments required absent: %s")
-				, *Description);
-#endif // WITH_MASSENTITY_DEBUG
-			continue;
-		}
-
-		if (Requirements.GetRequiredAnyFragments().IsEmpty() == false 
-			&& Archetype.GetFragmentBitSet().HasAny(Requirements.GetRequiredAnyFragments()) == false)
-		{
-#if WITH_MASSENTITY_DEBUG
-			FStringOutputDevice Description;
-			Requirements.GetRequiredAnyFragments().DebugGetStringDesc(Description);
-			UE_LOG(LogMass, VeryVerbose, TEXT("Archetype did not match due to missing \'any\' fragments: %s")
-				, *Description);
-#endif // WITH_MASSENTITY_DEBUG
-			continue;
-		}
-
-		if (Archetype.GetChunkFragmentBitSet().HasAll(Requirements.GetRequiredAllChunkFragments()) == false)
-		{
-			// missing some required fragments, skip.
+		if (FMassArchetypeHelper::DoesArchetypeMatchRequirements(Archetype, Requirements
 #if WITH_MASSENTITY_DEBUG
-			const FMassChunkFragmentBitSet UnsatisfiedFragments = Requirements.GetRequiredAllChunkFragments() - Archetype.GetChunkFragmentBitSet();
-			FStringOutputDevice Description;
-			UnsatisfiedFragments.DebugGetStringDesc(Description);
-			UE_LOG(LogMass, VeryVerbose, TEXT("Archetype did not match due to missing Chunk Fragments: %s")
-				, *Description);
+			, /*bBailOutOnFirstFail=*/true, & LogOutput
 #endif // WITH_MASSENTITY_DEBUG
-			continue;
-		}
-
-		if (Archetype.GetChunkFragmentBitSet().HasNone(Requirements.GetRequiredNoneChunkFragments()) == false)
+		))
 		{
-			// has some Fragments required to be absent
-#if WITH_MASSENTITY_DEBUG
-			const FMassChunkFragmentBitSet UnwantedFragments = Requirements.GetRequiredNoneChunkFragments().GetOverlap(Archetype.GetChunkFragmentBitSet());
-			FStringOutputDevice Description;
-			UnwantedFragments.DebugGetStringDesc(Description);
-			UE_LOG(LogMass, VeryVerbose, TEXT("Archetype has Chunk Fragments required absent: %s")
-				, *Description);
-#endif // WITH_MASSENTITY_DEBUG
-			continue;
-		}
-
-		if (Archetype.GetSharedFragmentBitSet().HasAll(Requirements.GetRequiredAllSharedFragments()) == false)
-		{
-			// missing some required fragments, skip.
-#if WITH_MASSENTITY_DEBUG
-			const FMassSharedFragmentBitSet UnsatisfiedFragments = Requirements.GetRequiredAllSharedFragments() - Archetype.GetSharedFragmentBitSet();
-			FStringOutputDevice Description;
-			UnsatisfiedFragments.DebugGetStringDesc(Description);
-			UE_LOG(LogMass, VeryVerbose, TEXT("Archetype did not match due to missing Shared Fragments: %s")
-				, *Description);
-#endif // WITH_MASSENTITY_DEBUG
-			continue;
+			OutValidArchetypes.Add(ArchetypePtr);
 		}
-
-		if (Archetype.GetSharedFragmentBitSet().HasNone(Requirements.GetRequiredNoneSharedFragments()) == false)
-		{
-			// has some Fragments required to be absent
 #if WITH_MASSENTITY_DEBUG
-			const FMassSharedFragmentBitSet UnwantedFragments = Requirements.GetRequiredNoneSharedFragments().GetOverlap(Archetype.GetSharedFragmentBitSet());
-			FStringOutputDevice Description;
-			UnwantedFragments.DebugGetStringDesc(Description);
-			UE_LOG(LogMass, VeryVerbose, TEXT("Archetype has Shared Fragments required absent: %s")
-				, *Description);
-#endif // WITH_MASSENTITY_DEBUG
-			continue;
-		}
-
-		if (Archetype.GetConstSharedFragmentBitSet().HasAll(Requirements.GetRequiredAllConstSharedFragments()) == false)
+		else
 		{
-			// missing some required fragments, skip.
-#if WITH_MASSENTITY_DEBUG
-			const FMassConstSharedFragmentBitSet UnsatisfiedFragments = Requirements.GetRequiredAllConstSharedFragments() - Archetype.GetConstSharedFragmentBitSet();
-			FStringOutputDevice Description;
-			UnsatisfiedFragments.DebugGetStringDesc(Description);
-			UE_LOG(LogMass, VeryVerbose, TEXT("Archetype did not match due to missing Const Shared Fragments: %s")
-				, *Description);
-#endif // WITH_MASSENTITY_DEBUG
-			continue;
+			UE_VLOG_UELOG(GetOwner(), LogMass, VeryVerbose, TEXT("%s"), *LogOutput);
+			LogOutput.Reset();
 		}
-
-		if (Archetype.GetConstSharedFragmentBitSet().HasNone(Requirements.GetRequiredNoneConstSharedFragments()) == false)
-		{
-			// has some Fragments required to be absent
-#if WITH_MASSENTITY_DEBUG
-			const FMassConstSharedFragmentBitSet UnwantedFragments = Requirements.GetRequiredNoneConstSharedFragments().GetOverlap(Archetype.GetConstSharedFragmentBitSet());
-			FStringOutputDevice Description;
-			UnwantedFragments.DebugGetStringDesc(Description);
-			UE_LOG(LogMass, VeryVerbose, TEXT("Archetype has Const Shared Fragments required absent: %s")
-				, *Description);
 #endif // WITH_MASSENTITY_DEBUG
-			continue;
-		}
-
-		OutValidArchetypes.Add(ArchetypePtr);
 	}
 }
 
diff --git a/Engine/Plugins/Runtime/MassEntity/Source/MassEntity/Private/MassEntityUtils.cpp b/Engine/Plugins/Runtime/MassEntity/Source/MassEntity/Private/MassEntityUtils.cpp
--- a/Engine/Plugins/Runtime/MassEntity/Source/MassEntity/Private/MassEntityUtils.cpp
+++ b/Engine/Plugins/Runtime/MassEntity/Source/MassEntity/Private/MassEntityUtils.cpp
@@ -3,6 +3,8 @@
 #include "MassEntityUtils.h"
 #include "MassEntityTypes.h"
 #include "MassArchetypeTypes.h"
+#include "MassArchetypeData.h"
+#include "MassRequirements.h"
 #include "MassEntityManager.h"
 #include "MassEntitySubsystem.h"
 #include "Engine/EngineBaseTypes.h"
diff --git a/Engine/Plugins/Runtime/MassEntity/Source/MassEntity/Public/MassDebugger.h b/Engine/Plugins/Runtime/MassEntity/Source/MassEntity/Public/MassDebugger.h
--- a/Engine/Plugins/Runtime/MassEntity/Source/MassEntity/Public/MassDebugger.h
+++ b/Engine/Plugins/Runtime/MassEntity/Source/MassEntity/Public/MassDebugger.h
@@ -46,6 +46,7 @@ namespace UE::Mass::Debug
 {
 	extern MASSENTITY_API bool bAllowProceduralDebuggedEntitySelection;
 	extern MASSENTITY_API bool bAllowBreakOnDebuggedEntity;
+	extern MASSENTITY_API bool bTestSelectedEntityAgainstProcessorQueries;
 } // namespace UE::Mass::Debug
 
 #define MASS_IF_ENTITY_DEBUGGED(Manager, EntityHandle) (FMassDebugger::GetSelectedEntity(Manager) == EntityHandle)
@@ -139,6 +140,11 @@ struct MASSENTITY_API FMassDebugger
 	static void UnregisterEntityManager(FMassEntityManager& EntityManager);
 	static TConstArrayView<FEnvironment> GetEnvironments() { return ActiveEnvironments; }
 
+	/**
+	 * Determines whether given Archetype matches given Requirements. In case of a mismatch description of failed conditions will be added to OutputDevice.
+	 */
+	static bool DoesArchetypeMatchRequirements(const FMassArchetypeHandle& ArchetypeHandle, const FMassFragmentRequirements& Requirements, FOutputDevice& OutputDevice);
+
 private:
 	static TArray<FEnvironment> ActiveEnvironments;
 	static UE::FSpinLock EntityManagerRegistrationLock;
diff --git a/Engine/Plugins/Runtime/MassEntity/Source/MassEntityDebugger/Private/MassDebuggerModel.cpp b/Engine/Plugins/Runtime/MassEntity/Source/MassEntityDebugger/Private/MassDebuggerModel.cpp
--- a/Engine/Plugins/Runtime/MassEntity/Source/MassEntityDebugger/Private/MassDebuggerModel.cpp
+++ b/Engine/Plugins/Runtime/MassEntity/Source/MassEntityDebugger/Private/MassDebuggerModel.cpp
@@ -9,6 +9,7 @@
 #include "MassDebuggerSettings.h"
 #include "UObject/UObjectIterator.h"
 #include "Containers/UnrealString.h"
+#include "MassArchetypeData.h"
 
 #define LOCTEXT_NAMESPACE "SMassDebugger"
 
@@ -139,7 +140,8 @@ FMassDebuggerProcessorData::FMassDebuggerProcessorData(const UMassProcessor& InP
 #endif // WITH_MASSENTITY_DEBUG
 }
 
-FMassDebuggerProcessorData::FMassDebuggerProcessorData(const FMassEntityManager& EntityManager, UMassProcessor& InProcessor, const TMap<FMassArchetypeHandle, TSharedPtr<FMassDebuggerArchetypeData>>& InTransientArchetypesMap)
+FMassDebuggerProcessorData::FMassDebuggerProcessorData(const FMassEntityManager& EntityManager, UMassProcessor& InProcessor
+	, const TMap<FMassArchetypeHandle, TSharedPtr<FMassDebuggerArchetypeData>>& InTransientArchetypesMap)
 {
 	SetProcessor(InProcessor);
 #if WITH_MASSENTITY_DEBUG
@@ -147,11 +149,39 @@ FMassDebuggerProcessorData::FMassDebuggerProcessorData(const FMassEntityManager&
 
 	ProcessorRequirements = MakeShareable(new FMassDebuggerQueryData(InProcessor.GetProcessorRequirements(), LOCTEXT("MassProcessorRequirementsLabel", "Processor Requirements")));
 
+	const FMassEntityHandle SelectedEntityHandle = UE::Mass::Debug::bTestSelectedEntityAgainstProcessorQueries
+		? FMassDebugger::GetSelectedEntity(EntityManager)
+		: FMassEntityHandle();
+	FStringOutputDevice SelectedEntityFailureJustificationLog;
+	SelectedEntityFailureJustificationLog.SetAutoEmitLineTerminator(true);
+	const FText SelectedEntityHandleDescription = UE::Mass::Debug::bTestSelectedEntityAgainstProcessorQueries
+		? FText::Format(LOCTEXT("WhyNotEntityJustificationLabel", "Why not entity {0}:"), FText::FromString(SelectedEntityHandle.DebugGetDescription()))
+		: FText();
+
 	Queries.Reserve(ProcessorQueries.Num());
 	for (const FMassEntityQuery* Query : ProcessorQueries)
 	{
 		check(Query);
-		Queries.Add(MakeShareable(new FMassDebuggerQueryData(*Query, LOCTEXT("MassEntityQueryLabel", "Query"))));
+		TSharedPtr<FMassDebuggerQueryData>& QueryData = Queries.Add_GetRef(MakeShareable(new FMassDebuggerQueryData(*Query, LOCTEXT("MassEntityQueryLabel", "Query"))));
+
+		if (SelectedEntityHandle.IsValid())
+		{
+			const FMassArchetypeHandle ArchetypeHandle = EntityManager.GetArchetypeForEntity(SelectedEntityHandle);
+			if (ArchetypeHandle.IsValid() && Query->GetArchetypes().Contains(ArchetypeHandle) == false)
+			{
+				if (FMassArchetypeHelper::DoesArchetypeMatchRequirements(FMassArchetypeHelper::ArchetypeDataFromHandleChecked(ArchetypeHandle), *Query
+					, false, &SelectedEntityFailureJustificationLog) == false)
+				{
+					FTextBuilder DescriptionBuilder;
+					DescriptionBuilder.AppendLine(QueryData->AdditionalInformation);
+					DescriptionBuilder.AppendLine(SelectedEntityHandleDescription);
+					DescriptionBuilder.AppendLine(SelectedEntityFailureJustificationLog);
+					QueryData->AdditionalInformation = DescriptionBuilder.ToText();
+					
+					SelectedEntityFailureJustificationLog.Reset();
+				}
+			}
+		}
 
 		for (const FMassArchetypeHandle& ArchetypeHandle : Query->GetArchetypes())
 		{
diff --git a/Engine/Plugins/Runtime/MassEntity/Source/MassEntityDebugger/Private/MassDebuggerModel.h b/Engine/Plugins/Runtime/MassEntity/Source/MassEntityDebugger/Private/MassDebuggerModel.h
--- a/Engine/Plugins/Runtime/MassEntity/Source/MassEntityDebugger/Private/MassDebuggerModel.h
+++ b/Engine/Plugins/Runtime/MassEntity/Source/MassEntityDebugger/Private/MassDebuggerModel.h
@@ -48,6 +48,7 @@ struct FMassDebuggerQueryData
 
 	FMassExecutionRequirements ExecutionRequirements;
 	FText Label;
+	FText AdditionalInformation;
 
 	int32 GetTotalBitsUsedCount();
 	bool IsEmpty() const;
diff --git a/Engine/Plugins/Runtime/MassEntity/Source/MassEntityDebugger/Private/SMassQuery.cpp b/Engine/Plugins/Runtime/MassEntity/Source/MassEntityDebugger/Private/SMassQuery.cpp
--- a/Engine/Plugins/Runtime/MassEntity/Source/MassEntityDebugger/Private/SMassQuery.cpp
+++ b/Engine/Plugins/Runtime/MassEntity/Source/MassEntityDebugger/Private/SMassQuery.cpp
@@ -69,6 +69,16 @@ void SMassQuery::Construct(const FArguments& InArgs, TSharedPtr<FMassDebuggerQue
 	AddBitSetPair(Box, ExecutionRequirements.ConstSharedFragments, TEXT("Const Shared Fragments"));
 	AddBitSetPair(Box, ExecutionRequirements.RequiredSubsystems, TEXT("Required Subsystems"));
 
+	if (InQueryData->AdditionalInformation.IsEmpty() == false)
+	{
+		Box->AddSlot()
+		.AutoHeight()
+		[
+			SNew(SRichTextBlock)
+			.Text(InQueryData->AdditionalInformation)
+		];
+	}
+
 	ChildSlot
 	[
 		SNew(SBorder)
