Fix for a nasty Mass crash that can happen when:
1. FEntityCreationContext shared pointer is being retained
2. While the context is "alive" modifying freshly created entities' composition
3. Some of the composition of the newly created entities is being observed.
When all three conditions are met the crash is almost certain since due to composition-changing operations the entity collection stored in the creation context gets out of data

The change boils down to invalidating the entity collection when entity mutating operations take place. Also, while the creation context is active, we don't trigger observers - same "on add" observers will get triggered when creation context is released.

diff --git a/Engine/Plugins/Runtime/MassEntity/Source/MassEntity/Private/MassEntityManager.cpp b/Engine/Plugins/Runtime/MassEntity/Source/MassEntity/Private/MassEntityManager.cpp
--- a/Engine/Plugins/Runtime/MassEntity/Source/MassEntity/Private/MassEntityManager.cpp
+++ b/Engine/Plugins/Runtime/MassEntity/Source/MassEntity/Private/MassEntityManager.cpp
@@ -13,6 +13,7 @@
 #include "Misc/Fork.h"
 #include "Misc/CoreDelegates.h"
 #include "Algo/Find.h"
+#include "MassEntityUtils.h"
 
 const FMassEntityHandle FMassEntityManager::InvalidEntity;
 
@@ -43,6 +44,64 @@ namespace UE::Mass::Private
 	constexpr int32 InvalidEntityIndex = 0;
 }
 
+//-----------------------------------------------------------------------------
+// FMassEntityManager::FEntityCreationContext
+//-----------------------------------------------------------------------------
+FMassEntityManager::FEntityCreationContext::FEntityCreationContext(FMassEntityManager& InManager, const TConstArrayView<FMassEntityHandle> InCreatedEntities)
+	: CreatedEntities(InCreatedEntities)
+{
+	Manager = InManager.AsShared();
+}
+
+FMassEntityManager::FEntityCreationContext::~FEntityCreationContext()
+{
+	if (OnSpawningFinished)
+	{
+		OnSpawningFinished(*this);
+	}
+}
+
+TConstArrayView<FMassArchetypeEntityCollection> FMassEntityManager::FEntityCreationContext::GetEntityCollections() const
+{
+	// the EntityCollection has been dirtied, we need to rebuild it
+	if (EntityCollections.IsEmpty() == true && CreatedEntities.IsEmpty() == false && Manager)
+	{
+		UE::Mass::Utils::CreateEntityCollections(*Manager.Get(), CreatedEntities, CollectionCreationDuplicatesHandling, EntityCollections);
+	}
+
+	return EntityCollections;
+}
+
+void FMassEntityManager::FEntityCreationContext::MarkDirty()
+{
+	EntityCollections.Reset();
+}
+
+void FMassEntityManager::FEntityCreationContext::AddCollection(FMassArchetypeEntityCollection&& Collection)
+{
+	EntityCollections.Add(MoveTemp(Collection));
+}
+
+void FMassEntityManager::FEntityCreationContext::AppendEntities(const TConstArrayView<FMassEntityHandle> EntitiesToAppend)
+{
+	if (EntitiesToAppend.Num())
+	{
+		if (CreatedEntities.Num())
+		{
+			// since we already have entities in CreatedEntities (initially ensured to have no duplicates) we cannot 
+			// guarantee anymore that we'll have no duplicates after adding EntitiesToAppend
+			CollectionCreationDuplicatesHandling = FMassArchetypeEntityCollection::FoldDuplicates;
+		}
+
+		CreatedEntities.Append(EntitiesToAppend);
+
+		MarkDirty();
+	}
+}
+
+//-----------------------------------------------------------------------------
+// FMassEntityManager
+//-----------------------------------------------------------------------------
 #if MASS_CONCURRENT_RESERVE
 UE::Mass::IEntityStorageInterface& FMassEntityManager::GetEntityStorageInterface()
 {
@@ -119,6 +178,9 @@ const UE::Mass::IEntityStorageInterface& FMassEntityManager::DebugGetEntityStora
 }
 #endif
 
+//-----------------------------------------------------------------------------
+// FMassEntityManager
+//-----------------------------------------------------------------------------
 FMassEntityManager::FMassEntityManager(UObject* InOwner)
 	: ObserverManager(*this)
 	, Owner(InOwner)
@@ -128,9 +190,6 @@ FMassEntityManager::FMassEntityManager(UObject* InOwner)
 #endif
 }
 
-//////////////////////////////////////////////////////////////////////
-// FMassEntityManager
-
 FMassEntityManager::~FMassEntityManager()
 {
 	if (bInitialized)
@@ -768,7 +827,7 @@ TSharedRef<FMassEntityManager::FEntityCreationContext> FMassEntityManager::Batch
 	TRACE_CPUPROFILER_EVENT_SCOPE(Mass_BatchCreateEntities);
 
 	checkf(IsProcessing() == false, TEXT("Synchronous API function %hs called during mass processing. Use asynchronous API instead."), __FUNCTION__);
-	testableCheckfReturn(ArchetypeHandle.IsValid(), return MakeShareable(new FEntityCreationContext(0))
+	testableCheckfReturn(ArchetypeHandle.IsValid(), return GetOrMakeCreationContext()
 		, TEXT("%hs expecting a valid ArchetypeHandle"), __FUNCTION__);
 
 	TConstArrayView<FMassEntityHandle> ReservedEntities = BatchReserveEntities(Count, InOutEntities);
@@ -797,18 +856,21 @@ TSharedRef<FMassEntityManager::FEntityCreationContext> FMassEntityManager::Inter
 	FMassArchetypeEntityCollection::FEntityRangeArray TargetArchetypeEntityRanges;
 	ArchetypeData->BatchAddEntities(ReservedEntities, SharedFragmentValues, TargetArchetypeEntityRanges);
 
-	FEntityCreationContext* CreationContext = new FEntityCreationContext(ReservedEntities.Num());
-	new (&CreationContext->EntityCollection) FMassArchetypeEntityCollection(ArchetypeHandle, MoveTemp(TargetArchetypeEntityRanges));
+	TSharedRef<FEntityCreationContext> SharedContext = GetOrMakeCreationContext(ReservedEntities);
+	FEntityCreationContext* CreationContext = &SharedContext.Get();
+	CreationContext->AddCollection(FMassArchetypeEntityCollection(ArchetypeHandle, MoveTemp(TargetArchetypeEntityRanges)));
 
-	if (ObserverManager.HasObserversForBitSet(ArchetypeData->GetCompositionDescriptor().Fragments, EMassObservedOperation::Add)
-		|| ObserverManager.HasObserversForBitSet(ArchetypeData->GetCompositionDescriptor().Tags, EMassObservedOperation::Add))
+	if (!CreationContext->OnSpawningFinished
+		&& (ObserverManager.HasObserversForBitSet(ArchetypeData->GetCompositionDescriptor().Fragments, EMassObservedOperation::Add)
+			|| ObserverManager.HasObserversForBitSet(ArchetypeData->GetCompositionDescriptor().Tags, EMassObservedOperation::Add)))
 	{
-		CreationContext->OnSpawningFinished = [this](FEntityCreationContext& Context) {
-			ObserverManager.OnPostEntitiesCreated(Context.EntityCollection);
+		CreationContext->OnSpawningFinished = [this](FEntityCreationContext& Context)
+		{
+			ObserverManager.OnPostEntitiesCreated(Context.GetEntityCollections());
 		};
 	}
 
-	return MakeShareable(CreationContext);
+	return SharedContext;
 }
 
 void FMassEntityManager::DestroyEntity(FMassEntityHandle Entity)
@@ -833,6 +895,7 @@ void FMassEntityManager::BatchDestroyEntities(TConstArrayView<FMassEntityHandle>
 	TRACE_CPUPROFILER_EVENT_SCOPE(Mass_BatchDestroyEntities);
 
 	checkf(IsProcessing() == false, TEXT("Synchronous API function %hs called during mass processing. Use asynchronous API instead."), __FUNCTION__);
+	checkf(IsDuringEntityCreation() == false, TEXT("%hs: Trying to destroy entities while entity creation is under way. This operation is not supported."), __FUNCTION__);
 
 	// @todo optimize, we can make savings by implementing Archetype->RemoveEntities()
 	for (const FMassEntityHandle Entity : InEntities)
@@ -862,6 +925,7 @@ void FMassEntityManager::BatchDestroyEntityChunks(const FMassArchetypeEntityColl
 	TRACE_CPUPROFILER_EVENT_SCOPE(Mass_BatchDestroyEntityChunks);
 
 	checkf(IsProcessing() == false, TEXT("Synchronous API function %hs called during mass processing. Use asynchronous API instead."), __FUNCTION__);
+	checkf(IsDuringEntityCreation() == false, TEXT("%hs: Trying to destroy entities while entity creation is under way. This operation is not supported."), __FUNCTION__);
 
 	TArray<FMassEntityHandle> EntitiesRemoved;
 	// note that it's important to place the context instance in the same scope as the loop below that updates 
@@ -896,7 +960,14 @@ void FMassEntityManager::AddFragmentToEntity(FMassEntityHandle Entity, const USc
 	CheckIfEntityIsActive(Entity);
 
 	const FMassArchetypeCompositionDescriptor Descriptor(InternalAddFragmentListToEntityChecked(Entity, FMassFragmentBitSet(*FragmentType)));
-	ObserverManager.OnPostCompositionAdded(Entity, Descriptor);
+	if (IsDuringEntityCreation())
+	{
+		DirtyCreationContext();
+	}
+	else
+	{
+		ObserverManager.OnPostCompositionAdded(Entity, Descriptor);
+	}
 }
 
 void FMassEntityManager::AddFragmentToEntity(FMassEntityHandle Entity, const UScriptStruct* FragmentType, const FStructInitializationCallback& Initializer)
@@ -913,7 +984,14 @@ void FMassEntityManager::AddFragmentToEntity(FMassEntityHandle Entity, const USc
 	Initializer(FragmentData, *FragmentType);
 
 	const FMassArchetypeCompositionDescriptor Descriptor(MoveTemp(Fragments));
-	ObserverManager.OnPostCompositionAdded(Entity, Descriptor);
+	if (IsDuringEntityCreation())
+	{
+		DirtyCreationContext();
+	}
+	else 
+	{
+		ObserverManager.OnPostCompositionAdded(Entity, Descriptor);
+	}
 }
 
 void FMassEntityManager::AddFragmentListToEntity(FMassEntityHandle Entity, TConstArrayView<const UScriptStruct*> FragmentList)
@@ -921,7 +999,14 @@ void FMassEntityManager::AddFragmentListToEntity(FMassEntityHandle Entity, TCons
 	CheckIfEntityIsActive(Entity);
 
 	const FMassArchetypeCompositionDescriptor Descriptor(InternalAddFragmentListToEntityChecked(Entity, FMassFragmentBitSet(FragmentList)));
-	ObserverManager.OnPostCompositionAdded(Entity, Descriptor);
+	if (IsDuringEntityCreation())
+	{
+		DirtyCreationContext();
+	}
+	else
+	{
+		ObserverManager.OnPostCompositionAdded(Entity, Descriptor);
+	}
 }
 
 void FMassEntityManager::AddCompositionToEntity_GetDelta(FMassEntityHandle Entity, FMassArchetypeCompositionDescriptor& InDescriptor)
@@ -953,7 +1038,14 @@ void FMassEntityManager::AddCompositionToEntity_GetDelta(FMassEntityHandle Entit
 
 			GetEntityStorageInterface().SetArchetypeFromShared(Entity.Index, NewArchetypeHandle.DataPtr);
 
-			ObserverManager.OnPostCompositionAdded(Entity, InDescriptor);
+			if (IsDuringEntityCreation())
+			{
+				DirtyCreationContext();
+			}
+			else
+			{
+				ObserverManager.OnPostCompositionAdded(Entity, InDescriptor);
+			}
 		}
 	}
 }
@@ -977,7 +1069,14 @@ void FMassEntityManager::RemoveCompositionFromEntity(FMassEntityHandle Entity, c
 		if (NewDescriptor.IsEquivalent(OldArchetype->GetCompositionDescriptor()) == false)
 		{
 			ensureMsgf(OldArchetype->GetCompositionDescriptor().HasAll(InDescriptor), TEXT("Some of the elements being removed are already missing from entity\'s composition."));
-			ObserverManager.OnPreCompositionRemoved(Entity, InDescriptor);
+			if (IsDuringEntityCreation())
+			{
+				DirtyCreationContext();
+			}
+			else
+			{
+				ObserverManager.OnPreCompositionRemoved(Entity, InDescriptor);
+			}
 
 			const FMassArchetypeHandle NewArchetypeHandle = CreateArchetype(NewDescriptor, FMassArchetypeCreationParams(*OldArchetype));
 
@@ -1005,7 +1104,14 @@ void FMassEntityManager::InternalBuildEntity(FMassEntityHandle Entity, const FMa
 	GetEntityStorageInterface().SetArchetypeFromShared(Entity.Index, ArchetypeHandle.DataPtr);
 	NewArchetype->AddEntity(Entity, SharedFragmentValues);
 
-	ObserverManager.OnPostCompositionAdded(Entity, NewArchetype->GetCompositionDescriptor());
+	if (IsDuringEntityCreation())
+	{
+		DirtyCreationContext();
+	}
+	else
+	{
+		ObserverManager.OnPostCompositionAdded(Entity, NewArchetype->GetCompositionDescriptor());
+	}
 }
 
 void FMassEntityManager::InternalReleaseEntity(FMassEntityHandle Entity)
@@ -1063,7 +1169,14 @@ void FMassEntityManager::AddFragmentInstanceListToEntity(FMassEntityHandle Entit
 	check(CurrentArchetype);
 	CurrentArchetype->SetFragmentsData(Entity, FragmentInstanceList);
 
-	ObserverManager.OnPostCompositionAdded(Entity, Descriptor);
+	if (IsDuringEntityCreation())
+	{
+		DirtyCreationContext();
+	}
+	else 
+	{
+		ObserverManager.OnPostCompositionAdded(Entity, Descriptor);
+	}
 }
 
 void FMassEntityManager::RemoveFragmentFromEntity(FMassEntityHandle Entity, const UScriptStruct* FragmentType)
@@ -1095,7 +1208,14 @@ void FMassEntityManager::RemoveFragmentListFromEntity(FMassEntityHandle Entity,
 		FMassArchetypeCompositionDescriptor CompositionDelta;
 		// Find overlap.  It isn't guaranteed that the old archetype has all of the fragments being removed.
 		CompositionDelta.Fragments = OldArchetype->GetFragmentBitSet().GetOverlap(FragmentsToRemove);
-		ObserverManager.OnPreCompositionRemoved(Entity, CompositionDelta);
+		if (IsDuringEntityCreation())
+		{
+			DirtyCreationContext();
+		}
+		else
+		{
+			ObserverManager.OnPreCompositionRemoved(Entity, CompositionDelta);
+		}
 
 		// Move the entity over
 		FMassArchetypeData& NewArchetype = FMassArchetypeHelper::ArchetypeDataFromHandleChecked(NewArchetypeHandle);
@@ -1159,7 +1279,14 @@ void FMassEntityManager::AddTagToEntity(FMassEntityHandle Entity, const UScriptS
 		FMassTagBitSet TagDelta;
 		TagDelta.Add(*TagType);
 		CompositionDelta.Tags = TagDelta;
-		ObserverManager.OnPostCompositionAdded(Entity, CompositionDelta);
+		if (IsDuringEntityCreation())
+		{
+			DirtyCreationContext();
+		}
+		else
+		{
+			ObserverManager.OnPostCompositionAdded(Entity, CompositionDelta);
+		}
 	}
 }
 	
@@ -1178,7 +1305,14 @@ void FMassEntityManager::RemoveTagFromEntity(FMassEntityHandle Entity, const USc
 		FMassTagBitSet TagDelta;
 		TagDelta.Add(*TagType);
 		CompositionDelta.Tags = TagDelta;
-		ObserverManager.OnPreCompositionRemoved(Entity, CompositionDelta);
+		if (IsDuringEntityCreation())
+		{
+			DirtyCreationContext();
+		}
+		else
+		{
+			ObserverManager.OnPreCompositionRemoved(Entity, CompositionDelta);
+		}
 		
 		// CurrentArchetype->GetTagBitSet() -  *TagType
 		const FMassTagBitSet NewTagComposition = CurrentArchetype->GetTagBitSet() - TagDelta;
@@ -1255,7 +1389,11 @@ void FMassEntityManager::BatchChangeTagsForEntities(TConstArrayView<FMassArchety
 			FMassTagBitSet TagsAdded = TagsToAdd - CurrentArchetype->GetTagBitSet();
 			FMassTagBitSet TagsRemoved = TagsToRemove.GetOverlap(CurrentArchetype->GetTagBitSet());
 
-			if (ObserverManager.HasObserversForBitSet(TagsRemoved, EMassObservedOperation::Remove))
+			if (IsDuringEntityCreation())
+			{
+				DirtyCreationContext();
+			}
+			else if (ObserverManager.HasObserversForBitSet(TagsRemoved, EMassObservedOperation::Remove))
 			{
 				ObserverManager.OnCompositionChanged(Collection, FMassArchetypeCompositionDescriptor(MoveTemp(TagsRemoved)), EMassObservedOperation::Remove);
 			}
@@ -1277,7 +1415,7 @@ void FMassEntityManager::BatchChangeTagsForEntities(TConstArrayView<FMassArchety
 				GetEntityStorageInterface().SetArchetypeFromShared(Entity.Index, NewArchetypeHandle.DataPtr);
 			}
 
-			if (bTagsAddedAreObserved)
+			if (bTagsAddedAreObserved && !IsDuringEntityCreation())
 			{
 				ObserverManager.OnCompositionChanged(
 					FMassArchetypeEntityCollection(NewArchetypeHandle, MoveTemp(NewArchetypeEntityRanges))
@@ -1306,7 +1444,12 @@ void FMassEntityManager::BatchChangeFragmentCompositionForEntities(TConstArrayVi
 				FMassFragmentBitSet FragmentsAdded = FragmentsToAdd - CurrentArchetype->GetFragmentBitSet();
 				const bool bFragmentsAddedAreObserved = ObserverManager.HasObserversForBitSet(FragmentsAdded, EMassObservedOperation::Add);
 				FMassFragmentBitSet FragmentsRemoved = FragmentsToRemove.GetOverlap(CurrentArchetype->GetFragmentBitSet());
-				if (ObserverManager.HasObserversForBitSet(FragmentsRemoved, EMassObservedOperation::Remove))
+				
+				if (IsDuringEntityCreation())
+				{
+					DirtyCreationContext();
+				}
+				else if (ObserverManager.HasObserversForBitSet(FragmentsRemoved, EMassObservedOperation::Remove))
 				{
 					ObserverManager.OnCompositionChanged(Collection, FMassArchetypeCompositionDescriptor(MoveTemp(FragmentsRemoved)), EMassObservedOperation::Remove);
 				}
@@ -1327,7 +1470,7 @@ void FMassEntityManager::BatchChangeFragmentCompositionForEntities(TConstArrayVi
 					GetEntityStorageInterface().SetArchetypeFromShared(Entity.Index, NewArchetypeHandle.DataPtr);
 				}
 
-				if (bFragmentsAddedAreObserved)
+				if (bFragmentsAddedAreObserved && !IsDuringEntityCreation())
 				{
 					ObserverManager.OnCompositionChanged(
 						FMassArchetypeEntityCollection(NewArchetypeHandle, MoveTemp(NewArchetypeEntityRanges))
@@ -1404,7 +1547,11 @@ void FMassEntityManager::BatchAddFragmentInstancesForEntities(TConstArrayView<FM
 			// corresponds to the order in FMassArchetypeEntityCollectionWithPayload's payload
 			TargetArchetypeHandle.DataPtr->BatchSetFragmentValues(TargetArchetypeEntityRanges, EntityRangesWithPayload.GetPayload());
 
-			if (bFragmentsAddedAreObserved)
+			if (IsDuringEntityCreation())
+			{
+				DirtyCreationContext();
+			}
+			else if (bFragmentsAddedAreObserved)
 			{
 				ObserverManager.OnCompositionChanged(
 					FMassArchetypeEntityCollection(TargetArchetypeHandle, MoveTemp(TargetArchetypeEntityRanges))
@@ -1479,14 +1626,30 @@ void FMassEntityManager::SetEntityFragmentsValues(FMassEntityHandle Entity, TArr
 
 void FMassEntityManager::BatchSetEntityFragmentsValues(const FMassArchetypeEntityCollection& SparseEntities, TArrayView<const FInstancedStruct> FragmentInstanceList)
 {
-	TRACE_CPUPROFILER_EVENT_SCOPE(Mass_BatchSetEntityFragmentsValues);
+	if (FragmentInstanceList.Num())
+	{
+		BatchSetEntityFragmentsValues(MakeArrayView(&SparseEntities, 1), FragmentInstanceList);
+	}
+}
 
-	FMassArchetypeData* Archetype = SparseEntities.GetArchetype().DataPtr.Get();
-	check(Archetype);
+void FMassEntityManager::BatchSetEntityFragmentsValues(TConstArrayView<FMassArchetypeEntityCollection> EntityCollections, TArrayView<const FInstancedStruct> FragmentInstanceList)
+{
+	if (FragmentInstanceList.IsEmpty())
+	{
+		return;
+	}
 
-	for (const FInstancedStruct& FragmentTemplate : FragmentInstanceList)
+	for (const FMassArchetypeEntityCollection& SparseEntities : EntityCollections)
 	{
-		Archetype->SetFragmentData(SparseEntities.GetRanges(), FragmentTemplate);
+		TRACE_CPUPROFILER_EVENT_SCOPE(Mass_BatchSetEntityFragmentsValues);
+
+		FMassArchetypeData* Archetype = SparseEntities.GetArchetype().DataPtr.Get();
+		check(Archetype);
+
+		for (const FInstancedStruct& FragmentTemplate : FragmentInstanceList)
+		{
+			Archetype->SetFragmentData(SparseEntities.GetRanges(), FragmentTemplate);
+		}
 	}
 }
 
@@ -1722,6 +1885,32 @@ void FMassEntityManager::AppendCommands(TSharedPtr<FMassCommandBuffer>& InOutCom
 	Defer().MoveAppend(*InOutCommandBuffer.Get());
 }
 
+TSharedRef<FMassEntityManager::FEntityCreationContext> FMassEntityManager::GetOrMakeCreationContext(TConstArrayView<FMassEntityHandle> ReservedEntities)
+{
+	if (ActiveCreationContext.IsValid())
+	{
+		TSharedPtr<FEntityCreationContext> SharedContext = ActiveCreationContext.Pin();
+		CA_ASSUME(SharedContext);
+		SharedContext->AppendEntities(ReservedEntities);
+		return SharedContext.ToSharedRef();
+	}
+	else
+	{
+		FEntityCreationContext* CreationContext = new FEntityCreationContext(*this, ReservedEntities);
+		TSharedRef<FEntityCreationContext> SharedContext = MakeShareable(CreationContext);
+		ActiveCreationContext = SharedContext;
+		return SharedContext;
+	}
+}
+
+void FMassEntityManager::DirtyCreationContext()
+{
+	if (TSharedPtr<FEntityCreationContext> AsSharedPtr = ActiveCreationContext.Pin())
+	{
+		AsSharedPtr->MarkDirty();
+	}
+}
+
 void FMassEntityManager::SetDebugName(const FString& NewDebugGame) 
 { 
 #if WITH_MASSENTITY_DEBUG
@@ -1917,4 +2106,10 @@ void FMassEntityManager::BatchBuildEntities(const FMassArchetypeEntityCollection
 	FMassArchetypeCreationParams Params;
 	Params.DebugName = ArchetypeDebugName;
 	BatchBuildEntities(EncodedEntitiesWithPayload, MoveTemp(Composition), SharedFragmentValues, Params);
-}
\ No newline at end of file
+}
+
+const FMassArchetypeEntityCollection& FMassEntityManager::FEntityCreationContext::GetEntityCollection() const
+{
+	static FMassArchetypeEntityCollection EmptyCollection;
+	return EntityCollections.Num() ? EntityCollections[0] : EmptyCollection;
+}
diff --git a/Engine/Plugins/Runtime/MassEntity/Source/MassEntity/Private/MassExecutor.cpp b/Engine/Plugins/Runtime/MassEntity/Source/MassEntity/Private/MassExecutor.cpp
--- a/Engine/Plugins/Runtime/MassEntity/Source/MassEntity/Private/MassExecutor.cpp
+++ b/Engine/Plugins/Runtime/MassEntity/Source/MassEntity/Private/MassExecutor.cpp
@@ -9,6 +9,15 @@
 
 namespace UE::Mass::Executor
 {
+
+FORCEINLINE void ExecuteProcessors(FMassEntityManager& EntityManager, TArrayView<UMassProcessor* const> Processors, FMassExecutionContext& ExecutionContext)
+{
+	for (UMassProcessor* Proc : Processors)
+	{
+		Proc->CallExecute(EntityManager, ExecutionContext);
+	}
+}
+
 void Run(FMassRuntimePipeline& RuntimePipeline, FMassProcessingContext& ProcessingContext)
 {
 	if (!ensure(ProcessingContext.EntityManager) || 
@@ -35,7 +44,7 @@ void RunSparse(FMassRuntimePipeline& RuntimePipeline, FMassProcessingContext& Pr
 	TRACE_CPUPROFILER_EVENT_SCOPE_STR("MassExecutor RunSparseEntities");
 
 	const FMassArchetypeEntityCollection EntityCollection(Archetype, Entities, FMassArchetypeEntityCollection::NoDuplicates);
-	RunProcessorsView(RuntimePipeline.GetMutableProcessors(), ProcessingContext, &EntityCollection);
+	RunProcessorsView(RuntimePipeline.GetMutableProcessors(), ProcessingContext, MakeArrayView(&EntityCollection, 1));
 }
 
 void RunSparse(FMassRuntimePipeline& RuntimePipeline, FMassProcessingContext& ProcessingContext, const FMassArchetypeEntityCollection& EntityCollection)
@@ -50,7 +59,7 @@ void RunSparse(FMassRuntimePipeline& RuntimePipeline, FMassProcessingContext& Pr
 
 	TRACE_CPUPROFILER_EVENT_SCOPE_STR("MassExecutor RunSparse");
 
-	RunProcessorsView(RuntimePipeline.GetMutableProcessors(), ProcessingContext, &EntityCollection);
+	RunProcessorsView(RuntimePipeline.GetMutableProcessors(), ProcessingContext, MakeArrayView(&EntityCollection, 1));
 }
 
 void Run(UMassProcessor& Processor, FMassProcessingContext& ProcessingContext)
@@ -66,7 +75,7 @@ void Run(UMassProcessor& Processor, FMassProcessingContext& ProcessingContext)
 	RunProcessorsView(MakeArrayView(&ProcPtr, 1), ProcessingContext);
 }
 
-void RunProcessorsView(TArrayView<UMassProcessor* const> Processors, FMassProcessingContext& ProcessingContext, const FMassArchetypeEntityCollection* EntityCollection)
+void RunProcessorsView(TArrayView<UMassProcessor* const> Processors, FMassProcessingContext& ProcessingContext, TConstArrayView<FMassArchetypeEntityCollection> EntityCollections)
 {
 	TRACE_CPUPROFILER_EVENT_SCOPE(RunProcessorsView);
 
@@ -86,11 +95,7 @@ void RunProcessorsView(TArrayView<UMassProcessor* const> Processors, FMassProces
 	TRACE_CPUPROFILER_EVENT_SCOPE_STR("MassExecutor RunProcessorsView")
 
 	FMassExecutionContext ExecutionContext(*ProcessingContext.EntityManager.Get(), ProcessingContext.DeltaSeconds);
-	if (EntityCollection)
-	{
-		ExecutionContext.SetEntityCollection(*EntityCollection);
-	}
-	
+
 	// if ProcessingContext points at a valid CommandBuffer use that one, otherwise manually create a new command buffer 
 	// to let the default one still be used by code unaware of mass processing
 	TSharedPtr<FMassCommandBuffer> CommandBuffer = ProcessingContext.CommandBuffer 
@@ -106,9 +111,18 @@ void RunProcessorsView(TArrayView<UMassProcessor* const> Processors, FMassProces
 		
 		FMassEntityManager::FScopedProcessing ProcessingScope = EntityManager.NewProcessingScope();
 
-		for (UMassProcessor* Proc : Processors)
+		if (EntityCollections.Num() == 0)
 		{
-			Proc->CallExecute(*ProcessingContext.EntityManager, ExecutionContext);
+			ExecuteProcessors(*ProcessingContext.EntityManager, Processors, ExecutionContext);
+		}
+		else
+		{
+			// @todo change ExecutionContext to contain a TConstArrayView of collections. Will lead up to changes in the entity query as well.
+			for (const FMassArchetypeEntityCollection& Collection : EntityCollections)
+			{
+				ExecutionContext.SetEntityCollection(Collection);
+				ExecuteProcessors(*ProcessingContext.EntityManager, Processors, ExecutionContext);
+			}
 		}
 	}
 	
@@ -212,4 +226,19 @@ FGraphEventRef TriggerParallelTasks(UMassProcessor& Processor, FMassProcessingCo
 	return CompletionEvent;
 }
 
+//-----------------------------------------------------------------------------
+// DEPRECATED
+//-----------------------------------------------------------------------------
+void RunProcessorsView(TArrayView<UMassProcessor* const> Processors, FMassProcessingContext& ProcessingContext, const FMassArchetypeEntityCollection* EntityCollection)
+{
+	if (EntityCollection)
+	{
+		RunProcessorsView(Processors, ProcessingContext, MakeArrayView(EntityCollection, 1));
+	}
+	else
+	{
+		RunProcessorsView(Processors, ProcessingContext);
+	}
+}
+
 } // namespace UE::Mass::Executor
diff --git a/Engine/Plugins/Runtime/MassEntity/Source/MassEntity/Private/MassObserverManager.cpp b/Engine/Plugins/Runtime/MassEntity/Source/MassEntity/Private/MassObserverManager.cpp
--- a/Engine/Plugins/Runtime/MassEntity/Source/MassEntity/Private/MassObserverManager.cpp
+++ b/Engine/Plugins/Runtime/MassEntity/Source/MassEntity/Private/MassObserverManager.cpp
@@ -83,6 +83,11 @@ void FMassObserverManager::Initialize()
 }
 
 bool FMassObserverManager::OnPostEntitiesCreated(const FMassArchetypeEntityCollection& EntityCollection)
+{
+	return OnPostEntitiesCreated(MakeArrayView(&EntityCollection , 1));
+}
+
+bool FMassObserverManager::OnPostEntitiesCreated(TConstArrayView<FMassArchetypeEntityCollection> EntityCollections)
 {
 	FMassProcessingContext ProcessingContext(EntityManager, /*DeltaSeconds=*/0.f);
 	// requesting not to flush commands since handling creation of new entities can result in multiple collections of
@@ -91,7 +96,7 @@ bool FMassObserverManager::OnPostEntitiesCreated(const FMassArchetypeEntityColle
 	ProcessingContext.bFlushCommandBuffer = false;
 	ProcessingContext.CommandBuffer = MakeShareable(new FMassCommandBuffer());
 
-	if (OnPostEntitiesCreated(ProcessingContext, EntityCollection))
+	if (OnPostEntitiesCreated(ProcessingContext, EntityCollections))
 	{
 		EntityManager.FlushCommands(ProcessingContext.CommandBuffer);
 		return true;
@@ -99,14 +104,27 @@ bool FMassObserverManager::OnPostEntitiesCreated(const FMassArchetypeEntityColle
 	return false;
 }
 
+// DEPRECATED
 bool FMassObserverManager::OnPostEntitiesCreated(FMassProcessingContext& ProcessingContext, const FMassArchetypeEntityCollection& EntityCollection)
+{
+	return OnPostEntitiesCreated(ProcessingContext, MakeArrayView(&EntityCollection, 1));
+}
+
+bool FMassObserverManager::OnPostEntitiesCreated(FMassProcessingContext& ProcessingContext, TConstArrayView<FMassArchetypeEntityCollection> EntityCollections)
 {
 	TRACE_CPUPROFILER_EVENT_SCOPE_STR("OnPostEntitiesCreated")
 
 	check(ProcessingContext.EntityManager);
-	const FMassArchetypeCompositionDescriptor& ArchetypeComposition = ProcessingContext.EntityManager->GetArchetypeComposition(EntityCollection.GetArchetype());
 
-	return OnCompositionChanged(ProcessingContext, EntityCollection, ArchetypeComposition, EMassObservedOperation::Add);
+	bool bReturnValue = false;
+
+	for (const FMassArchetypeEntityCollection& Collection : EntityCollections)
+	{
+		const FMassArchetypeCompositionDescriptor& ArchetypeComposition = ProcessingContext.EntityManager->GetArchetypeComposition(Collection.GetArchetype());
+		bReturnValue = OnCompositionChanged(ProcessingContext, Collection, ArchetypeComposition, EMassObservedOperation::Add) || bReturnValue;
+	}
+
+	return bReturnValue;
 }
 
 bool FMassObserverManager::OnPreEntitiesDestroyed(const FMassArchetypeEntityCollection& EntityCollection)
@@ -240,7 +258,7 @@ void FMassObserverManager::HandleFragmentsImpl(FMassProcessingContext& Processin
 		ProcessingContext.AuxData.InitializeAs(Type);
 		FMassRuntimePipeline& Pipeline = (*HandlersContainer).FindChecked(Type);
 
-		UE::Mass::Executor::RunProcessorsView(Pipeline.GetMutableProcessors(), ProcessingContext, &EntityCollection);
+		UE::Mass::Executor::RunProcessorsView(Pipeline.GetMutableProcessors(), ProcessingContext, MakeArrayView(&EntityCollection, 1));
 	}
 }
 
@@ -253,7 +271,7 @@ void FMassObserverManager::HandleSingleEntityImpl(const UScriptStruct& FragmentT
 	ProcessingContext.AuxData.InitializeAs(&FragmentType);
 	FMassRuntimePipeline& Pipeline = (*HandlersContainer).FindChecked(&FragmentType);
 
-	UE::Mass::Executor::RunProcessorsView(Pipeline.GetMutableProcessors(), ProcessingContext, &EntityCollection);
+	UE::Mass::Executor::RunProcessorsView(Pipeline.GetMutableProcessors(), ProcessingContext, MakeArrayView(&EntityCollection, 1));
 }
 
 void FMassObserverManager::AddObserverInstance(const UScriptStruct& FragmentOrTagType, const EMassObservedOperation Operation, UMassProcessor& ObserverProcessor)
diff --git a/Engine/Plugins/Runtime/MassEntity/Source/MassEntity/Public/MassEntityManager.h b/Engine/Plugins/Runtime/MassEntity/Source/MassEntity/Public/MassEntityManager.h
--- a/Engine/Plugins/Runtime/MassEntity/Source/MassEntity/Public/MassEntityManager.h
+++ b/Engine/Plugins/Runtime/MassEntity/Source/MassEntity/Public/MassEntityManager.h
@@ -209,19 +209,33 @@ public:
 	/**
 	 * A dedicated structure for ensuring the "on entities creation" observers get notified only once all other 
 	 * initialization operations are done and this creation context instance gets released. */
-	struct FEntityCreationContext
+	struct MASSENTITY_API FEntityCreationContext
 	{
-		explicit FEntityCreationContext(const int32 InNumSpawned = 0)
-			: NumberSpawned(InNumSpawned)
-		{}
-		~FEntityCreationContext() { if (OnSpawningFinished) OnSpawningFinished(*this); }
+	private:
+		FEntityCreationContext() = default;
+		explicit FEntityCreationContext(FMassEntityManager& InManager, const TConstArrayView<FMassEntityHandle> InCreatedEntities = {});
+
+	public:
+		~FEntityCreationContext();
+
+		TConstArrayView<FMassArchetypeEntityCollection> GetEntityCollections() const;
+		int32 GetSpawnedNum() const { return CreatedEntities.Num(); }
+		void MarkDirty();
+		void AddCollection(FMassArchetypeEntityCollection&& Collection);
+		void AppendEntities(const TConstArrayView<FMassEntityHandle> EntitiesToAppend);
+
+		UE_DEPRECATED(5.5, "This constructor is now deprecated and defunct. Use one of the others instead.")
+		explicit FEntityCreationContext(const int32 InNumSpawned) : FEntityCreationContext() {}
+		UE_DEPRECATED(5.5, "This function is now deprecated since FEntityCreationContext can contain more than a single collection now. Use GetEntityCollections instead.")
+		const FMassArchetypeEntityCollection& GetEntityCollection() const;
 
-		const FMassArchetypeEntityCollection& GetEntityCollection() const { return EntityCollection; }
 	private:
 		friend FMassEntityManager;
-		int32 NumberSpawned;
-		FMassArchetypeEntityCollection EntityCollection;
+		mutable TArray<FMassArchetypeEntityCollection> EntityCollections;
+		TArray<FMassEntityHandle> CreatedEntities;
+		FMassArchetypeEntityCollection::EDuplicatesHandling CollectionCreationDuplicatesHandling = FMassArchetypeEntityCollection::EDuplicatesHandling::NoDuplicates;
 		TFunction<void(FEntityCreationContext&)> OnSpawningFinished;
+		TSharedPtr<FMassEntityManager> Manager;
 	};
 
 	/**
@@ -359,6 +373,8 @@ public:
 	 *  Failing this assumption will cause a check-fail. */
 	static void BatchSetEntityFragmentsValues(const FMassArchetypeEntityCollection& SparseEntities, TArrayView<const FInstancedStruct> FragmentInstanceList);
 
+	static void BatchSetEntityFragmentsValues(TConstArrayView<FMassArchetypeEntityCollection> EntityCollections, TArrayView<const FInstancedStruct> FragmentInstanceList);
+
 	// Return true if it is an valid built entity
 	bool IsEntityActive(FMassEntityHandle Entity) const 
 	{
@@ -650,6 +666,14 @@ private:
 	UE::Mass::FSingleThreadedEntityStorage& GetEntityStorageInterface();
 	const UE::Mass::FSingleThreadedEntityStorage& GetEntityStorageInterface() const;
 #endif
+
+	/** 
+	 * If ActiveCreationContext is not valid the function creates a new shared FEntityCreationContext instance and returns that.
+	 * Otherwise ActiveCreationContext will get extended with ReservedEntities and returned by the function.
+	 */
+	TSharedRef<FEntityCreationContext> GetOrMakeCreationContext(TConstArrayView<FMassEntityHandle> ReservedEntities = {});
+	bool IsDuringEntityCreation() const { return ActiveCreationContext.IsValid(); }
+	void DirtyCreationContext();
 	
 private:
 
@@ -660,6 +684,9 @@ private:
 		UE::Mass::FConcurrentEntityStorage>;
 	FEntityStorageContainerType EntityStorage;
 
+	/** Never access directly, use GetOrMakeCreationContext instead. */
+	TWeakPtr<FEntityCreationContext> ActiveCreationContext;
+
 	std::atomic<bool> bCommandBufferFlushingInProgress = false;
 	/**
 	 * This index will be enough to control which buffer is available for pushing commands since flashing is taking place 
diff --git a/Engine/Plugins/Runtime/MassEntity/Source/MassEntity/Public/MassExecutor.h b/Engine/Plugins/Runtime/MassEntity/Source/MassEntity/Public/MassExecutor.h
--- a/Engine/Plugins/Runtime/MassEntity/Source/MassEntity/Public/MassExecutor.h
+++ b/Engine/Plugins/Runtime/MassEntity/Source/MassEntity/Public/MassExecutor.h
@@ -35,7 +35,7 @@ namespace UE::Mass::Executor
 	MASSENTITY_API void RunSparse(FMassRuntimePipeline& RuntimePipeline, FMassProcessingContext& ProcessingContext, const FMassArchetypeEntityCollection& EntityCollection);
 
 	/** Executes given Processors array view. This function gets called under the hood by the rest of Run* functions */
-	MASSENTITY_API void RunProcessorsView(TArrayView<UMassProcessor* const> Processors, FMassProcessingContext& ProcessingContext, const FMassArchetypeEntityCollection* EntityCollection = nullptr);
+	MASSENTITY_API void RunProcessorsView(TArrayView<UMassProcessor* const> Processors, FMassProcessingContext& ProcessingContext, TConstArrayView<FMassArchetypeEntityCollection> EntityCollections = {});
 
 	/** 
 	 *  Triggers tasks executing Processor (and potentially it's children) and returns the task graph event representing 
@@ -44,4 +44,7 @@ namespace UE::Mass::Executor
 	 *    Note that OnDoneNotification will be executed on GameThread.
 	 */
 	MASSENTITY_API FGraphEventRef TriggerParallelTasks(UMassProcessor& Processor, FMassProcessingContext& ProcessingContext, TFunction<void()> OnDoneNotification);
+
+	UE_DEPRECATED(5.5, "This flavor of RunProcessorsView is deprecated. Use the one with TConstArrayView<FMassArchetypeEntityCollection> parameter instead.")
+	MASSENTITY_API void RunProcessorsView(TArrayView<UMassProcessor* const> Processors, FMassProcessingContext& ProcessingContext, const FMassArchetypeEntityCollection* EntityCollection);
 };
diff --git a/Engine/Plugins/Runtime/MassEntity/Source/MassEntity/Public/MassObserverManager.h b/Engine/Plugins/Runtime/MassEntity/Source/MassEntity/Public/MassObserverManager.h
--- a/Engine/Plugins/Runtime/MassEntity/Source/MassEntity/Public/MassObserverManager.h
+++ b/Engine/Plugins/Runtime/MassEntity/Source/MassEntity/Public/MassObserverManager.h
@@ -70,25 +70,36 @@ public:
 		return HasObserversForBitSet(Composition.Fragments, Operation) || HasObserversForBitSet(Composition.Tags, Operation);
 	}
 
+	/** @return whether any observers have been executed */
 	bool OnPostEntitiesCreated(const FMassArchetypeEntityCollection& EntityCollection);
-	bool OnPostEntitiesCreated(FMassProcessingContext& InProcessingContext, const FMassArchetypeEntityCollection& EntityCollection);
+	/** @return whether any observers have been executed */
+	bool OnPostEntitiesCreated(TConstArrayView<FMassArchetypeEntityCollection> EntityCollections);
+	/** @return whether any observers have been executed */
+	bool OnPostEntitiesCreated(FMassProcessingContext& InProcessingContext, TConstArrayView<FMassArchetypeEntityCollection> EntityCollections);
 
+	/** @return whether any observers have been executed */
 	bool OnPreEntitiesDestroyed(const FMassArchetypeEntityCollection& EntityCollection);
+	/** @return whether any observers have been executed */
 	bool OnPreEntitiesDestroyed(FMassProcessingContext& ProcessingContext, const FMassArchetypeEntityCollection& EntityCollection);
+	/** @return whether any observers have been executed */
 	bool OnPreEntityDestroyed(const FMassArchetypeCompositionDescriptor& ArchetypeComposition, const FMassEntityHandle Entity);
 
+	/** @return whether any observers have been executed */
 	bool OnCompositionChanged(const FMassEntityHandle Entity, const FMassArchetypeCompositionDescriptor& CompositionDelta, const EMassObservedOperation Operation);
+	/** @return whether any observers have been executed */
 	bool OnPostCompositionAdded(const FMassEntityHandle Entity, const FMassArchetypeCompositionDescriptor& Composition)
 	{
 		return OnCompositionChanged(Entity, Composition, EMassObservedOperation::Add);
 	}
+	/** @return whether any observers have been executed */
 	bool OnPreCompositionRemoved(const FMassEntityHandle Entity, const FMassArchetypeCompositionDescriptor& Composition)
 	{
 		return OnCompositionChanged(Entity, Composition, EMassObservedOperation::Remove);
 	}
 
+	/** @return whether any observers have been executed */
 	bool OnCompositionChanged(FMassProcessingContext& InProcessingContext, const FMassArchetypeEntityCollection& EntityCollection, const FMassArchetypeCompositionDescriptor& Composition, const EMassObservedOperation Operation);
-
+	/** @return whether any observers have been executed */
 	bool OnCompositionChanged(const FMassArchetypeEntityCollection& EntityCollection, const FMassArchetypeCompositionDescriptor& Composition, const EMassObservedOperation Operation)
 	{
 		FMassProcessingContext LocalContext(EntityManager, /*DeltaSeconds=*/0.f);
@@ -145,6 +156,9 @@ public:
 		}
 		return OnCompositionChanged(EntityCollection, Composition, Operation);
 	}
+
+	UE_DEPRECATED(5.5, "This flavor of OnPostEntitiesCreated is deprecated. Please use the one taking a TConstArrayView<FMassArchetypeEntityCollection> parameter instead")
+	bool OnPostEntitiesCreated(FMassProcessingContext& InProcessingContext, const FMassArchetypeEntityCollection& EntityCollection);
 };
 
 template<>
diff --git a/Engine/Plugins/Runtime/MassEntity/Source/MassEntityTestSuite/Private/MassSharedFragmentsTest.cpp b/Engine/Plugins/Runtime/MassEntity/Source/MassEntityTestSuite/Private/MassSharedFragmentsTest.cpp
--- a/Engine/Plugins/Runtime/MassEntity/Source/MassEntityTestSuite/Private/MassSharedFragmentsTest.cpp
+++ b/Engine/Plugins/Runtime/MassEntity/Source/MassEntityTestSuite/Private/MassSharedFragmentsTest.cpp
@@ -511,7 +511,7 @@ struct FSharedFragment_BatchSetAttempt : FEntityTestBase
 		SharedIntValues.AddConstSharedFragment(ConstSharedFragment);
 
 		TSharedRef<FMassEntityManager::FEntityCreationContext> CreationContext = EntityManager->BatchCreateEntities(FloatsArchetype, SharedIntValues, EntitiesToCreateNum, CreatedEntityHandles);
-		const FMassArchetypeHandle ResultingArchetype = CreationContext->GetEntityCollection().GetArchetype();
+		const FMassArchetypeHandle ResultingArchetype = CreationContext->GetEntityCollections()[0].GetArchetype();
 
 		FMassArchetypeEntityCollection EntityCollection(ResultingArchetype, CreatedEntityHandles, FMassArchetypeEntityCollection::EDuplicatesHandling::NoDuplicates);
 		
diff --git a/Engine/Plugins/Runtime/MassGameplay/Source/MassSpawner/Private/MassSpawner.cpp b/Engine/Plugins/Runtime/MassGameplay/Source/MassSpawner/Private/MassSpawner.cpp
--- a/Engine/Plugins/Runtime/MassGameplay/Source/MassSpawner/Private/MassSpawner.cpp
+++ b/Engine/Plugins/Runtime/MassGameplay/Source/MassSpawner/Private/MassSpawner.cpp
@@ -499,10 +499,7 @@ void AMassSpawner::SpawnGeneratedEntities(TConstArrayView<FMassEntitySpawnDataGe
 			// create entity collections and run Processors on them. 
 			TArray<FMassArchetypeEntityCollection> EntityCollections;
 			UE::Mass::Utils::CreateEntityCollections(EntityManager, AllEntities, FMassArchetypeEntityCollection::NoDuplicates, EntityCollections);
-			for (const FMassArchetypeEntityCollection& Collection : EntityCollections)
-			{
-				UE::Mass::Executor::RunProcessorsView(Processors, ProcessingContext, &Collection);
-			}
+			UE::Mass::Executor::RunProcessorsView(Processors, ProcessingContext, EntityCollections);
 		}
 	}
 
diff --git a/Engine/Plugins/Runtime/MassGameplay/Source/MassSpawner/Private/MassSpawnerSubsystem.cpp b/Engine/Plugins/Runtime/MassGameplay/Source/MassSpawner/Private/MassSpawnerSubsystem.cpp
--- a/Engine/Plugins/Runtime/MassGameplay/Source/MassSpawner/Private/MassSpawnerSubsystem.cpp
+++ b/Engine/Plugins/Runtime/MassGameplay/Source/MassSpawner/Private/MassSpawnerSubsystem.cpp
@@ -138,10 +138,11 @@ void UMassSpawnerSubsystem::DoSpawning(const FMassEntityTemplate& EntityTemplate
 	// 4. "OnEntitiesCreated" notifies will be sent out once the CreationContext gets destroyed (via its destructor).
 
 	TArray<FMassEntityHandle> SpawnedEntities;
-	TSharedRef<FMassEntityManager::FEntityCreationContext> CreationContext = EntityManager->BatchCreateEntities(EntityTemplate.GetArchetype(), EntityTemplate.GetSharedFragmentValues(), NumToSpawn, SpawnedEntities);
+	TSharedRef<FMassEntityManager::FEntityCreationContext> CreationContext
+		= EntityManager->BatchCreateEntities(EntityTemplate.GetArchetype(), EntityTemplate.GetSharedFragmentValues(), NumToSpawn, SpawnedEntities);
 
 	TConstArrayView<FInstancedStruct> FragmentInstances = EntityTemplate.GetInitialFragmentValues();
-	EntityManager->BatchSetEntityFragmentsValues(CreationContext->GetEntityCollection(), FragmentInstances);
+	EntityManager->BatchSetEntityFragmentsValues(CreationContext->GetEntityCollections(), FragmentInstances);
 	
 	UMassProcessor* SpawnDataInitializer = SpawnData.IsValid() ? GetSpawnDataInitializer(InitializerClass) : nullptr;
 
@@ -149,7 +150,7 @@ void UMassSpawnerSubsystem::DoSpawning(const FMassEntityTemplate& EntityTemplate
 	{
 		FMassProcessingContext ProcessingContext(EntityManager, /*TimeDelta=*/0.0f);
 		ProcessingContext.AuxData = SpawnData;
-		UE::Mass::Executor::RunProcessorsView(MakeArrayView(&SpawnDataInitializer, 1), ProcessingContext, &CreationContext->GetEntityCollection());
+		UE::Mass::Executor::RunProcessorsView(MakeArrayView(&SpawnDataInitializer, 1), ProcessingContext, CreationContext->GetEntityCollections());
 	}
 
 	OutEntities.Append(MoveTemp(SpawnedEntities));
