Fix for a nasty Mass crash that can happen when:
1. FEntityCreationContext shared pointer is being retained
2. While the context is "alive" modifying freshly created entities' composition
3. Some of the composition of the newly created entities is being observed.
When all three conditions are met the crash is almost certain since due to composition-changing operations the entity collection stored in the creation context gets out of data

The change boils down to invalidating the entity collection when entity mutating operations take place. Also, while the creation context is active, we don't trigger observers - same "on add" observers will get triggered when creation context is released.

The CL also contains new MassEntity creation/observers unit tests. Implementation of these tests surfaced one other issue with the code being tested:
- MassEntityManager's batch entity "building" functions bypass the whole creation context altogether. The CL fixes that as well.

diff --git a/Engine/Plugins/Runtime/MassEntity/Source/MassEntity/Private/MassEntityManager.cpp b/Engine/Plugins/Runtime/MassEntity/Source/MassEntity/Private/MassEntityManager.cpp
--- a/Engine/Plugins/Runtime/MassEntity/Source/MassEntity/Private/MassEntityManager.cpp
+++ b/Engine/Plugins/Runtime/MassEntity/Source/MassEntity/Private/MassEntityManager.cpp
@@ -13,6 +13,7 @@
 #include "Misc/Fork.h"
 #include "Misc/CoreDelegates.h"
 #include "Algo/Find.h"
+#include "MassEntityUtils.h"
 
 const FMassEntityHandle FMassEntityManager::InvalidEntity;
 
@@ -43,6 +44,89 @@ namespace UE::Mass::Private
 	constexpr int32 InvalidEntityIndex = 0;
 }
 
+//-----------------------------------------------------------------------------
+// FMassEntityManager::FEntityCreationContext
+//-----------------------------------------------------------------------------
+FMassEntityManager::FEntityCreationContext::FEntityCreationContext(FMassEntityManager& InManager, const TConstArrayView<FMassEntityHandle> InCreatedEntities)
+	: CreatedEntities(InCreatedEntities)
+{
+	Manager = InManager.AsShared();
+}
+
+FMassEntityManager::FEntityCreationContext::FEntityCreationContext(FMassEntityManager& InManager, const TConstArrayView<FMassEntityHandle> InCreatedEntities
+	, FMassArchetypeEntityCollection&& InEntityCollection)
+	: FEntityCreationContext(InManager, InCreatedEntities)
+{
+	checkf(InCreatedEntities.IsEmpty() || InEntityCollection.IsEmpty() == false, TEXT("Trying to create FEntityCreationContext instance with no entities but non-empty entity collection. This is not supported."))
+	if (InCreatedEntities.IsEmpty() == false)
+	{
+		EntityCollections.Add(MoveTemp(InEntityCollection));
+	}
+}
+
+FMassEntityManager::FEntityCreationContext::~FEntityCreationContext()
+{
+	if (((EntityCollections.IsEmpty() == false) || (CreatedEntities.IsEmpty() == false)) && ensure(Manager))
+	{
+		Manager->GetObserverManager().OnPostEntitiesCreated(GetEntityCollections());
+	}
+}
+
+TConstArrayView<FMassArchetypeEntityCollection> FMassEntityManager::FEntityCreationContext::GetEntityCollections() const
+{
+	// the EntityCollection has been dirtied, we need to rebuild it
+	if (IsDirty() && ensure(Manager))
+	{
+		UE::Mass::Utils::CreateEntityCollections(*Manager.Get(), CreatedEntities, CollectionCreationDuplicatesHandling, EntityCollections);
+	}
+
+	return EntityCollections;
+}
+
+void FMassEntityManager::FEntityCreationContext::MarkDirty()
+{
+	EntityCollections.Reset();
+}
+
+void FMassEntityManager::FEntityCreationContext::AppendEntities(const TConstArrayView<FMassEntityHandle> EntitiesToAppend)
+{
+	if (EntitiesToAppend.Num())
+	{
+		if (CreatedEntities.Num())
+		{
+			// since we already have entities in CreatedEntities (initially ensured to have no duplicates) we cannot 
+			// guarantee anymore that we'll have no duplicates after adding EntitiesToAppend
+			CollectionCreationDuplicatesHandling = FMassArchetypeEntityCollection::FoldDuplicates;
+			MarkDirty();
+		}
+		// else, if there are no entities the resulting state will bE "dirty" by design 
+		ensureMsgf(EntityCollections.IsEmpty(), TEXT("Having a non-empty array of entity collecitons is unexpected at this point!"));
+
+		CreatedEntities.Append(EntitiesToAppend);
+		ensure(IsDirty());
+	}
+}
+
+void FMassEntityManager::FEntityCreationContext::AppendEntities(const TConstArrayView<FMassEntityHandle> EntitiesToAppend, FMassArchetypeEntityCollection&& InEntityCollection)
+{
+	if (EntitiesToAppend.Num() == 0)
+	{
+		return;
+	}
+	
+	AppendEntities(EntitiesToAppend);
+
+	// this condition boils down to checking if this FEntityCreationContext instance only connects the just added EntitiesToAppend
+	if (CreatedEntities.Num() == EntitiesToAppend.Num())
+	{
+		checkf(EntityCollections.Num() == 0, TEXT("We never expect EntityCollections to be non-empty while there are no entities in CreatedEntities."));
+		EntityCollections.Add(MoveTemp(InEntityCollection));
+	}
+}
+
+//-----------------------------------------------------------------------------
+// FMassEntityManager
+//-----------------------------------------------------------------------------
 #if MASS_CONCURRENT_RESERVE
 UE::Mass::IEntityStorageInterface& FMassEntityManager::GetEntityStorageInterface()
 {
@@ -119,6 +203,9 @@ const UE::Mass::IEntityStorageInterface& FMassEntityManager::DebugGetEntityStora
 }
 #endif
 
+//-----------------------------------------------------------------------------
+// FMassEntityManager
+//-----------------------------------------------------------------------------
 FMassEntityManager::FMassEntityManager(UObject* InOwner)
 	: ObserverManager(*this)
 	, Owner(InOwner)
@@ -128,9 +215,6 @@ FMassEntityManager::FMassEntityManager(UObject* InOwner)
 #endif
 }
 
-//////////////////////////////////////////////////////////////////////
-// FMassEntityManager
-
 FMassEntityManager::~FMassEntityManager()
 {
 	if (bInitialized)
@@ -676,7 +760,7 @@ TConstArrayView<FMassEntityHandle> FMassEntityManager::BatchReserveEntities(cons
 	return MakeArrayView(InOutEntities.GetData() + Index, Count);
 }
 
-void FMassEntityManager::BatchBuildEntities(const FMassArchetypeEntityCollectionWithPayload& EncodedEntitiesWithPayload
+TSharedRef<FMassEntityManager::FEntityCreationContext> FMassEntityManager::BatchBuildEntities(const FMassArchetypeEntityCollectionWithPayload& EncodedEntitiesWithPayload
 	, const FMassFragmentBitSet& FragmentsAffected, const FMassArchetypeSharedFragmentValues& SharedFragmentValues, const FMassArchetypeCreationParams& CreationParams)
 {
 	checkf(IsProcessing() == false, TEXT("Synchronous API function %hs called during mass processing. Use asynchronous API instead."), __FUNCTION__);
@@ -692,10 +776,11 @@ void FMassEntityManager::BatchBuildEntities(const FMassArchetypeEntityCollection
 		Composition.SharedFragments.Add(*SharedFragment.GetScriptStruct());
 	}
 
-	BatchBuildEntities(EncodedEntitiesWithPayload, MoveTemp(Composition), SharedFragmentValues, CreationParams);
+	return BatchBuildEntities(EncodedEntitiesWithPayload, MoveTemp(Composition), SharedFragmentValues, CreationParams);
 }
 
-void FMassEntityManager::BatchBuildEntities(const FMassArchetypeEntityCollectionWithPayload& EncodedEntitiesWithPayload, FMassArchetypeCompositionDescriptor&& Composition
+TSharedRef<FMassEntityManager::FEntityCreationContext> FMassEntityManager::BatchBuildEntities(const FMassArchetypeEntityCollectionWithPayload& EncodedEntitiesWithPayload
+	, FMassArchetypeCompositionDescriptor&& Composition
 	, const FMassArchetypeSharedFragmentValues& SharedFragmentValues, const FMassArchetypeCreationParams& CreationParams)
 {
 	checkf(IsProcessing() == false, TEXT("Synchronous API function %hs called during mass processing. Use asynchronous API instead."), __FUNCTION__);
@@ -740,13 +825,15 @@ void FMassEntityManager::BatchBuildEntities(const FMassArchetypeEntityCollection
 		TargetArchetypeHandle.DataPtr->BatchSetFragmentValues(TargetArchetypeEntityRanges, EncodedEntitiesWithPayload.GetPayload());
 	}
 
-	if (ObserverManager.HasObserversForBitSet(Composition.Fragments, EMassObservedOperation::Add) || ObserverManager.HasObserversForBitSet(Composition.Tags, EMassObservedOperation::Add))
-	{
-		ObserverManager.OnCompositionChanged(
-			FMassArchetypeEntityCollection(TargetArchetypeHandle, MoveTemp(TargetArchetypeEntityRanges))
-			, Composition
-			, EMassObservedOperation::Add);
-	}
+	// With this call we're either creating a fresh context populated with EntityHandles, or it will append 
+	// EntityHandles to active context.
+	// Not creating the context sooner since we want to reuse TargetArchetypeEntityRanges by moving it over to the context.
+	// Note that we can afford to create this context so late since all previous operations were on the archetype level
+	// and as such won't cause observers triggering (which usually is prevented by context's existence), and that we 
+	// strongly assume the all entity creation/building (not to be mistaken with "reserving") takes place in a single thread
+	// @todo add checks/ensures enforcing the assumption mentioned above.
+	return GetOrMakeCreationContext(EntityHandles
+		, FMassArchetypeEntityCollection(TargetArchetypeHandle, MoveTemp(TargetArchetypeEntityRanges)));
 }
 
 TSharedRef<FMassEntityManager::FEntityCreationContext> FMassEntityManager::BatchCreateReservedEntities(const FMassArchetypeHandle& ArchetypeHandle
@@ -768,7 +855,7 @@ TSharedRef<FMassEntityManager::FEntityCreationContext> FMassEntityManager::Batch
 	TRACE_CPUPROFILER_EVENT_SCOPE(Mass_BatchCreateEntities);
 
 	checkf(IsProcessing() == false, TEXT("Synchronous API function %hs called during mass processing. Use asynchronous API instead."), __FUNCTION__);
-	testableCheckfReturn(ArchetypeHandle.IsValid(), return MakeShareable(new FEntityCreationContext(0))
+	testableCheckfReturn(ArchetypeHandle.IsValid(), return GetOrMakeCreationContext()
 		, TEXT("%hs expecting a valid ArchetypeHandle"), __FUNCTION__);
 
 	TConstArrayView<FMassEntityHandle> ReservedEntities = BatchReserveEntities(Count, InOutEntities);
@@ -797,18 +884,7 @@ TSharedRef<FMassEntityManager::FEntityCreationContext> FMassEntityManager::Inter
 	FMassArchetypeEntityCollection::FEntityRangeArray TargetArchetypeEntityRanges;
 	ArchetypeData->BatchAddEntities(ReservedEntities, SharedFragmentValues, TargetArchetypeEntityRanges);
 
-	FEntityCreationContext* CreationContext = new FEntityCreationContext(ReservedEntities.Num());
-	new (&CreationContext->EntityCollection) FMassArchetypeEntityCollection(ArchetypeHandle, MoveTemp(TargetArchetypeEntityRanges));
-
-	if (ObserverManager.HasObserversForBitSet(ArchetypeData->GetCompositionDescriptor().Fragments, EMassObservedOperation::Add)
-		|| ObserverManager.HasObserversForBitSet(ArchetypeData->GetCompositionDescriptor().Tags, EMassObservedOperation::Add))
-	{
-		CreationContext->OnSpawningFinished = [this](FEntityCreationContext& Context) {
-			ObserverManager.OnPostEntitiesCreated(Context.EntityCollection);
-		};
-	}
-
-	return MakeShareable(CreationContext);
+	return GetOrMakeCreationContext(ReservedEntities, FMassArchetypeEntityCollection(ArchetypeHandle, MoveTemp(TargetArchetypeEntityRanges)));
 }
 
 void FMassEntityManager::DestroyEntity(FMassEntityHandle Entity)
@@ -833,6 +909,7 @@ void FMassEntityManager::BatchDestroyEntities(TConstArrayView<FMassEntityHandle>
 	TRACE_CPUPROFILER_EVENT_SCOPE(Mass_BatchDestroyEntities);
 
 	checkf(IsProcessing() == false, TEXT("Synchronous API function %hs called during mass processing. Use asynchronous API instead."), __FUNCTION__);
+	checkf(IsDuringEntityCreation() == false, TEXT("%hs: Trying to destroy entities while entity creation is under way. This operation is not supported."), __FUNCTION__);
 
 	// @todo optimize, we can make savings by implementing Archetype->RemoveEntities()
 	for (const FMassEntityHandle Entity : InEntities)
@@ -862,6 +939,7 @@ void FMassEntityManager::BatchDestroyEntityChunks(const FMassArchetypeEntityColl
 	TRACE_CPUPROFILER_EVENT_SCOPE(Mass_BatchDestroyEntityChunks);
 
 	checkf(IsProcessing() == false, TEXT("Synchronous API function %hs called during mass processing. Use asynchronous API instead."), __FUNCTION__);
+	checkf(IsDuringEntityCreation() == false, TEXT("%hs: Trying to destroy entities while entity creation is under way. This operation is not supported."), __FUNCTION__);
 
 	TArray<FMassEntityHandle> EntitiesRemoved;
 	// note that it's important to place the context instance in the same scope as the loop below that updates 
@@ -896,7 +974,11 @@ void FMassEntityManager::AddFragmentToEntity(FMassEntityHandle Entity, const USc
 	CheckIfEntityIsActive(Entity);
 
 	const FMassArchetypeCompositionDescriptor Descriptor(InternalAddFragmentListToEntityChecked(Entity, FMassFragmentBitSet(*FragmentType)));
-	ObserverManager.OnPostCompositionAdded(Entity, Descriptor);
+
+	if (IsAllowedToTriggerObservers())
+	{
+		ObserverManager.OnPostCompositionAdded(Entity, Descriptor);
+	}
 }
 
 void FMassEntityManager::AddFragmentToEntity(FMassEntityHandle Entity, const UScriptStruct* FragmentType, const FStructInitializationCallback& Initializer)
@@ -913,7 +995,11 @@ void FMassEntityManager::AddFragmentToEntity(FMassEntityHandle Entity, const USc
 	Initializer(FragmentData, *FragmentType);
 
 	const FMassArchetypeCompositionDescriptor Descriptor(MoveTemp(Fragments));
-	ObserverManager.OnPostCompositionAdded(Entity, Descriptor);
+	
+	if (IsAllowedToTriggerObservers())
+	{
+		ObserverManager.OnPostCompositionAdded(Entity, Descriptor);
+	}
 }
 
 void FMassEntityManager::AddFragmentListToEntity(FMassEntityHandle Entity, TConstArrayView<const UScriptStruct*> FragmentList)
@@ -921,7 +1007,11 @@ void FMassEntityManager::AddFragmentListToEntity(FMassEntityHandle Entity, TCons
 	CheckIfEntityIsActive(Entity);
 
 	const FMassArchetypeCompositionDescriptor Descriptor(InternalAddFragmentListToEntityChecked(Entity, FMassFragmentBitSet(FragmentList)));
-	ObserverManager.OnPostCompositionAdded(Entity, Descriptor);
+	
+	if (IsAllowedToTriggerObservers())
+	{
+		ObserverManager.OnPostCompositionAdded(Entity, Descriptor);
+	}
 }
 
 void FMassEntityManager::AddCompositionToEntity_GetDelta(FMassEntityHandle Entity, FMassArchetypeCompositionDescriptor& InDescriptor)
@@ -953,7 +1043,10 @@ void FMassEntityManager::AddCompositionToEntity_GetDelta(FMassEntityHandle Entit
 
 			GetEntityStorageInterface().SetArchetypeFromShared(Entity.Index, NewArchetypeHandle.DataPtr);
 
-			ObserverManager.OnPostCompositionAdded(Entity, InDescriptor);
+			if (IsAllowedToTriggerObservers())
+			{
+				ObserverManager.OnPostCompositionAdded(Entity, InDescriptor);
+			}
 		}
 	}
 }
@@ -977,7 +1070,11 @@ void FMassEntityManager::RemoveCompositionFromEntity(FMassEntityHandle Entity, c
 		if (NewDescriptor.IsEquivalent(OldArchetype->GetCompositionDescriptor()) == false)
 		{
 			ensureMsgf(OldArchetype->GetCompositionDescriptor().HasAll(InDescriptor), TEXT("Some of the elements being removed are already missing from entity\'s composition."));
-			ObserverManager.OnPreCompositionRemoved(Entity, InDescriptor);
+			
+			if (IsAllowedToTriggerObservers())
+			{
+				ObserverManager.OnPreCompositionRemoved(Entity, InDescriptor);
+			}
 
 			const FMassArchetypeHandle NewArchetypeHandle = CreateArchetype(NewDescriptor, FMassArchetypeCreationParams(*OldArchetype));
 
@@ -1005,7 +1102,10 @@ void FMassEntityManager::InternalBuildEntity(FMassEntityHandle Entity, const FMa
 	GetEntityStorageInterface().SetArchetypeFromShared(Entity.Index, ArchetypeHandle.DataPtr);
 	NewArchetype->AddEntity(Entity, SharedFragmentValues);
 
-	ObserverManager.OnPostCompositionAdded(Entity, NewArchetype->GetCompositionDescriptor());
+	if (IsAllowedToTriggerObservers())
+	{
+		ObserverManager.OnPostCompositionAdded(Entity, NewArchetype->GetCompositionDescriptor());
+	}
 }
 
 void FMassEntityManager::InternalReleaseEntity(FMassEntityHandle Entity)
@@ -1063,7 +1163,10 @@ void FMassEntityManager::AddFragmentInstanceListToEntity(FMassEntityHandle Entit
 	check(CurrentArchetype);
 	CurrentArchetype->SetFragmentsData(Entity, FragmentInstanceList);
 
-	ObserverManager.OnPostCompositionAdded(Entity, Descriptor);
+	if (IsAllowedToTriggerObservers())
+	{
+		ObserverManager.OnPostCompositionAdded(Entity, Descriptor);
+	}
 }
 
 void FMassEntityManager::RemoveFragmentFromEntity(FMassEntityHandle Entity, const UScriptStruct* FragmentType)
@@ -1095,7 +1198,11 @@ void FMassEntityManager::RemoveFragmentListFromEntity(FMassEntityHandle Entity,
 		FMassArchetypeCompositionDescriptor CompositionDelta;
 		// Find overlap.  It isn't guaranteed that the old archetype has all of the fragments being removed.
 		CompositionDelta.Fragments = OldArchetype->GetFragmentBitSet().GetOverlap(FragmentsToRemove);
-		ObserverManager.OnPreCompositionRemoved(Entity, CompositionDelta);
+
+		if (IsAllowedToTriggerObservers())
+		{
+			ObserverManager.OnPreCompositionRemoved(Entity, CompositionDelta);
+		}
 
 		// Move the entity over
 		FMassArchetypeData& NewArchetype = FMassArchetypeHelper::ArchetypeDataFromHandleChecked(NewArchetypeHandle);
@@ -1159,7 +1266,11 @@ void FMassEntityManager::AddTagToEntity(FMassEntityHandle Entity, const UScriptS
 		FMassTagBitSet TagDelta;
 		TagDelta.Add(*TagType);
 		CompositionDelta.Tags = TagDelta;
-		ObserverManager.OnPostCompositionAdded(Entity, CompositionDelta);
+		
+		if (IsAllowedToTriggerObservers())
+		{
+			ObserverManager.OnPostCompositionAdded(Entity, CompositionDelta);
+		}
 	}
 }
 	
@@ -1178,7 +1289,11 @@ void FMassEntityManager::RemoveTagFromEntity(FMassEntityHandle Entity, const USc
 		FMassTagBitSet TagDelta;
 		TagDelta.Add(*TagType);
 		CompositionDelta.Tags = TagDelta;
-		ObserverManager.OnPreCompositionRemoved(Entity, CompositionDelta);
+		
+		if (IsAllowedToTriggerObservers())
+		{
+			ObserverManager.OnPreCompositionRemoved(Entity, CompositionDelta);
+		}
 		
 		// CurrentArchetype->GetTagBitSet() -  *TagType
 		const FMassTagBitSet NewTagComposition = CurrentArchetype->GetTagBitSet() - TagDelta;
@@ -1243,6 +1358,8 @@ void FMassEntityManager::BatchChangeTagsForEntities(TConstArrayView<FMassArchety
 {
 	TRACE_CPUPROFILER_EVENT_SCOPE(Mass_BatchChangeTagsForEntities);
 
+	const FScopedCreationContextOperations CreationContextOperations(*this);
+	
 	for (const FMassArchetypeEntityCollection& Collection : EntityCollections)
 	{
 		FMassArchetypeData* CurrentArchetype = Collection.GetArchetype().DataPtr.Get();
@@ -1255,8 +1372,10 @@ void FMassEntityManager::BatchChangeTagsForEntities(TConstArrayView<FMassArchety
 			FMassTagBitSet TagsAdded = TagsToAdd - CurrentArchetype->GetTagBitSet();
 			FMassTagBitSet TagsRemoved = TagsToRemove.GetOverlap(CurrentArchetype->GetTagBitSet());
 
-			if (ObserverManager.HasObserversForBitSet(TagsRemoved, EMassObservedOperation::Remove))
+			if (CreationContextOperations.IsAllowedToTriggerObservers()
+				&& ObserverManager.HasObserversForBitSet(TagsRemoved, EMassObservedOperation::Remove))
 			{
+				// @todo should use OnPreCompositionRemoved here instead, but we're missing a FMassArchetypeEntityCollection version
 				ObserverManager.OnCompositionChanged(Collection, FMassArchetypeCompositionDescriptor(MoveTemp(TagsRemoved)), EMassObservedOperation::Remove);
 			}
 			const bool bTagsAddedAreObserved = ObserverManager.HasObserversForBitSet(TagsAdded, EMassObservedOperation::Add);
@@ -1277,8 +1396,9 @@ void FMassEntityManager::BatchChangeTagsForEntities(TConstArrayView<FMassArchety
 				GetEntityStorageInterface().SetArchetypeFromShared(Entity.Index, NewArchetypeHandle.DataPtr);
 			}
 
-			if (bTagsAddedAreObserved)
+			if (bTagsAddedAreObserved && CreationContextOperations.IsAllowedToTriggerObservers())
 			{
+				// @todo should use OnPostCompositionAdded here instead, but we're missing a FMassArchetypeEntityCollection version
 				ObserverManager.OnCompositionChanged(
 					FMassArchetypeEntityCollection(NewArchetypeHandle, MoveTemp(NewArchetypeEntityRanges))
 					, FMassArchetypeCompositionDescriptor(MoveTemp(TagsAdded))
@@ -1292,6 +1412,8 @@ void FMassEntityManager::BatchChangeFragmentCompositionForEntities(TConstArrayVi
 {
 	TRACE_CPUPROFILER_EVENT_SCOPE(Mass_BatchChangeFragmentCompositionForEntities);
 
+	FScopedCreationContextOperations CreationContextOperations(*this);
+
 	for (const FMassArchetypeEntityCollection& Collection : EntityCollections)
 	{
 		FMassArchetypeData* CurrentArchetype = Collection.GetArchetype().DataPtr.Get();
@@ -1306,7 +1428,9 @@ void FMassEntityManager::BatchChangeFragmentCompositionForEntities(TConstArrayVi
 				FMassFragmentBitSet FragmentsAdded = FragmentsToAdd - CurrentArchetype->GetFragmentBitSet();
 				const bool bFragmentsAddedAreObserved = ObserverManager.HasObserversForBitSet(FragmentsAdded, EMassObservedOperation::Add);
 				FMassFragmentBitSet FragmentsRemoved = FragmentsToRemove.GetOverlap(CurrentArchetype->GetFragmentBitSet());
-				if (ObserverManager.HasObserversForBitSet(FragmentsRemoved, EMassObservedOperation::Remove))
+				
+				if (CreationContextOperations.IsAllowedToTriggerObservers()
+					&& ObserverManager.HasObserversForBitSet(FragmentsRemoved, EMassObservedOperation::Remove))
 				{
 					ObserverManager.OnCompositionChanged(Collection, FMassArchetypeCompositionDescriptor(MoveTemp(FragmentsRemoved)), EMassObservedOperation::Remove);
 				}
@@ -1327,7 +1451,7 @@ void FMassEntityManager::BatchChangeFragmentCompositionForEntities(TConstArrayVi
 					GetEntityStorageInterface().SetArchetypeFromShared(Entity.Index, NewArchetypeHandle.DataPtr);
 				}
 
-				if (bFragmentsAddedAreObserved)
+				if (bFragmentsAddedAreObserved && CreationContextOperations.IsAllowedToTriggerObservers())
 				{
 					ObserverManager.OnCompositionChanged(
 						FMassArchetypeEntityCollection(NewArchetypeHandle, MoveTemp(NewArchetypeEntityRanges))
@@ -1358,6 +1482,8 @@ void FMassEntityManager::BatchAddFragmentInstancesForEntities(TConstArrayView<FM
 	//		archetype. This could be utilized as a basis for optimization. To be investigated.
 	// 
 
+	FScopedCreationContextOperations CreationContextOperations(*this);
+
 	for (const FMassArchetypeEntityCollectionWithPayload& EntityRangesWithPayload : EntityCollections)
 	{
 		FMassArchetypeHandle TargetArchetypeHandle = EntityRangesWithPayload.GetEntityCollection().GetArchetype();
@@ -1403,8 +1529,8 @@ void FMassEntityManager::BatchAddFragmentInstancesForEntities(TConstArrayView<FM
 			// note that even though the "subchunk" information could have changed the order of entities is the same and 
 			// corresponds to the order in FMassArchetypeEntityCollectionWithPayload's payload
 			TargetArchetypeHandle.DataPtr->BatchSetFragmentValues(TargetArchetypeEntityRanges, EntityRangesWithPayload.GetPayload());
-
-			if (bFragmentsAddedAreObserved)
+			
+			if (bFragmentsAddedAreObserved && CreationContextOperations.IsAllowedToTriggerObservers())
 			{
 				ObserverManager.OnCompositionChanged(
 					FMassArchetypeEntityCollection(TargetArchetypeHandle, MoveTemp(TargetArchetypeEntityRanges))
@@ -1479,14 +1605,30 @@ void FMassEntityManager::SetEntityFragmentsValues(FMassEntityHandle Entity, TArr
 
 void FMassEntityManager::BatchSetEntityFragmentsValues(const FMassArchetypeEntityCollection& SparseEntities, TArrayView<const FInstancedStruct> FragmentInstanceList)
 {
-	TRACE_CPUPROFILER_EVENT_SCOPE(Mass_BatchSetEntityFragmentsValues);
+	if (FragmentInstanceList.Num())
+	{
+		BatchSetEntityFragmentsValues(MakeArrayView(&SparseEntities, 1), FragmentInstanceList);
+	}
+}
 
-	FMassArchetypeData* Archetype = SparseEntities.GetArchetype().DataPtr.Get();
-	check(Archetype);
+void FMassEntityManager::BatchSetEntityFragmentsValues(TConstArrayView<FMassArchetypeEntityCollection> EntityCollections, TArrayView<const FInstancedStruct> FragmentInstanceList)
+{
+	if (FragmentInstanceList.IsEmpty())
+	{
+		return;
+	}
 
-	for (const FInstancedStruct& FragmentTemplate : FragmentInstanceList)
+	for (const FMassArchetypeEntityCollection& SparseEntities : EntityCollections)
 	{
-		Archetype->SetFragmentData(SparseEntities.GetRanges(), FragmentTemplate);
+		TRACE_CPUPROFILER_EVENT_SCOPE(Mass_BatchSetEntityFragmentsValues);
+
+		FMassArchetypeData* Archetype = SparseEntities.GetArchetype().DataPtr.Get();
+		check(Archetype);
+
+		for (const FInstancedStruct& FragmentTemplate : FragmentInstanceList)
+		{
+			Archetype->SetFragmentData(SparseEntities.GetRanges(), FragmentTemplate);
+		}
 	}
 }
 
@@ -1722,6 +1864,62 @@ void FMassEntityManager::AppendCommands(TSharedPtr<FMassCommandBuffer>& InOutCom
 	Defer().MoveAppend(*InOutCommandBuffer.Get());
 }
 
+TSharedRef<FMassEntityManager::FEntityCreationContext> FMassEntityManager::GetOrMakeCreationContext()
+{
+	if (ActiveCreationContext.IsValid())
+	{
+		return ActiveCreationContext.Pin().ToSharedRef();
+	}
+	else
+	{
+		FEntityCreationContext* CreationContext = new FEntityCreationContext(*this);
+		TSharedRef<FEntityCreationContext> SharedContext = MakeShareable(CreationContext);
+		ActiveCreationContext = SharedContext;
+		return SharedContext;
+	}
+}
+
+TSharedRef<FMassEntityManager::FEntityCreationContext> FMassEntityManager::GetOrMakeCreationContext(TConstArrayView<FMassEntityHandle> ReservedEntities
+	, FMassArchetypeEntityCollection&& EntityCollection)
+{
+	if (ActiveCreationContext.IsValid())
+	{
+		TSharedPtr<FEntityCreationContext> SharedContext = ActiveCreationContext.Pin();
+		CA_ASSUME(SharedContext);
+		SharedContext->AppendEntities(ReservedEntities, MoveTemp(EntityCollection));
+		return SharedContext.ToSharedRef();
+	}
+	else
+	{
+		FEntityCreationContext* CreationContext = new FEntityCreationContext(*this, ReservedEntities, MoveTemp(EntityCollection));
+		TSharedRef<FEntityCreationContext> SharedContext = MakeShareable(CreationContext);
+		ActiveCreationContext = SharedContext;
+		return SharedContext;
+	}
+}
+
+bool FMassEntityManager::DirtyCreationContext()
+{
+	if (TSharedPtr<FEntityCreationContext> AsSharedPtr = ActiveCreationContext.Pin())
+	{
+		AsSharedPtr->MarkDirty();
+		return true;
+	}
+	return false;
+}
+
+bool FMassEntityManager::DebugDoCollectionsOverlapCreationContext(TConstArrayView<FMassArchetypeEntityCollection> EntityCollections) const
+{
+	if (TSharedPtr<FEntityCreationContext> AsSharedPtr = ActiveCreationContext.Pin())
+	{
+		TConstArrayView<FMassArchetypeEntityCollection> CreationCollections = AsSharedPtr->EntityCollections;
+		return CreationCollections.GetData() <= EntityCollections.GetData()
+			&& EntityCollections.GetData() <= CreationCollections.GetData() + CreationCollections.Num();
+	}
+
+	return false;
+}
+
 void FMassEntityManager::SetDebugName(const FString& NewDebugGame) 
 { 
 #if WITH_MASSENTITY_DEBUG
@@ -1917,4 +2115,9 @@ void FMassEntityManager::BatchBuildEntities(const FMassArchetypeEntityCollection
 	FMassArchetypeCreationParams Params;
 	Params.DebugName = ArchetypeDebugName;
 	BatchBuildEntities(EncodedEntitiesWithPayload, MoveTemp(Composition), SharedFragmentValues, Params);
+}
+const FMassArchetypeEntityCollection& FMassEntityManager::FEntityCreationContext::GetEntityCollection() const
+{
+	static FMassArchetypeEntityCollection EmptyCollection;
+	return EntityCollections.Num() ? EntityCollections[0] : EmptyCollection;
 }
\ No newline at end of file
diff --git a/Engine/Plugins/Runtime/MassEntity/Source/MassEntity/Private/MassExecutor.cpp b/Engine/Plugins/Runtime/MassEntity/Source/MassEntity/Private/MassExecutor.cpp
--- a/Engine/Plugins/Runtime/MassEntity/Source/MassEntity/Private/MassExecutor.cpp
+++ b/Engine/Plugins/Runtime/MassEntity/Source/MassEntity/Private/MassExecutor.cpp
@@ -9,6 +9,15 @@
 
 namespace UE::Mass::Executor
 {
+
+FORCEINLINE void ExecuteProcessors(FMassEntityManager& EntityManager, TArrayView<UMassProcessor* const> Processors, FMassExecutionContext& ExecutionContext)
+{
+	for (UMassProcessor* Proc : Processors)
+	{
+		Proc->CallExecute(EntityManager, ExecutionContext);
+	}
+}
+
 void Run(FMassRuntimePipeline& RuntimePipeline, FMassProcessingContext& ProcessingContext)
 {
 	if (!ensure(ProcessingContext.EntityManager) || 
@@ -35,7 +44,7 @@ void RunSparse(FMassRuntimePipeline& RuntimePipeline, FMassProcessingContext& Pr
 	TRACE_CPUPROFILER_EVENT_SCOPE_STR("MassExecutor RunSparseEntities");
 
 	const FMassArchetypeEntityCollection EntityCollection(Archetype, Entities, FMassArchetypeEntityCollection::NoDuplicates);
-	RunProcessorsView(RuntimePipeline.GetMutableProcessors(), ProcessingContext, &EntityCollection);
+	RunProcessorsView(RuntimePipeline.GetMutableProcessors(), ProcessingContext, MakeArrayView(&EntityCollection, 1));
 }
 
 void RunSparse(FMassRuntimePipeline& RuntimePipeline, FMassProcessingContext& ProcessingContext, const FMassArchetypeEntityCollection& EntityCollection)
@@ -50,7 +59,7 @@ void RunSparse(FMassRuntimePipeline& RuntimePipeline, FMassProcessingContext& Pr
 
 	TRACE_CPUPROFILER_EVENT_SCOPE_STR("MassExecutor RunSparse");
 
-	RunProcessorsView(RuntimePipeline.GetMutableProcessors(), ProcessingContext, &EntityCollection);
+	RunProcessorsView(RuntimePipeline.GetMutableProcessors(), ProcessingContext, MakeArrayView(&EntityCollection, 1));
 }
 
 void Run(UMassProcessor& Processor, FMassProcessingContext& ProcessingContext)
@@ -66,7 +75,7 @@ void Run(UMassProcessor& Processor, FMassProcessingContext& ProcessingContext)
 	RunProcessorsView(MakeArrayView(&ProcPtr, 1), ProcessingContext);
 }
 
-void RunProcessorsView(TArrayView<UMassProcessor* const> Processors, FMassProcessingContext& ProcessingContext, const FMassArchetypeEntityCollection* EntityCollection)
+void RunProcessorsView(TArrayView<UMassProcessor* const> Processors, FMassProcessingContext& ProcessingContext, TConstArrayView<FMassArchetypeEntityCollection> EntityCollections)
 {
 	TRACE_CPUPROFILER_EVENT_SCOPE(RunProcessorsView);
 
@@ -86,11 +95,7 @@ void RunProcessorsView(TArrayView<UMassProcessor* const> Processors, FMassProces
 	TRACE_CPUPROFILER_EVENT_SCOPE_STR("MassExecutor RunProcessorsView")
 
 	FMassExecutionContext ExecutionContext(*ProcessingContext.EntityManager.Get(), ProcessingContext.DeltaSeconds);
-	if (EntityCollection)
-	{
-		ExecutionContext.SetEntityCollection(*EntityCollection);
-	}
-	
+
 	// if ProcessingContext points at a valid CommandBuffer use that one, otherwise manually create a new command buffer 
 	// to let the default one still be used by code unaware of mass processing
 	TSharedPtr<FMassCommandBuffer> CommandBuffer = ProcessingContext.CommandBuffer 
@@ -106,9 +111,18 @@ void RunProcessorsView(TArrayView<UMassProcessor* const> Processors, FMassProces
 		
 		FMassEntityManager::FScopedProcessing ProcessingScope = EntityManager.NewProcessingScope();
 
-		for (UMassProcessor* Proc : Processors)
+		if (EntityCollections.Num() == 0)
 		{
-			Proc->CallExecute(*ProcessingContext.EntityManager, ExecutionContext);
+			ExecuteProcessors(*ProcessingContext.EntityManager, Processors, ExecutionContext);
+		}
+		else
+		{
+			// @todo change ExecutionContext to contain a TConstArrayView of collections. Will lead up to changes in the entity query as well.
+			for (const FMassArchetypeEntityCollection& Collection : EntityCollections)
+			{
+				ExecutionContext.SetEntityCollection(Collection);
+				ExecuteProcessors(*ProcessingContext.EntityManager, Processors, ExecutionContext);
+			}
 		}
 	}
 	
@@ -212,4 +226,19 @@ FGraphEventRef TriggerParallelTasks(UMassProcessor& Processor, FMassProcessingCo
 	return CompletionEvent;
 }
 
+//-----------------------------------------------------------------------------
+// DEPRECATED
+//-----------------------------------------------------------------------------
+void RunProcessorsView(TArrayView<UMassProcessor* const> Processors, FMassProcessingContext& ProcessingContext, const FMassArchetypeEntityCollection* EntityCollection)
+{
+	if (EntityCollection)
+	{
+		RunProcessorsView(Processors, ProcessingContext, MakeArrayView(EntityCollection, 1));
+	}
+	else
+	{
+		RunProcessorsView(Processors, ProcessingContext);
+	}
+}
+
 } // namespace UE::Mass::Executor
diff --git a/Engine/Plugins/Runtime/MassEntity/Source/MassEntity/Private/MassObserverManager.cpp b/Engine/Plugins/Runtime/MassEntity/Source/MassEntity/Private/MassObserverManager.cpp
--- a/Engine/Plugins/Runtime/MassEntity/Source/MassEntity/Private/MassObserverManager.cpp
+++ b/Engine/Plugins/Runtime/MassEntity/Source/MassEntity/Private/MassObserverManager.cpp
@@ -83,6 +83,11 @@ void FMassObserverManager::Initialize()
 }
 
 bool FMassObserverManager::OnPostEntitiesCreated(const FMassArchetypeEntityCollection& EntityCollection)
+{
+	return OnPostEntitiesCreated(MakeArrayView(&EntityCollection , 1));
+}
+
+bool FMassObserverManager::OnPostEntitiesCreated(TConstArrayView<FMassArchetypeEntityCollection> EntityCollections)
 {
 	FMassProcessingContext ProcessingContext(EntityManager, /*DeltaSeconds=*/0.f);
 	// requesting not to flush commands since handling creation of new entities can result in multiple collections of
@@ -91,7 +96,7 @@ bool FMassObserverManager::OnPostEntitiesCreated(const FMassArchetypeEntityColle
 	ProcessingContext.bFlushCommandBuffer = false;
 	ProcessingContext.CommandBuffer = MakeShareable(new FMassCommandBuffer());
 
-	if (OnPostEntitiesCreated(ProcessingContext, EntityCollection))
+	if (OnPostEntitiesCreated(ProcessingContext, EntityCollections))
 	{
 		EntityManager.FlushCommands(ProcessingContext.CommandBuffer);
 		return true;
@@ -99,14 +104,27 @@ bool FMassObserverManager::OnPostEntitiesCreated(const FMassArchetypeEntityColle
 	return false;
 }
 
+// DEPRECATED
 bool FMassObserverManager::OnPostEntitiesCreated(FMassProcessingContext& ProcessingContext, const FMassArchetypeEntityCollection& EntityCollection)
+{
+	return OnPostEntitiesCreated(ProcessingContext, MakeArrayView(&EntityCollection, 1));
+}
+
+bool FMassObserverManager::OnPostEntitiesCreated(FMassProcessingContext& ProcessingContext, TConstArrayView<FMassArchetypeEntityCollection> EntityCollections)
 {
 	TRACE_CPUPROFILER_EVENT_SCOPE_STR("OnPostEntitiesCreated")
 
 	check(ProcessingContext.EntityManager);
-	const FMassArchetypeCompositionDescriptor& ArchetypeComposition = ProcessingContext.EntityManager->GetArchetypeComposition(EntityCollection.GetArchetype());
 
-	return OnCompositionChanged(ProcessingContext, EntityCollection, ArchetypeComposition, EMassObservedOperation::Add);
+	bool bReturnValue = false;
+
+	for (const FMassArchetypeEntityCollection& Collection : EntityCollections)
+	{
+		const FMassArchetypeCompositionDescriptor& ArchetypeComposition = ProcessingContext.EntityManager->GetArchetypeComposition(Collection.GetArchetype());
+		bReturnValue = OnCompositionChanged(ProcessingContext, Collection, ArchetypeComposition, EMassObservedOperation::Add) || bReturnValue;
+	}
+
+	return bReturnValue;
 }
 
 bool FMassObserverManager::OnPreEntitiesDestroyed(const FMassArchetypeEntityCollection& EntityCollection)
@@ -240,7 +258,7 @@ void FMassObserverManager::HandleFragmentsImpl(FMassProcessingContext& Processin
 		ProcessingContext.AuxData.InitializeAs(Type);
 		FMassRuntimePipeline& Pipeline = (*HandlersContainer).FindChecked(Type);
 
-		UE::Mass::Executor::RunProcessorsView(Pipeline.GetMutableProcessors(), ProcessingContext, &EntityCollection);
+		UE::Mass::Executor::RunProcessorsView(Pipeline.GetMutableProcessors(), ProcessingContext, MakeArrayView(&EntityCollection, 1));
 	}
 }
 
@@ -253,7 +271,7 @@ void FMassObserverManager::HandleSingleEntityImpl(const UScriptStruct& FragmentT
 	ProcessingContext.AuxData.InitializeAs(&FragmentType);
 	FMassRuntimePipeline& Pipeline = (*HandlersContainer).FindChecked(&FragmentType);
 
-	UE::Mass::Executor::RunProcessorsView(Pipeline.GetMutableProcessors(), ProcessingContext, &EntityCollection);
+	UE::Mass::Executor::RunProcessorsView(Pipeline.GetMutableProcessors(), ProcessingContext, MakeArrayView(&EntityCollection, 1));
 }
 
 void FMassObserverManager::AddObserverInstance(const UScriptStruct& FragmentOrTagType, const EMassObservedOperation Operation, UMassProcessor& ObserverProcessor)
diff --git a/Engine/Plugins/Runtime/MassEntity/Source/MassEntity/Public/MassEntityManager.h b/Engine/Plugins/Runtime/MassEntity/Source/MassEntity/Public/MassEntityManager.h
--- a/Engine/Plugins/Runtime/MassEntity/Source/MassEntity/Public/MassEntityManager.h
+++ b/Engine/Plugins/Runtime/MassEntity/Source/MassEntity/Public/MassEntityManager.h
@@ -208,21 +208,51 @@ public:
 
 	/**
 	 * A dedicated structure for ensuring the "on entities creation" observers get notified only once all other 
-	 * initialization operations are done and this creation context instance gets released. */
-	struct FEntityCreationContext
+	 * initialization operations are done and this creation context instance gets released. 
+	 * 
+	 * @NOTE the whole concept relies on the assumption we only ever create or "build" entities in a single thread (usually the GameThread).
+	 */
+	struct MASSENTITY_API FEntityCreationContext
 	{
-		explicit FEntityCreationContext(const int32 InNumSpawned = 0)
-			: NumberSpawned(InNumSpawned)
-		{}
-		~FEntityCreationContext() { if (OnSpawningFinished) OnSpawningFinished(*this); }
+	private:
+		FEntityCreationContext() = default;
+		explicit FEntityCreationContext(FMassEntityManager& InManager, const TConstArrayView<FMassEntityHandle> InCreatedEntities = {});
+		FEntityCreationContext(FMassEntityManager& InManager, const TConstArrayView<FMassEntityHandle> InCreatedEntities, FMassArchetypeEntityCollection&& EntityCollection);
+		
+	public:
+		~FEntityCreationContext();
+
+		/** Returns EntityCollections, reconstructing them if needed (empty or dirtied). */
+		TConstArrayView<FMassArchetypeEntityCollection> GetEntityCollections() const;
+		int32 GetSpawnedNum() const { return CreatedEntities.Num(); }
+		void MarkDirty();
+		FORCEINLINE bool IsDirty() const { return EntityCollections.IsEmpty() && (CreatedEntities.IsEmpty() == false); }
+		void AppendEntities(const TConstArrayView<FMassEntityHandle> EntitiesToAppend);
+		void AppendEntities(const TConstArrayView<FMassEntityHandle> EntitiesToAppend, FMassArchetypeEntityCollection&& EntityCollection);
+
+		/** Function for debugging/testing purposes. We don't expect users to ever call it, always get collections via GetEntityCollections */
+		bool DebugAreEntityCollectionsUpToDate() const { return EntityCollections.IsEmpty() == CreatedEntities.IsEmpty(); }
+
+		UE_DEPRECATED(5.5, "This constructor is now deprecated and defunct. Use one of the others instead.")
+		explicit FEntityCreationContext(const int32 InNumSpawned) : FEntityCreationContext() {}
+		UE_DEPRECATED(5.5, "This function is now deprecated since FEntityCreationContext can contain more than a single collection now. Use GetEntityCollections instead.")
+		const FMassArchetypeEntityCollection& GetEntityCollection() const;
 
-		const FMassArchetypeEntityCollection& GetEntityCollection() const { return EntityCollection; }
 	private:
 		friend FMassEntityManager;
-		int32 NumberSpawned;
-		FMassArchetypeEntityCollection EntityCollection;
-		TFunction<void(FEntityCreationContext&)> OnSpawningFinished;
+		mutable TArray<FMassArchetypeEntityCollection> EntityCollections;
+		TArray<FMassEntityHandle> CreatedEntities;
+		FMassArchetypeEntityCollection::EDuplicatesHandling CollectionCreationDuplicatesHandling = FMassArchetypeEntityCollection::EDuplicatesHandling::NoDuplicates;
+		TSharedPtr<FMassEntityManager> Manager;
 	};
+	/**
+	 * The main use-case for this function is to create a blank FEntityCreationContext and hold on to it while creating 
+	 * a bunch of entities (with multiple calls to BatchCreate* and/or BatchBuild*) and modifying them (with mutating batched API)
+	 * while not causing multiple Observers to trigger. All the observers will be triggered at one go, once the FEntityCreationContext 
+	 * instance gets destroyed. 
+	 * @return ActiveCreationContext. If it's valid FEntityCreationContext instance will be created and assigned to ActiveCreationContext first.
+	 */
+	TSharedRef<FEntityCreationContext> GetOrMakeCreationContext();
 
 	/**
 	 * A version of CreateEntity that's creating a number of entities (Count) in one go
@@ -321,9 +351,9 @@ public:
 	 * @return a view into InOutEntities containing only the freshly reserved entities
 	 */
 	TConstArrayView<FMassEntityHandle> BatchReserveEntities(const int32 Count, TArray<FMassEntityHandle>& InOutEntities);
-	void BatchBuildEntities(const FMassArchetypeEntityCollectionWithPayload& EncodedEntitiesWithPayload, const FMassFragmentBitSet& FragmentsAffected
+	TSharedRef<FEntityCreationContext> BatchBuildEntities(const FMassArchetypeEntityCollectionWithPayload& EncodedEntitiesWithPayload, const FMassFragmentBitSet& FragmentsAffected
 		, const FMassArchetypeSharedFragmentValues& SharedFragmentValues = {}, const FMassArchetypeCreationParams& CreationParams = FMassArchetypeCreationParams());
-	void BatchBuildEntities(const FMassArchetypeEntityCollectionWithPayload& EncodedEntitiesWithPayload, FMassArchetypeCompositionDescriptor&& Composition
+	TSharedRef<FEntityCreationContext> BatchBuildEntities(const FMassArchetypeEntityCollectionWithPayload& EncodedEntitiesWithPayload, FMassArchetypeCompositionDescriptor&& Composition
 		, const FMassArchetypeSharedFragmentValues& SharedFragmentValues = {}, const FMassArchetypeCreationParams& CreationParams = FMassArchetypeCreationParams());
 	void BatchChangeTagsForEntities(TConstArrayView<FMassArchetypeEntityCollection> EntityCollections, const FMassTagBitSet& TagsToAdd, const FMassTagBitSet& TagsToRemove);
 	void BatchChangeFragmentCompositionForEntities(TConstArrayView<FMassArchetypeEntityCollection> EntityCollections, const FMassFragmentBitSet& FragmentsToAdd, const FMassFragmentBitSet& FragmentsToRemove);
@@ -359,6 +389,8 @@ public:
 	 *  Failing this assumption will cause a check-fail. */
 	static void BatchSetEntityFragmentsValues(const FMassArchetypeEntityCollection& SparseEntities, TArrayView<const FInstancedStruct> FragmentInstanceList);
 
+	static void BatchSetEntityFragmentsValues(TConstArrayView<FMassArchetypeEntityCollection> EntityCollections, TArrayView<const FInstancedStruct> FragmentInstanceList);
+
 	// Return true if it is an valid built entity
 	bool IsEntityActive(FMassEntityHandle Entity) const 
 	{
@@ -650,7 +682,52 @@ private:
 	UE::Mass::FSingleThreadedEntityStorage& GetEntityStorageInterface();
 	const UE::Mass::FSingleThreadedEntityStorage& GetEntityStorageInterface() const;
 #endif
+
+	/**
+	 * If ActiveCreationContext is not valid the function creates a new shared FEntityCreationContext instance and returns that.
+	 * Otherwise ActiveCreationContext will get extended with ReservedEntities and EntityCollection, and returned by the function.
+	 */
+	TSharedRef<FEntityCreationContext> GetOrMakeCreationContext(TConstArrayView<FMassEntityHandle> ReservedEntities, FMassArchetypeEntityCollection&& EntityCollection);
+	bool IsDuringEntityCreation() const { return ActiveCreationContext.IsValid(); }
 	
+	/** 
+	 * This type is used in entity mutating batched API to ensure the active FEntityCreationContext gets dirtied 
+	 * upon function's end (since the mutating operations render FEntityCreationContext.EntityCollections invalid).
+	 * It also serves as a cached IsDuringEntityCreation value.
+	 */
+	struct FScopedCreationContextOperations
+	{
+		FScopedCreationContextOperations(FMassEntityManager& InManager)
+			: bIsDuringEntityCreation(InManager.IsDuringEntityCreation())
+			, Manager(InManager)
+		{}
+		~FScopedCreationContextOperations()
+		{
+			if (bIsDuringEntityCreation)
+			{
+				Manager.DirtyCreationContext();
+			}
+			// else, there's nothing to do, there's no creation context to call functions for
+		}
+
+		bool IsAllowedToTriggerObservers() const { return (bIsDuringEntityCreation == false); }
+
+	private:
+		const bool bIsDuringEntityCreation = false;
+		FMassEntityManager& Manager;
+	};
+	friend FScopedCreationContextOperations;
+
+	/** @return whether actual context dirtying took place which is equivalent to IsDuringEntityCreation */
+	bool DirtyCreationContext();
+
+	/** 
+	 * @return whether it's allowed for observers to get triggered. If not then the active creation context will be dirtied
+	 *	to cause observers triggering upon context's destruction
+	 */
+	bool IsAllowedToTriggerObservers() { return DirtyCreationContext() == false; }
+	bool DebugDoCollectionsOverlapCreationContext(TConstArrayView<FMassArchetypeEntityCollection> EntityCollections) const;
+		
 private:
 
 	friend struct UE::Mass::Private::FEntityStorageInitializer;
@@ -660,7 +737,11 @@ private:
 		UE::Mass::FConcurrentEntityStorage>;
 	FEntityStorageContainerType EntityStorage;
 
+	/** Never access directly, use GetOrMakeCreationContext instead. */
+	TWeakPtr<FEntityCreationContext> ActiveCreationContext;
+
 	std::atomic<bool> bCommandBufferFlushingInProgress = false;
+
 	/**
 	 * This index will be enough to control which buffer is available for pushing commands since flashing is taking place 
 	 * in the game thread and pushing commands to the buffer fetched by Defer() is only supported also on the game thread
diff --git a/Engine/Plugins/Runtime/MassEntity/Source/MassEntity/Public/MassExecutor.h b/Engine/Plugins/Runtime/MassEntity/Source/MassEntity/Public/MassExecutor.h
--- a/Engine/Plugins/Runtime/MassEntity/Source/MassEntity/Public/MassExecutor.h
+++ b/Engine/Plugins/Runtime/MassEntity/Source/MassEntity/Public/MassExecutor.h
@@ -35,7 +35,7 @@ namespace UE::Mass::Executor
 	MASSENTITY_API void RunSparse(FMassRuntimePipeline& RuntimePipeline, FMassProcessingContext& ProcessingContext, const FMassArchetypeEntityCollection& EntityCollection);
 
 	/** Executes given Processors array view. This function gets called under the hood by the rest of Run* functions */
-	MASSENTITY_API void RunProcessorsView(TArrayView<UMassProcessor* const> Processors, FMassProcessingContext& ProcessingContext, const FMassArchetypeEntityCollection* EntityCollection = nullptr);
+	MASSENTITY_API void RunProcessorsView(TArrayView<UMassProcessor* const> Processors, FMassProcessingContext& ProcessingContext, TConstArrayView<FMassArchetypeEntityCollection> EntityCollections = {});
 
 	/** 
 	 *  Triggers tasks executing Processor (and potentially it's children) and returns the task graph event representing 
@@ -44,4 +44,7 @@ namespace UE::Mass::Executor
 	 *    Note that OnDoneNotification will be executed on GameThread.
 	 */
 	MASSENTITY_API FGraphEventRef TriggerParallelTasks(UMassProcessor& Processor, FMassProcessingContext& ProcessingContext, TFunction<void()> OnDoneNotification);
+
+	UE_DEPRECATED(5.5, "This flavor of RunProcessorsView is deprecated. Use the one with TConstArrayView<FMassArchetypeEntityCollection> parameter instead.")
+	MASSENTITY_API void RunProcessorsView(TArrayView<UMassProcessor* const> Processors, FMassProcessingContext& ProcessingContext, const FMassArchetypeEntityCollection* EntityCollection);
 };
diff --git a/Engine/Plugins/Runtime/MassEntity/Source/MassEntity/Public/MassObserverManager.h b/Engine/Plugins/Runtime/MassEntity/Source/MassEntity/Public/MassObserverManager.h
--- a/Engine/Plugins/Runtime/MassEntity/Source/MassEntity/Public/MassObserverManager.h
+++ b/Engine/Plugins/Runtime/MassEntity/Source/MassEntity/Public/MassObserverManager.h
@@ -70,25 +70,41 @@ public:
 		return HasObserversForBitSet(Composition.Fragments, Operation) || HasObserversForBitSet(Composition.Tags, Operation);
 	}
 
+	/** @return whether any observers have been executed */
 	bool OnPostEntitiesCreated(const FMassArchetypeEntityCollection& EntityCollection);
-	bool OnPostEntitiesCreated(FMassProcessingContext& InProcessingContext, const FMassArchetypeEntityCollection& EntityCollection);
 
+	/** @return whether any observers have been executed */
+	bool OnPostEntitiesCreated(TConstArrayView<FMassArchetypeEntityCollection> EntityCollections);
+
+	/** @return whether any observers have been executed */
+	bool OnPostEntitiesCreated(FMassProcessingContext& InProcessingContext, TConstArrayView<FMassArchetypeEntityCollection> EntityCollections);
+
+	/** @return whether any observers have been executed */
 	bool OnPreEntitiesDestroyed(const FMassArchetypeEntityCollection& EntityCollection);
+
+	/** @return whether any observers have been executed */
 	bool OnPreEntitiesDestroyed(FMassProcessingContext& ProcessingContext, const FMassArchetypeEntityCollection& EntityCollection);
+
+	/** @return whether any observers have been executed */
 	bool OnPreEntityDestroyed(const FMassArchetypeCompositionDescriptor& ArchetypeComposition, const FMassEntityHandle Entity);
 
+	/** @return whether any observers have been executed */
 	bool OnCompositionChanged(const FMassEntityHandle Entity, const FMassArchetypeCompositionDescriptor& CompositionDelta, const EMassObservedOperation Operation);
+
+	/** @return whether any observers have been executed */
 	bool OnPostCompositionAdded(const FMassEntityHandle Entity, const FMassArchetypeCompositionDescriptor& Composition)
 	{
 		return OnCompositionChanged(Entity, Composition, EMassObservedOperation::Add);
 	}
+	/** @return whether any observers have been executed */
 	bool OnPreCompositionRemoved(const FMassEntityHandle Entity, const FMassArchetypeCompositionDescriptor& Composition)
 	{
 		return OnCompositionChanged(Entity, Composition, EMassObservedOperation::Remove);
 	}
 
+	/** @return whether any observers have been executed */
 	bool OnCompositionChanged(FMassProcessingContext& InProcessingContext, const FMassArchetypeEntityCollection& EntityCollection, const FMassArchetypeCompositionDescriptor& Composition, const EMassObservedOperation Operation);
-
+	/** @return whether any observers have been executed */
 	bool OnCompositionChanged(const FMassArchetypeEntityCollection& EntityCollection, const FMassArchetypeCompositionDescriptor& Composition, const EMassObservedOperation Operation)
 	{
 		FMassProcessingContext LocalContext(EntityManager, /*DeltaSeconds=*/0.f);
@@ -145,6 +161,9 @@ public:
 		}
 		return OnCompositionChanged(EntityCollection, Composition, Operation);
 	}
+
+	UE_DEPRECATED(5.5, "This flavor of OnPostEntitiesCreated is deprecated. Please use the one taking a TConstArrayView<FMassArchetypeEntityCollection> parameter instead")
+	bool OnPostEntitiesCreated(FMassProcessingContext& InProcessingContext, const FMassArchetypeEntityCollection& EntityCollection);
 };
 
 template<>
diff --git a/Engine/Plugins/Runtime/MassEntity/Source/MassEntityTestSuite/Private/MassObserverTest.cpp b/Engine/Plugins/Runtime/MassEntity/Source/MassEntityTestSuite/Private/MassObserverTest.cpp
--- a/Engine/Plugins/Runtime/MassEntity/Source/MassEntityTestSuite/Private/MassObserverTest.cpp
+++ b/Engine/Plugins/Runtime/MassEntity/Source/MassEntityTestSuite/Private/MassObserverTest.cpp
@@ -24,23 +24,25 @@ auto EntityIndexSorted = [](const FMassEntityHandle& A, const FMassEntityHandle&
 
 struct FTagBaseOperation : FEntityTestBase
 {
+	using FTagStruct = FTestTag_A;
+
 	TArray<FMassEntityHandle> AffectedEntities;
-	UMassTestProcessorBase* TagObserver = nullptr;
+	UMassTestProcessorBase* ObserverProcessor = nullptr;
 	EMassObservedOperation OperationObserved = EMassObservedOperation::MAX;
 	TArray<FMassEntityHandle> EntitiesInt;
 	TArray<FMassEntityHandle> EntitiesIntsFloat;
 	TArray<FMassEntityHandle> ExpectedEntities;
 	// @return signifies if the test can continue
-	virtual bool PerformOperation() = 0;
+	virtual bool PerformOperation() { return false; }
 
 	virtual bool SetUp() override
 	{
 		if (FEntityTestBase::SetUp())
 		{
-			TagObserver = NewObject<UMassTestProcessorBase>();
-			TagObserver->EntityQuery.AddRequirement<FTestFragment_Int>(EMassFragmentAccess::ReadOnly);
-			TagObserver->EntityQuery.AddTagRequirement<FTestTag_A>(EMassFragmentPresence::All);
-			TagObserver->ForEachEntityChunkExecutionFunction = [this](FMassExecutionContext& Context)
+			ObserverProcessor = NewObject<UMassTestProcessorBase>();
+			ObserverProcessor->EntityQuery.AddRequirement<FTestFragment_Int>(EMassFragmentAccess::ReadOnly);
+			ObserverProcessor->EntityQuery.AddTagRequirement<FTagStruct>(EMassFragmentPresence::All);
+			ObserverProcessor->ForEachEntityChunkExecutionFunction = [this](FMassExecutionContext& Context)
 			{
 				AffectedEntities.Append(Context.GetEntities().GetData(), Context.GetEntities().Num());
 			};
@@ -53,7 +55,7 @@ struct FTagBaseOperation : FEntityTestBase
 	virtual bool InstantTest() override
 	{
 		FMassObserverManager& ObserverManager = EntityManager->GetObserverManager();
-		ObserverManager.AddObserverInstance(*FTestTag_A::StaticStruct(), OperationObserved, *TagObserver);
+		ObserverManager.AddObserverInstance(*FTagStruct::StaticStruct(), OperationObserved, *ObserverProcessor);
 
 		EntityManager->BatchCreateEntities(IntsArchetype, 3, EntitiesInt);
 		EntityManager->BatchCreateEntities(FloatsIntsArchetype, 3, EntitiesIntsFloat);
@@ -70,51 +72,48 @@ struct FTagBaseOperation : FEntityTestBase
 			{
 				AITEST_EQUAL(TEXT("Expected and affected sets should be the same"), AffectedEntities[i], ExpectedEntities[i]);
 			}
-
-			/*AITEST_EQUAL(TEXT("The tag observer is expected to be run for a single entity"), AffectedEntities.Num(), 1);
-			AITEST_EQUAL(TEXT("The tag observer is expected to be run for the entity that received the tag"), AffectedEntities[0], Entities[1]);*/
 		}
 
 		return true;
 	}
 };
 
-struct FTag_SingleEntitySingleArchetypeAdd : FTagBaseOperation
+struct FObserverProcessorTest_SingleEntitySingleArchetypeAdd : FTagBaseOperation
 {
-	FTag_SingleEntitySingleArchetypeAdd() { OperationObserved = EMassObservedOperation::Add; }
+	FObserverProcessorTest_SingleEntitySingleArchetypeAdd() { OperationObserved = EMassObservedOperation::Add; }
 	virtual bool PerformOperation() override 
 	{
 		ExpectedEntities = { EntitiesInt[1] };
-		EntityManager->Defer().AddTag<FTestTag_A>(EntitiesInt[1]);
+		EntityManager->Defer().AddTag<FTagStruct>(EntitiesInt[1]);
 		return true;
 	}
 };
-IMPLEMENT_AI_INSTANT_TEST(FTag_SingleEntitySingleArchetypeAdd, "System.Mass.Observer.Tag.SingleEntitySingleArchetypeAdd");
+IMPLEMENT_AI_INSTANT_TEST(FObserverProcessorTest_SingleEntitySingleArchetypeAdd, "System.Mass.Observer.Tag.SingleEntitySingleArchetypeAdd");
 
-struct FTag_SingleEntitySingleArchetypeRemove : FTagBaseOperation
+struct FObserverProcessorTest_SingleEntitySingleArchetypeRemove : FTagBaseOperation
 {
-	FTag_SingleEntitySingleArchetypeRemove() { OperationObserved = EMassObservedOperation::Remove; }
+	FObserverProcessorTest_SingleEntitySingleArchetypeRemove() { OperationObserved = EMassObservedOperation::Remove; }
 	virtual bool PerformOperation() override
 	{
 		ExpectedEntities = { EntitiesInt[1] };
 
-		EntityManager->Defer().AddTag<FTestTag_A>(EntitiesInt[1]);
+		EntityManager->Defer().AddTag<FTagStruct>(EntitiesInt[1]);
 		EntityManager->FlushCommands();
 		// since we're only observing tag removal we don't expect AffectedEntities to contain any data at this point
 		AITEST_EQUAL(TEXT("Tag addition is not being observed and is not expected to produce results yet"), AffectedEntities.Num(), 0);
-		EntityManager->Defer().RemoveTag<FTestTag_A>(EntitiesInt[1]);
+		EntityManager->Defer().RemoveTag<FTagStruct>(EntitiesInt[1]);
 		return true;
 	}
 };
-IMPLEMENT_AI_INSTANT_TEST(FTag_SingleEntitySingleArchetypeRemove, "System.Mass.Observer.Tag.SingleEntitySingleArchetypeRemove");
+IMPLEMENT_AI_INSTANT_TEST(FObserverProcessorTest_SingleEntitySingleArchetypeRemove, "System.Mass.Observer.Tag.SingleEntitySingleArchetypeRemove");
 
-struct FTag_SingleEntitySingleArchetypeDestroy : FTagBaseOperation
+struct FObserverProcessorTest_SingleEntitySingleArchetypeDestroy : FTagBaseOperation
 {
-	FTag_SingleEntitySingleArchetypeDestroy() { OperationObserved = EMassObservedOperation::Remove; }
+	FObserverProcessorTest_SingleEntitySingleArchetypeDestroy() { OperationObserved = EMassObservedOperation::Remove; }
 	virtual bool PerformOperation() override
 	{
 		ExpectedEntities = { EntitiesInt[1] };
-		EntityManager->Defer().AddTag<FTestTag_A>(EntitiesInt[1]); 
+		EntityManager->Defer().AddTag<FTagStruct>(EntitiesInt[1]);
 		EntityManager->FlushCommands();
 		// since we're only observing tag removal we don't expect AffectedEntities to contain any data at this point
 		AITEST_EQUAL(TEXT("Tag addition is not being observed and is not expected to produce results yet"), AffectedEntities.Num(), 0);
@@ -122,57 +121,57 @@ struct FTag_SingleEntitySingleArchetypeDestroy : FTagBaseOperation
 		return true;
 	}
 };
-IMPLEMENT_AI_INSTANT_TEST(FTag_SingleEntitySingleArchetypeDestroy, "System.Mass.Observer.Tag.SingleEntitySingleArchetypeDestroy");
+IMPLEMENT_AI_INSTANT_TEST(FObserverProcessorTest_SingleEntitySingleArchetypeDestroy, "System.Mass.Observer.Tag.SingleEntitySingleArchetypeDestroy");
 
-struct FTag_MultipleArchetypeAdd : FTagBaseOperation
+struct FObserverProcessorTest_MultipleArchetypeAdd : FTagBaseOperation
 {
-	FTag_MultipleArchetypeAdd() { OperationObserved = EMassObservedOperation::Add; }
+	FObserverProcessorTest_MultipleArchetypeAdd() { OperationObserved = EMassObservedOperation::Add; }
 
 	virtual bool PerformOperation() override
 	{
 		ExpectedEntities = { EntitiesInt[0], EntitiesInt[2], EntitiesIntsFloat[1] };
 		for (const FMassEntityHandle& ModifiedEntity : ExpectedEntities)
 		{
-			EntityManager->Defer().AddTag<FTestTag_A>(ModifiedEntity);
+			EntityManager->Defer().AddTag<FTagStruct>(ModifiedEntity);
 		}
 		return true;
 	}
 };
-IMPLEMENT_AI_INSTANT_TEST(FTag_MultipleArchetypeAdd, "System.Mass.Observer.Tag.MultipleArchetypesAdd");
+IMPLEMENT_AI_INSTANT_TEST(FObserverProcessorTest_MultipleArchetypeAdd, "System.Mass.Observer.Tag.MultipleArchetypesAdd");
 
-struct FTag_MultipleArchetypeRemove : FTagBaseOperation
+struct FObserverProcessorTest_MultipleArchetypeRemove : FTagBaseOperation
 {
-	FTag_MultipleArchetypeRemove() { OperationObserved = EMassObservedOperation::Remove; }
+	FObserverProcessorTest_MultipleArchetypeRemove() { OperationObserved = EMassObservedOperation::Remove; }
 
 	virtual bool PerformOperation() override
 	{
 		ExpectedEntities = { EntitiesInt[0], EntitiesInt[2], EntitiesIntsFloat[1] };
 		for (const FMassEntityHandle& ModifiedEntity : ExpectedEntities)
 		{
-			EntityManager->Defer().AddTag<FTestTag_A>(ModifiedEntity);
+			EntityManager->Defer().AddTag<FTagStruct>(ModifiedEntity);
 		}
 		EntityManager->FlushCommands();
 		// since we're only observing tag removal we don't expect AffectedEntities to contain any data at this point
 		AITEST_EQUAL(TEXT("Tag addition is not being observed and is not expected to produce results yet"), AffectedEntities.Num(), 0);
 		for (const FMassEntityHandle& ModifiedEntity : ExpectedEntities)
 		{
-			EntityManager->Defer().RemoveTag<FTestTag_A>(ModifiedEntity);
+			EntityManager->Defer().RemoveTag<FTagStruct>(ModifiedEntity);
 		}
 		return true;
 	}
 };
-IMPLEMENT_AI_INSTANT_TEST(FTag_MultipleArchetypeRemove, "System.Mass.Observer.Tag.MultipleArchetypesRemove");
+IMPLEMENT_AI_INSTANT_TEST(FObserverProcessorTest_MultipleArchetypeRemove, "System.Mass.Observer.Tag.MultipleArchetypesRemove");
 
-struct FTag_MultipleArchetypeDestroy : FTagBaseOperation
+struct FObserverProcessorTest_MultipleArchetypeDestroy : FTagBaseOperation
 {
-	FTag_MultipleArchetypeDestroy() { OperationObserved = EMassObservedOperation::Remove; }
+	FObserverProcessorTest_MultipleArchetypeDestroy() { OperationObserved = EMassObservedOperation::Remove; }
 
 	virtual bool PerformOperation() override
 	{
 		ExpectedEntities = { EntitiesInt[0], EntitiesInt[2], EntitiesIntsFloat[1] };
 		for (const FMassEntityHandle& ModifiedEntity : ExpectedEntities)
 		{
-			EntityManager->Defer().AddTag<FTestTag_A>(ModifiedEntity);
+			EntityManager->Defer().AddTag<FTagStruct>(ModifiedEntity);
 		}
 		EntityManager->FlushCommands();
 		// since we're only observing tag removal we don't expect AffectedEntities to contain any data at this point
@@ -184,30 +183,390 @@ struct FTag_MultipleArchetypeDestroy : FTagBaseOperation
 		return true;
 	}
 };
-IMPLEMENT_AI_INSTANT_TEST(FTag_MultipleArchetypeDestroy, "System.Mass.Observer.Tag.MultipleArchetypesDestroy");
+IMPLEMENT_AI_INSTANT_TEST(FObserverProcessorTest_MultipleArchetypeDestroy, "System.Mass.Observer.Tag.MultipleArchetypesDestroy");
 
-struct FTag_MultipleArchetypeSwap : FTagBaseOperation
+struct FObserverProcessorTest_MultipleArchetypeSwap : FTagBaseOperation
 {
-	FTag_MultipleArchetypeSwap() { OperationObserved = EMassObservedOperation::Remove; }
+	FObserverProcessorTest_MultipleArchetypeSwap() { OperationObserved = EMassObservedOperation::Remove; }
 
 	virtual bool PerformOperation() override
 	{
 		ExpectedEntities = { EntitiesIntsFloat[1], EntitiesInt[0], EntitiesInt[2] };
 		for (const FMassEntityHandle& ModifiedEntity : ExpectedEntities)
 		{
-			EntityManager->Defer().AddTag<FTestTag_A>(ModifiedEntity);
+			EntityManager->Defer().AddTag<FTagStruct>(ModifiedEntity);
 		}
 		EntityManager->FlushCommands();
 		// since we're only observing tag removal we don't expect AffectedEntities to contain any data at this point
 		AITEST_EQUAL(TEXT("Tag addition is not being observed and is not expected to produce results yet"), AffectedEntities.Num(), 0);
 		for (const FMassEntityHandle& ModifiedEntity : ExpectedEntities)
 		{
-			EntityManager->Defer().SwapTags<FTestTag_A, FTestTag_B>(ModifiedEntity);
+			EntityManager->Defer().SwapTags<FTagStruct, FTestTag_B>(ModifiedEntity);
+		}
+		return true;
+	}
+};
+IMPLEMENT_AI_INSTANT_TEST(FObserverProcessorTest_MultipleArchetypeSwap, "System.Mass.Observer.Tag.MultipleArchetypesSwap");
+
+struct FObserverProcessorTest_EntityCreation_Individuals : FTagBaseOperation
+{
+	FObserverProcessorTest_EntityCreation_Individuals() { OperationObserved = EMassObservedOperation::Add; }
+
+	virtual bool InstantTest() override
+	{
+		constexpr int32 EntitiesToSpawnCount = 6;
+		
+		FMassObserverManager& ObserverManager = EntityManager->GetObserverManager();
+		ObserverManager.AddObserverInstance(*FTagStruct::StaticStruct(), OperationObserved, *ObserverProcessor);
+
+		int32 ArrayMidPoint = 0;
+		{
+			TSharedRef<FMassEntityManager::FEntityCreationContext> CreationContext = EntityManager->BatchCreateEntities(IntsArchetype, EntitiesToSpawnCount, EntitiesInt);
+			ArrayMidPoint = EntitiesInt.Num() / 2;
+
+			for (int32 Index = 0; Index < ArrayMidPoint; ++Index)
+			{
+				EntityManager->AddTagToEntity(EntitiesInt[Index], FTagStruct::StaticStruct());
+			}
+			AITEST_EQUAL(TEXT("The tag observer is not expected to run yet"), AffectedEntities.Num(), 0);
+		}
+		AITEST_EQUAL(TEXT("The tag observer is expected to run just after FEntityCreationContext's destruction"), AffectedEntities.Num(), ArrayMidPoint);
+
+		return true;
+	}
+};
+IMPLEMENT_AI_INSTANT_TEST(FObserverProcessorTest_EntityCreation_Individuals, "System.Mass.Observer.Create.TagInvididualEntities");
+
+struct FObserverProcessorTest_EntityCreation_Batched : FTagBaseOperation
+{
+	FObserverProcessorTest_EntityCreation_Batched() { OperationObserved = EMassObservedOperation::Add; }
+
+	virtual bool InstantTest() override
+	{
+		constexpr int32 EntitiesToSpawnCount = 6;
+
+		FMassObserverManager& ObserverManager = EntityManager->GetObserverManager();
+		ObserverManager.AddObserverInstance(*FTagStruct::StaticStruct(), OperationObserved, *ObserverProcessor);
+
+		{
+			TSharedRef<FMassEntityManager::FEntityCreationContext> CreationContext = EntityManager->BatchCreateEntities(IntsArchetype, EntitiesToSpawnCount, EntitiesInt);
+
+			EntityManager->BatchChangeTagsForEntities(CreationContext->GetEntityCollections(), FMassTagBitSet(*FTagStruct::StaticStruct()), FMassTagBitSet());
+			AITEST_TRUE(TEXT("The tag observer is not expected to run yet"), AffectedEntities.Num() == 0);
+			AITEST_FALSE(TEXT("CreationContext's entity collection should be invalidated at this moment"), CreationContext->DebugAreEntityCollectionsUpToDate());
+
+			EntityManager->BatchChangeTagsForEntities(CreationContext->GetEntityCollections(), FMassTagBitSet(*FTagStruct::StaticStruct()), FMassTagBitSet());
+			AITEST_TRUE(TEXT("The tag observer is still not expected to run"), AffectedEntities.Num() == 0);
+		}
+		AITEST_TRUE(TEXT("The tag observer is expected to run just after FEntityCreationContext's destruction"), AffectedEntities.Num() > 0);
+		AITEST_EQUAL(TEXT("The tag observer is expected to process every entity just once"), AffectedEntities.Num(), EntitiesInt.Num());
+
+		return true;
+	}
+};
+IMPLEMENT_AI_INSTANT_TEST(FObserverProcessorTest_EntityCreation_Batched, "System.Mass.Observer.Create.TagBatchedEntities");
+
+//-----------------------------------------------------------------------------
+// fragments
+//-----------------------------------------------------------------------------
+struct FFragmentTestBase : FEntityTestBase
+{
+	using FFragmentStruct = FTestFragment_Float;
+
+	TArray<FMassEntityHandle> AffectedEntities;
+	UMassTestProcessorBase* ObserverProcessor = nullptr;
+	EMassObservedOperation OperationObserved = EMassObservedOperation::MAX;
+	TArray<FMassEntityHandle> EntitiesInt;
+	TArray<FMassEntityHandle> EntitiesIntsFloat;
+	TArray<FMassEntityHandle> ExpectedEntities;
+	// @return signifies if the test can continue
+	virtual bool PerformOperation() { return false; }
+
+	virtual bool SetUp() override
+	{
+		if (FEntityTestBase::SetUp())
+		{
+			ObserverProcessor = NewObject<UMassTestProcessorBase>();
+			ObserverProcessor->EntityQuery.AddRequirement(FFragmentStruct::StaticStruct(), EMassFragmentAccess::ReadOnly);
+			ObserverProcessor->ForEachEntityChunkExecutionFunction = [this](FMassExecutionContext& Context)
+				{
+					AffectedEntities.Append(Context.GetEntities().GetData(), Context.GetEntities().Num());
+				};
+
+			return true;
+		}
+		return false;
+	}
+
+	virtual bool InstantTest() override
+	{
+		EntityManager->BatchCreateEntities(IntsArchetype, 3, EntitiesInt);
+		EntityManager->BatchCreateEntities(FloatsIntsArchetype, 3, EntitiesIntsFloat);
+
+		FMassObserverManager& ObserverManager = EntityManager->GetObserverManager();
+		ObserverManager.AddObserverInstance(*FFragmentStruct::StaticStruct(), OperationObserved, *ObserverProcessor);
+				
+		if (PerformOperation())
+		{
+			EntityManager->FlushCommands();
+			AITEST_EQUAL(TEXT("The fragment observer is expected to be run for predicted number of entities"), AffectedEntities.Num(), ExpectedEntities.Num());
+
+			ExpectedEntities.Sort(EntityIndexSorted);
+			AffectedEntities.Sort(EntityIndexSorted);
+
+			for (int i = 0; i < ExpectedEntities.Num(); ++i)
+			{
+				AITEST_EQUAL(TEXT("Expected and affected sets should be the same"), AffectedEntities[i], ExpectedEntities[i]);
+			}
 		}
+
+		return true;
+	}
+};
+
+struct FFragmentObserverTest_SingleEntitySingleArchetypeAdd : FFragmentTestBase
+{
+	FFragmentObserverTest_SingleEntitySingleArchetypeAdd() { OperationObserved = EMassObservedOperation::Add; }
+	virtual bool PerformOperation() override
+	{
+		ExpectedEntities = { EntitiesInt[1] };
+		EntityManager->Defer().AddFragment<FFragmentStruct>(EntitiesInt[1]);
+		return true;
+	}
+};
+IMPLEMENT_AI_INSTANT_TEST(FFragmentObserverTest_SingleEntitySingleArchetypeAdd, "System.Mass.Observer.Fragment.SingleEntitySingleArchetypeAdd");
+
+struct FFragmentObserverTest_SingleEntitySingleArchetypeRemove : FFragmentTestBase
+{
+	FFragmentObserverTest_SingleEntitySingleArchetypeRemove() { OperationObserved = EMassObservedOperation::Remove; }
+	virtual bool PerformOperation() override
+	{
+		ExpectedEntities = { EntitiesInt[1] };
+
+		EntityManager->Defer().AddFragment<FFragmentStruct>(EntitiesInt[1]);
+		EntityManager->FlushCommands();
+		// since we're only observing Fragment removal we don't expect AffectedEntities to contain any data at this point
+		AITEST_EQUAL(TEXT("Fragment addition is not being observed and is not expected to produce results yet"), AffectedEntities.Num(), 0);
+		EntityManager->Defer().RemoveFragment<FFragmentStruct>(EntitiesInt[1]);
+		return true;
+	}
+};
+IMPLEMENT_AI_INSTANT_TEST(FFragmentObserverTest_SingleEntitySingleArchetypeRemove, "System.Mass.Observer.Fragment.SingleEntitySingleArchetypeRemove");
+
+struct FFragmentObserverTest_SingleEntitySingleArchetypeDestroy : FFragmentTestBase
+{
+	FFragmentObserverTest_SingleEntitySingleArchetypeDestroy() { OperationObserved = EMassObservedOperation::Remove; }
+	virtual bool PerformOperation() override
+	{
+		ExpectedEntities = { EntitiesInt[1] };
+		EntityManager->Defer().AddFragment<FFragmentStruct>(EntitiesInt[1]);
+		EntityManager->FlushCommands();
+		// since we're only observing Fragment removal we don't expect AffectedEntities to contain any data at this point
+		AITEST_EQUAL(TEXT("Fragment addition is not being observed and is not expected to produce results yet"), AffectedEntities.Num(), 0);
+		EntityManager->Defer().DestroyEntity(EntitiesInt[1]);
+		return true;
+	}
+};
+IMPLEMENT_AI_INSTANT_TEST(FFragmentObserverTest_SingleEntitySingleArchetypeDestroy, "System.Mass.Observer.Fragment.SingleEntitySingleArchetypeDestroy");
+
+struct FFragmentObserverTest_MultipleArchetypeAdd : FFragmentTestBase
+{
+	FFragmentObserverTest_MultipleArchetypeAdd() { OperationObserved = EMassObservedOperation::Add; }
+
+	virtual bool PerformOperation() override
+	{
+		ExpectedEntities = { EntitiesInt[0], EntitiesInt[2], EntitiesInt[1] };
+		for (const FMassEntityHandle& ModifiedEntity : ExpectedEntities)
+		{
+			EntityManager->Defer().AddFragment<FFragmentStruct>(ModifiedEntity);
+		}
+		return true;
+	}
+};
+IMPLEMENT_AI_INSTANT_TEST(FFragmentObserverTest_MultipleArchetypeAdd, "System.Mass.Observer.Fragment.MultipleArchetypesAdd");
+
+struct FFragmentObserverTest_MultipleArchetypeRemove : FFragmentTestBase
+{
+	FFragmentObserverTest_MultipleArchetypeRemove() { OperationObserved = EMassObservedOperation::Remove; }
+
+	virtual bool PerformOperation() override
+	{
+		ExpectedEntities = { EntitiesInt[0], EntitiesInt[2], EntitiesInt[1] };
+		for (const FMassEntityHandle& ModifiedEntity : ExpectedEntities)
+		{
+			EntityManager->Defer().AddFragment<FFragmentStruct>(ModifiedEntity);
+		}
+		EntityManager->FlushCommands();
+		// since we're only observing Fragment removal we don't expect AffectedEntities to contain any data at this point
+		AITEST_EQUAL(TEXT("Fragment addition is not being observed and is not expected to produce results yet"), AffectedEntities.Num(), 0);
+		for (const FMassEntityHandle& ModifiedEntity : ExpectedEntities)
+		{
+			EntityManager->Defer().RemoveFragment<FFragmentStruct>(ModifiedEntity);
+		}
+		return true;
+	}
+};
+IMPLEMENT_AI_INSTANT_TEST(FFragmentObserverTest_MultipleArchetypeRemove, "System.Mass.Observer.Fragment.MultipleArchetypesRemove");
+
+struct FFragmentObserverTest_MultipleArchetypeDestroy : FFragmentTestBase
+{
+	FFragmentObserverTest_MultipleArchetypeDestroy() { OperationObserved = EMassObservedOperation::Remove; }
+
+	virtual bool PerformOperation() override
+	{
+		ExpectedEntities = { EntitiesInt[0], EntitiesInt[2], EntitiesInt[1] };
+		for (const FMassEntityHandle& ModifiedEntity : ExpectedEntities)
+		{
+			EntityManager->Defer().AddFragment<FFragmentStruct>(ModifiedEntity);
+		}
+		EntityManager->FlushCommands();
+		// since we're only observing Fragment removal we don't expect AffectedEntities to contain any data at this point
+		AITEST_EQUAL(TEXT("Fragment addition is not being observed and is not expected to produce results yet"), AffectedEntities.Num(), 0);
+		for (const FMassEntityHandle& ModifiedEntity : ExpectedEntities)
+		{
+			EntityManager->Defer().DestroyEntity(ModifiedEntity);
+		}
+		return true;
+	}
+};
+IMPLEMENT_AI_INSTANT_TEST(FFragmentObserverTest_MultipleArchetypeDestroy, "System.Mass.Observer.Fragment.MultipleArchetypesDestroy");
+
+struct FFragmentObserverTest_EntityCreation_Individuals : FFragmentTestBase
+{
+	FFragmentObserverTest_EntityCreation_Individuals() { OperationObserved = EMassObservedOperation::Add; }
+
+	virtual bool InstantTest() override
+	{
+		constexpr int32 EntitiesToSpawnCount = 6;
+
+		FMassObserverManager& ObserverManager = EntityManager->GetObserverManager();
+		ObserverManager.AddObserverInstance(*FFragmentStruct::StaticStruct(), OperationObserved, *ObserverProcessor);
+
+		int32 ArrayMidPoint = 0;
+		{
+			TSharedRef<FMassEntityManager::FEntityCreationContext> CreationContext = EntityManager->BatchCreateEntities(IntsArchetype, EntitiesToSpawnCount, EntitiesInt);
+			ArrayMidPoint = EntitiesInt.Num() / 2;
+
+			for (int32 Index = 0; Index < ArrayMidPoint; ++Index)
+			{
+				EntityManager->AddFragmentToEntity(EntitiesInt[Index], FFragmentStruct::StaticStruct());
+			}
+			AITEST_EQUAL(TEXT("The fragment observer is not expected to run yet"), AffectedEntities.Num(), 0);
+		}
+		AITEST_EQUAL(TEXT("The fragment observer is expected to run just after FEntityCreationContext's destruction"), AffectedEntities.Num(), ArrayMidPoint);
+
+		return true;
+	}
+};
+IMPLEMENT_AI_INSTANT_TEST(FFragmentObserverTest_EntityCreation_Individuals, "System.Mass.Observer.Create.FragmentInvididualEntities");
+
+//-----------------------------------------------------------------------------
+// creation context 
+//-----------------------------------------------------------------------------
+struct FCreationContextTest : FEntityTestBase
+{
+	virtual bool InstantTest() override
+	{
+		constexpr int32 IntEntitiesToSpawnCount = 6;
+		constexpr int32 FloatEntitiesToSpawnCount = 7;
+
+		TArray<FMassEntityHandle> Entities;
+		TSharedRef<FMassEntityManager::FEntityCreationContext> CreationContextInt = EntityManager->BatchCreateEntities(IntsArchetype, IntEntitiesToSpawnCount, Entities);
+		TSharedRef<FMassEntityManager::FEntityCreationContext> CreationContextFloat = EntityManager->BatchCreateEntities(FloatsArchetype, FloatEntitiesToSpawnCount, Entities);
+		const int32 NumDifferentArchetypesUsed = 2;
+
+		AITEST_EQUAL(TEXT("Two back to back entity creation operations should result in the same creation context"), CreationContextInt, CreationContextFloat);
+		AITEST_FALSE(TEXT("CreationContext's entity collection should be invalidated at this moment"), CreationContextInt->DebugAreEntityCollectionsUpToDate());
+
+		TConstArrayView<FMassArchetypeEntityCollection> EntityCollections = CreationContextInt->GetEntityCollections();
+		AITEST_EQUAL(TEXT("We expect the number of resulting collections to match expectations"), EntityCollections.Num(), NumDifferentArchetypesUsed);
+
+		return true;
+	}
+};
+IMPLEMENT_AI_INSTANT_TEST(FCreationContextTest, "System.Mass.CreationContext.Append");
+
+struct FCreationContextTest_ManualCreate : FEntityTestBase
+{
+	virtual bool InstantTest() override
+	{
+		constexpr int32 IntEntitiesToSpawnCount = 6;
+		constexpr int32 FloatEntitiesToSpawnCount = 7;
+		int NumDifferentArchetypesUsed = 0;
+
+		TArray<FMassEntityHandle> Entities;
+		TSharedRef<FMassEntityManager::FEntityCreationContext> ObtainedContext = EntityManager->GetOrMakeCreationContext();
+		{
+			TSharedRef<FMassEntityManager::FEntityCreationContext> ObtainedContextCopy = EntityManager->GetOrMakeCreationContext();
+			AITEST_EQUAL(TEXT("Two back to back creation context fetching should result in the same instance"), ObtainedContext, ObtainedContextCopy);
+		}
+
+		{
+			TSharedRef<FMassEntityManager::FEntityCreationContext> CreationContextInt = EntityManager->BatchCreateEntities(IntsArchetype, IntEntitiesToSpawnCount, Entities);
+			AITEST_EQUAL(TEXT("Creating entities should return the original context"), ObtainedContext, CreationContextInt);
+			++NumDifferentArchetypesUsed;
+		}
+		
+		AITEST_TRUE(TEXT("CreationContext's entity collection should be still valid at this moment since we only added one entity collection/array")
+			, ObtainedContext->DebugAreEntityCollectionsUpToDate());
+
+		{
+			TSharedRef<FMassEntityManager::FEntityCreationContext> TempContext = EntityManager->BatchCreateEntities(IntsArchetype, IntEntitiesToSpawnCount, Entities);
+			AITEST_EQUAL(TEXT("Creating entities should return the original context"), ObtainedContext, TempContext);
+
+			AITEST_FALSE(TEXT("CreationContext's entity collection should be invalidated at this moment")
+				, TempContext->DebugAreEntityCollectionsUpToDate());
+		}
+
+		TConstArrayView<FMassArchetypeEntityCollection> EntityCollections = ObtainedContext->GetEntityCollections();
+		AITEST_EQUAL(TEXT("We expect the number of resulting collections to match expectations"), EntityCollections.Num(), NumDifferentArchetypesUsed);
+
+		return true;
+	}
+};
+IMPLEMENT_AI_INSTANT_TEST(FCreationContextTest_ManualCreate, "System.Mass.CreationContext.ManualCreate");
+
+struct FCreationContextTest_ManualBuild : FEntityTestBase
+{
+	virtual bool InstantTest() override
+	{
+		constexpr int32 FloatEntitiesToSpawnCount = 7;
+		int NumDifferentArchetypesUsed = 0;
+
+		TArray<FTestFragment_Float> Payload;
+		for (int Index = 0; Index < FloatEntitiesToSpawnCount; ++Index)
+		{ 
+			Payload.Add(FTestFragment_Float(float(Index)));
+		}
+
+		TSharedRef<FMassEntityManager::FEntityCreationContext> ObtainedContext = EntityManager->GetOrMakeCreationContext();
+		
+		TArray<FMassEntityHandle> Entities;
+		EntityManager->BatchReserveEntities(FloatEntitiesToSpawnCount, Entities);
+
+		FStructArrayView PaloadView(Payload);
+		TArray<FMassArchetypeEntityCollectionWithPayload> EntityCollections;
+		FMassArchetypeEntityCollectionWithPayload::CreateEntityRangesWithPayload(*EntityManager, Entities, FMassArchetypeEntityCollection::NoDuplicates
+			, FMassGenericPayloadView(MakeArrayView(&PaloadView, 1)), EntityCollections);
+
+		checkf(EntityCollections.Num() <= 1, TEXT("We expect TargetEntities to only contain archetype-less entities, ones that need to be \'build\'"));
+
+		{
+			TSharedRef<FMassEntityManager::FEntityCreationContext> CreationContext = EntityManager->BatchBuildEntities(EntityCollections[0], FMassFragmentBitSet(*FTestFragment_Float::StaticStruct()));
+			AITEST_EQUAL(TEXT("Creating entities should return the original context"), ObtainedContext, CreationContext);
+			++NumDifferentArchetypesUsed;
+		}
+
+		AITEST_TRUE(TEXT("CreationContext's entity collection should be still valid at this moment since we only added one entity collection/array")
+			, ObtainedContext->DebugAreEntityCollectionsUpToDate());
+
+		TConstArrayView<FMassArchetypeEntityCollection> ContextEntityCollections = ObtainedContext->GetEntityCollections();
+		AITEST_EQUAL(TEXT("We expect the number of resulting collections to match expectations"), ContextEntityCollections.Num(), NumDifferentArchetypesUsed);
+
 		return true;
 	}
 };
-IMPLEMENT_AI_INSTANT_TEST(FTag_MultipleArchetypeSwap, "System.Mass.Observer.Tag.MultipleArchetypesSwap");
+IMPLEMENT_AI_INSTANT_TEST(FCreationContextTest_ManualBuild, "System.Mass.CreationContext.ManualBuild");
 
 } // FMassObserverTest
 
diff --git a/Engine/Plugins/Runtime/MassEntity/Source/MassEntityTestSuite/Private/MassSharedFragmentsTest.cpp b/Engine/Plugins/Runtime/MassEntity/Source/MassEntityTestSuite/Private/MassSharedFragmentsTest.cpp
--- a/Engine/Plugins/Runtime/MassEntity/Source/MassEntityTestSuite/Private/MassSharedFragmentsTest.cpp
+++ b/Engine/Plugins/Runtime/MassEntity/Source/MassEntityTestSuite/Private/MassSharedFragmentsTest.cpp
@@ -511,7 +511,7 @@ struct FSharedFragment_BatchSetAttempt : FEntityTestBase
 		SharedIntValues.AddConstSharedFragment(ConstSharedFragment);
 
 		TSharedRef<FMassEntityManager::FEntityCreationContext> CreationContext = EntityManager->BatchCreateEntities(FloatsArchetype, SharedIntValues, EntitiesToCreateNum, CreatedEntityHandles);
-		const FMassArchetypeHandle ResultingArchetype = CreationContext->GetEntityCollection().GetArchetype();
+		const FMassArchetypeHandle ResultingArchetype = CreationContext->GetEntityCollections()[0].GetArchetype();
 
 		FMassArchetypeEntityCollection EntityCollection(ResultingArchetype, CreatedEntityHandles, FMassArchetypeEntityCollection::EDuplicatesHandling::NoDuplicates);
 		
diff --git a/Engine/Plugins/Runtime/MassGameplay/Source/MassSpawner/Private/MassSpawner.cpp b/Engine/Plugins/Runtime/MassGameplay/Source/MassSpawner/Private/MassSpawner.cpp
--- a/Engine/Plugins/Runtime/MassGameplay/Source/MassSpawner/Private/MassSpawner.cpp
+++ b/Engine/Plugins/Runtime/MassGameplay/Source/MassSpawner/Private/MassSpawner.cpp
@@ -499,10 +499,7 @@ void AMassSpawner::SpawnGeneratedEntities(TConstArrayView<FMassEntitySpawnDataGe
 			// create entity collections and run Processors on them. 
 			TArray<FMassArchetypeEntityCollection> EntityCollections;
 			UE::Mass::Utils::CreateEntityCollections(EntityManager, AllEntities, FMassArchetypeEntityCollection::NoDuplicates, EntityCollections);
-			for (const FMassArchetypeEntityCollection& Collection : EntityCollections)
-			{
-				UE::Mass::Executor::RunProcessorsView(Processors, ProcessingContext, &Collection);
-			}
+			UE::Mass::Executor::RunProcessorsView(Processors, ProcessingContext, EntityCollections);
 		}
 	}
 
diff --git a/Engine/Plugins/Runtime/MassGameplay/Source/MassSpawner/Private/MassSpawnerSubsystem.cpp b/Engine/Plugins/Runtime/MassGameplay/Source/MassSpawner/Private/MassSpawnerSubsystem.cpp
--- a/Engine/Plugins/Runtime/MassGameplay/Source/MassSpawner/Private/MassSpawnerSubsystem.cpp
+++ b/Engine/Plugins/Runtime/MassGameplay/Source/MassSpawner/Private/MassSpawnerSubsystem.cpp
@@ -138,10 +138,11 @@ void UMassSpawnerSubsystem::DoSpawning(const FMassEntityTemplate& EntityTemplate
 	// 4. "OnEntitiesCreated" notifies will be sent out once the CreationContext gets destroyed (via its destructor).
 
 	TArray<FMassEntityHandle> SpawnedEntities;
-	TSharedRef<FMassEntityManager::FEntityCreationContext> CreationContext = EntityManager->BatchCreateEntities(EntityTemplate.GetArchetype(), EntityTemplate.GetSharedFragmentValues(), NumToSpawn, SpawnedEntities);
+	TSharedRef<FMassEntityManager::FEntityCreationContext> CreationContext
+		= EntityManager->BatchCreateEntities(EntityTemplate.GetArchetype(), EntityTemplate.GetSharedFragmentValues(), NumToSpawn, SpawnedEntities);
 
 	TConstArrayView<FInstancedStruct> FragmentInstances = EntityTemplate.GetInitialFragmentValues();
-	EntityManager->BatchSetEntityFragmentsValues(CreationContext->GetEntityCollection(), FragmentInstances);
+	EntityManager->BatchSetEntityFragmentsValues(CreationContext->GetEntityCollections(), FragmentInstances);
 	
 	UMassProcessor* SpawnDataInitializer = SpawnData.IsValid() ? GetSpawnDataInitializer(InitializerClass) : nullptr;
 
@@ -149,7 +150,7 @@ void UMassSpawnerSubsystem::DoSpawning(const FMassEntityTemplate& EntityTemplate
 	{
 		FMassProcessingContext ProcessingContext(EntityManager, /*TimeDelta=*/0.0f);
 		ProcessingContext.AuxData = SpawnData;
-		UE::Mass::Executor::RunProcessorsView(MakeArrayView(&SpawnDataInitializer, 1), ProcessingContext, &CreationContext->GetEntityCollection());
+		UE::Mass::Executor::RunProcessorsView(MakeArrayView(&SpawnDataInitializer, 1), ProcessingContext, CreationContext->GetEntityCollections());
 	}
 
 	OutEntities.Append(MoveTemp(SpawnedEntities));
