Mass unit tests minor fixes

diff --git a/Engine/Plugins/Runtime/MassEntity/Source/MassEntityTestSuite/Private/MassChunkCollectionTest.cpp b/Engine/Plugins/Runtime/MassEntity/Source/MassEntityTestSuite/Private/MassChunkCollectionTest.cpp
--- a/Engine/Plugins/Runtime/MassEntity/Source/MassEntityTestSuite/Private/MassChunkCollectionTest.cpp
+++ b/Engine/Plugins/Runtime/MassEntity/Source/MassEntityTestSuite/Private/MassChunkCollectionTest.cpp
@@ -238,8 +238,11 @@ struct FEntityCollection_CreateWithInvalidDuplicatesWithPayload : FEntityCollect
 		// transform typed payload array into generic one for sorting purposes
 		FStructArrayView PaloadView(Payload);
 		TArray<FMassArchetypeEntityCollectionWithPayload> Result;
-		FMassArchetypeEntityCollectionWithPayload::CreateEntityRangesWithPayload(*EntityManager, EntitiesSubSet, FMassArchetypeEntityCollection::FoldDuplicates
-			, FMassGenericPayloadView(MakeArrayView(&PaloadView, 1)), Result);
+		{
+			AITEST_SCOPED_CHECK("Invalid entity handle passed in", 2);
+			FMassArchetypeEntityCollectionWithPayload::CreateEntityRangesWithPayload(*EntityManager, EntitiesSubSet, FMassArchetypeEntityCollection::FoldDuplicates
+				, FMassGenericPayloadView(MakeArrayView(&PaloadView, 1)), Result);
+		}
 
 		AITEST_EQUAL(TEXT("We expect only a single result"), Result.Num(), 1);
 		AITEST_EQUAL(TEXT("We expect only a single resulting range"), Result[0].GetEntityCollection().GetRanges().Num(), 1);
diff --git a/Engine/Plugins/Runtime/MassEntity/Source/MassEntityTestSuite/Private/MassObserverTest.cpp b/Engine/Plugins/Runtime/MassEntity/Source/MassEntityTestSuite/Private/MassObserverTest.cpp
--- a/Engine/Plugins/Runtime/MassEntity/Source/MassEntityTestSuite/Private/MassObserverTest.cpp
+++ b/Engine/Plugins/Runtime/MassEntity/Source/MassEntityTestSuite/Private/MassObserverTest.cpp
@@ -380,6 +380,11 @@ struct FFragmentObserverTest_MultipleArchetypeAdd : FFragmentTestBase
 		{
 			EntityManager->Defer().AddFragment<FFragmentStruct>(ModifiedEntity);
 		}
+		// also adding the fragment to the other archetype that already has the fragment. This should not yield any results
+		for (const FMassEntityHandle& OtherEntity : EntitiesIntsFloat)
+		{
+			EntityManager->Defer().AddFragment<FFragmentStruct>(OtherEntity);
+		}
 		return true;
 	}
 };
@@ -391,7 +396,7 @@ struct FFragmentObserverTest_MultipleArchetypeRemove : FFragmentTestBase
 
 	virtual bool PerformOperation() override
 	{
-		ExpectedEntities = { EntitiesInt[0], EntitiesInt[2], EntitiesInt[1] };
+		ExpectedEntities = { EntitiesInt[0], EntitiesInt[2], EntitiesIntsFloat[1] };
 		for (const FMassEntityHandle& ModifiedEntity : ExpectedEntities)
 		{
 			EntityManager->Defer().AddFragment<FFragmentStruct>(ModifiedEntity);
@@ -414,7 +419,7 @@ struct FFragmentObserverTest_MultipleArchetypeDestroy : FFragmentTestBase
 
 	virtual bool PerformOperation() override
 	{
-		ExpectedEntities = { EntitiesInt[0], EntitiesInt[2], EntitiesInt[1] };
+		ExpectedEntities = { EntitiesInt[0], EntitiesInt[2], EntitiesIntsFloat[1] };
 		for (const FMassEntityHandle& ModifiedEntity : ExpectedEntities)
 		{
 			EntityManager->Defer().AddFragment<FFragmentStruct>(ModifiedEntity);
