Added batched add capabilities to MassEntity Storage, along with unit tests for the new functionality.

diff --git a/Engine/Plugins/Runtime/MassEntity/Source/MassEntity/Private/MassEntityManager.cpp b/Engine/Plugins/Runtime/MassEntity/Source/MassEntity/Private/MassEntityManager.cpp
--- a/Engine/Plugins/Runtime/MassEntity/Source/MassEntity/Private/MassEntityManager.cpp
+++ b/Engine/Plugins/Runtime/MassEntity/Source/MassEntity/Private/MassEntityManager.cpp
@@ -751,13 +751,11 @@ void FMassEntityManager::BuildEntity(FMassEntityHandle Entity, TConstArrayView<F
 
 TConstArrayView<FMassEntityHandle> FMassEntityManager::BatchReserveEntities(const int32 Count, TArray<FMassEntityHandle>& InOutEntities)
 {
-	int32 Index = InOutEntities.Num();
-	InOutEntities.Reserve(Index + Count);
-	for (int32 Counter = 0; Counter < Count; ++Counter)
-	{
-		InOutEntities.Add(ReserveEntity());
-	}
-	return MakeArrayView(InOutEntities.GetData() + Index, Count);
+	const int32 Index = InOutEntities.Num();
+	const int32 NumAdded = GetEntityStorageInterface().Acquire(Count, InOutEntities);
+	ensureMsgf(NumAdded == Count, TEXT("Failed to reserve %d entities, was able to only reserve %d"), Count, NumAdded);
+
+	return MakeArrayView(InOutEntities.GetData() + Index, NumAdded);
 }
 
 TSharedRef<FMassEntityManager::FEntityCreationContext> FMassEntityManager::BatchBuildEntities(const FMassArchetypeEntityCollectionWithPayload& EncodedEntitiesWithPayload
@@ -925,13 +923,15 @@ void FMassEntityManager::BatchDestroyEntities(TConstArrayView<FMassEntityHandle>
 			continue;
 		}
 
-		FMassArchetypeData* Archetype = GetEntityStorageInterface().GetArchetype(Entity.Index);
-		check(Archetype);
-		ObserverManager.OnPreEntityDestroyed(Archetype->GetCompositionDescriptor(), Entity);
-		Archetype->RemoveEntity(Entity);
-
-		GetEntityStorageInterface().ReleaseOne(Entity);
+		if (FMassArchetypeData* Archetype = GetEntityStorageInterface().GetArchetype(Entity.Index))
+		{
+			ObserverManager.OnPreEntityDestroyed(Archetype->GetCompositionDescriptor(), Entity);
+			Archetype->RemoveEntity(Entity);
+		}
+		// else it's a "reserved" entity so it has not been assigned to an archetype yet, no archetype nor observers to notify
 	}
+
+	GetEntityStorageInterface().Release(InEntities);
 }
 
 void FMassEntityManager::BatchDestroyEntityChunks(const FMassArchetypeEntityCollection& EntityCollection)
@@ -1676,7 +1676,9 @@ const FSharedStruct* FMassEntityManager::InternalGetSharedFragmentPtr(FMassEntit
 
 bool FMassEntityManager::IsEntityValid(FMassEntityHandle Entity) const
 {
-	return (Entity.Index != UE::Mass::Private::InvalidEntityIndex) && GetEntityStorageInterface().IsValidIndex(Entity.Index) && (GetEntityStorageInterface().GetSerialNumber(Entity.Index) == Entity.SerialNumber);
+	return (Entity.Index != UE::Mass::Private::InvalidEntityIndex) 
+		&& GetEntityStorageInterface().IsValidIndex(Entity.Index) 
+		&& (GetEntityStorageInterface().GetSerialNumber(Entity.Index) == Entity.SerialNumber);
 }
 
 bool FMassEntityManager::IsEntityBuilt(FMassEntityHandle Entity) const
@@ -2116,8 +2118,9 @@ void FMassEntityManager::BatchBuildEntities(const FMassArchetypeEntityCollection
 	Params.DebugName = ArchetypeDebugName;
 	BatchBuildEntities(EncodedEntitiesWithPayload, MoveTemp(Composition), SharedFragmentValues, Params);
 }
+
 const FMassArchetypeEntityCollection& FMassEntityManager::FEntityCreationContext::GetEntityCollection() const
 {
 	static FMassArchetypeEntityCollection EmptyCollection;
 	return EntityCollections.Num() ? EntityCollections[0] : EmptyCollection;
-}
\ No newline at end of file
+}
diff --git a/Engine/Plugins/Runtime/MassEntity/Source/MassEntity/Private/MassEntityManagerStorage.cpp b/Engine/Plugins/Runtime/MassEntity/Source/MassEntity/Private/MassEntityManagerStorage.cpp
--- a/Engine/Plugins/Runtime/MassEntity/Source/MassEntity/Private/MassEntityManagerStorage.cpp
+++ b/Engine/Plugins/Runtime/MassEntity/Source/MassEntity/Private/MassEntityManagerStorage.cpp
@@ -1,4 +1,4 @@
-﻿// Copyright Epic Games, Inc. All Rights Reserved.
+// Copyright Epic Games, Inc. All Rights Reserved.
 
 #include "MassEntityManagerStorage.h"
 
@@ -7,8 +7,9 @@
 
 namespace UE::Mass
 {
-	//////////////////////////////////////////////////////////////////////
+	//-----------------------------------------------------------------------------
 	// FSingleThreadedEntityStorage
+	//-----------------------------------------------------------------------------
 	
 	FMassArchetypeData* FSingleThreadedEntityStorage::GetArchetype(int32 Index)
 	{
@@ -42,19 +43,13 @@ namespace UE::Mass
 
 	IEntityStorageInterface::EEntityState FSingleThreadedEntityStorage::GetEntityState(int32 Index) const
 	{
-		const FMassArchetypeData* CurrentArchetype = Entities[Index].CurrentArchetype.Get();
 		const uint32 CurrentSerialNumber = Entities[Index].SerialNumber;
 
 		if (CurrentSerialNumber != 0)
 		{
-			if (CurrentArchetype != nullptr)
-			{
-				return EEntityState::Created;
-			}
-			else // (CurrentArchetype == nullptr)
-			{
-				return EEntityState::Reserved;
-			}
+			return Entities[Index].CurrentArchetype.Get()
+				? EEntityState::Created 
+				: EEntityState::Reserved;
 		}
 
 		return EEntityState::Free;	
@@ -92,13 +87,53 @@ namespace UE::Mass
 		return Handle;
 	}
 
-	int32 FSingleThreadedEntityStorage::Release(TArrayView<FMassEntityHandle> Handles)
+	int32 FSingleThreadedEntityStorage::Acquire(const int32 Count, TArray<FMassEntityHandle>& OutEntityHandles)
+	{
+		check(Count >= 0);
+
+		const int32 SerialNumber = SerialNumberGenerator.fetch_add(1);
+
+		int32 NumAdded = 0;
+
+		OutEntityHandles.Reserve(OutEntityHandles.Num() + Count);
+
+		const int32 NumAvailableFromFreeList = FMath::Min(Count, EntityFreeIndexList.Num());
+		if (NumAvailableFromFreeList > 0)
+		{
+			const int32 FirstIndexToUse = EntityFreeIndexList.Num() - NumAvailableFromFreeList;
+			for (int32 Index = FirstIndexToUse; Index < EntityFreeIndexList.Num(); ++Index)
+			{
+				const int32 EntityIndex = EntityFreeIndexList[Index];
+				Entities[EntityIndex].SerialNumber = SerialNumber;
+				OutEntityHandles.Add({ EntityIndex, SerialNumber });
+			}
+			EntityFreeIndexList.RemoveAt(FirstIndexToUse, NumAvailableFromFreeList, EAllowShrinking::No);
+			NumAdded = NumAvailableFromFreeList;
+		}
+
+		if (NumAdded < Count)
+		{
+			const int32 RemainingCount = Count - NumAdded;
+			const int32 StartingIndex = Entities.Num();
+			Entities.Add(RemainingCount);
+			for (int32 EntityIndex = StartingIndex; EntityIndex < Entities.Num(); ++EntityIndex)
+			{
+				Entities[EntityIndex].SerialNumber = SerialNumber;
+				OutEntityHandles.Add({ EntityIndex, SerialNumber });
+			}
+			NumAdded += RemainingCount;
+		}
+
+		return NumAdded;
+	}
+
+	int32 FSingleThreadedEntityStorage::Release(TConstArrayView<FMassEntityHandle> Handles)
 	{
 		int DeallocateCount = 0;
 
 		EntityFreeIndexList.Reserve(EntityFreeIndexList.Num() + Handles.Num());
 
-		for (FMassEntityHandle& Handle : Handles)
+		for (const FMassEntityHandle& Handle : Handles)
 		{
 			FEntityData& EntityData = Entities[Handle.Index];
 			if (EntityData.SerialNumber == Handle.SerialNumber)
@@ -117,10 +152,10 @@ namespace UE::Mass
 		return Release(MakeArrayView(&Handle, 1));
 	}
 
-	int32 FSingleThreadedEntityStorage::ForceRelease(TArrayView<FMassEntityHandle> Handles)
+	int32 FSingleThreadedEntityStorage::ForceRelease(TConstArrayView<FMassEntityHandle> Handles)
 	{
 		EntityFreeIndexList.Reserve(EntityFreeIndexList.Num() + Handles.Num());
-		for (FMassEntityHandle& Handle : Handles)
+		for (const FMassEntityHandle& Handle : Handles)
 		{
 			FEntityData& EntityData = Entities[Handle.Index];
 			EntityData.Reset();
@@ -144,6 +179,10 @@ namespace UE::Mass
 		return EntityFreeIndexList.Num();
 	}
 
+	//-----------------------------------------------------------------------------
+	// FSingleThreadedEntityStorage::FEntityData
+	//-----------------------------------------------------------------------------
+
 	FSingleThreadedEntityStorage::FEntityData::~FEntityData() = default;
 
 	void FSingleThreadedEntityStorage::FEntityData::Reset()
@@ -157,8 +196,9 @@ namespace UE::Mass
 		return SerialNumber != 0 && CurrentArchetype.IsValid();
 	}
 
-	//////////////////////////////////////////////////////////////////////
+	//-----------------------------------------------------------------------------
 	// FConcurrentEntityStorage
+	//-----------------------------------------------------------------------------
 
 	void FConcurrentEntityStorage::Initialize(const FMassEntityManager_InitParams_Concurrent& InInitializationParams)
 	{
@@ -223,7 +263,7 @@ namespace UE::Mass
 	IEntityStorageInterface::EEntityState FConcurrentEntityStorage::GetEntityState(int32 Index) const
 	{
 		//
-		// || Archetype || IsAllocated || Result    |
+		// || Archetype || bIsAllocated || Result    |
 		//  |  nullptr   |      0       |  Free     |
 		//  |  nullptr   |      1       |  Reserved |
 		//  | !nullptr   |      1       |  Created  |
@@ -234,14 +274,9 @@ namespace UE::Mass
 		{
 			return EEntityState::Created;
 		}
-		else // EntityData.CurrentArchetype == nullptr
-		{
-			if (EntityData.IsAllocated == 1)
-			{
-				return EEntityState::Reserved;
-			}
-			return EEntityState::Free;
-		}
+		return EntityData.bIsAllocated
+			? EEntityState::Reserved
+			: EEntityState::Free;
 	}
 
 	int32 FConcurrentEntityStorage::GetSerialNumber(int32 Index) const
@@ -281,46 +316,59 @@ namespace UE::Mass
 		return LookupEntity(Index).CurrentArchetype != nullptr;
 	}
 
-	FMassEntityHandle FConcurrentEntityStorage::AcquireOne()
+	bool FConcurrentEntityStorage::AddPage()
 	{
-		int32 EntityIndex;
-		{
-			UE::TUniqueLock FreeListLock(FreeListMutex);
+		check(FreeListMutex.IsLocked());
+		UE::TUniqueLock PageAllocateLock(PageAllocateMutex);
+
+		// Allocate new page
+		const uint32 NewPageIndex = PageCount;
+		checkf((NewPageIndex + 1) * MaxEntitiesPerPage < (1llu << MaximumEntityCountShift), TEXT("Exhausted number of entities"));
+
+		const uint64 PageSize = ComputePageSize();
+		FEntityData* Page = static_cast<FEntityData*>(FMemory::Malloc(PageSize, alignof(FEntityData)));
 		
-			if (UNLIKELY(EntityFreeIndexList.IsEmpty()))
-			{
-				check(FreeListMutex.IsLocked());
-				UE::TUniqueLock PageAllocateLock(PageAllocateMutex);
+		if (Page == nullptr)
+		{
+			return false;
+		}
 
-				// Allocate new page
-				const uint32 NewPageIndex = PageCount;
-				checkf((NewPageIndex + 1) * MaxEntitiesPerPage < (1llu << MaximumEntityCountShift), TEXT("Exahusted number of entities"));
+		/*for (int32 Index = 0, End = MaxEntitiesPerPage; Index < End; ++Index)
+		{
+			new (Page + Index) FEntityData();
+		}*/
+		FMemory::Memzero(Page, PageSize);
 
-				const uint64 PageSize = ComputePageSize();
-				FEntityData* Page = static_cast<FEntityData*>(FMemory::Malloc(PageSize, alignof(FEntityData)));
-			
-				for (int32 Index = 0, End = MaxEntitiesPerPage; Index < End; ++Index)
-				{
-					new (Page + Index) FEntityData();
-				}
+		EntityPages[PageCount] = Page;
+		++PageCount;
 
-				EntityPages[PageCount] = Page;
-				++PageCount;
+		const int32 NewEntityIndexStart = NewPageIndex * MaxEntitiesPerPage;
+		const int32 NewEntityIndexEnd = (NewPageIndex + 1) * MaxEntitiesPerPage;
 
-				const int32 NewEntityIndexStart = NewPageIndex * MaxEntitiesPerPage;
-				const int32 NewEntityIndexEnd = (NewPageIndex + 1) * MaxEntitiesPerPage;
+		EntityFreeIndexList.Reserve(MaxEntitiesPerPage);
 
-				EntityFreeIndexList.Reserve(MaxEntitiesPerPage);
+		// Somewhat tricksy thing here to be aware of
+		// MassEntityManager expects the very first allocated entity to be at index 0
+		// However, EntityFreeIndexList.Pop() will return the last one added to the list
+		// Therefore, populate the free list backwards
+		for (int32 NewEntityIndex = NewEntityIndexEnd - 1; NewEntityIndex >= NewEntityIndexStart; --NewEntityIndex)
+		{
+			// Setup the free list
+			EntityFreeIndexList.Push(NewEntityIndex);
+		}
 
-				// Somewhat tricksy thing here to be aware of
-				// MassEntityManager expects the very first allocated entity to be at index 0
-				// However, EntityFreeIndexList.Pop() will return the last one added to the list
-				// Therefore, populate the free list backwards
-				for (int32 NewEntityIndex = NewEntityIndexEnd - 1; NewEntityIndex >= NewEntityIndexStart; --NewEntityIndex)
-				{
-					// Setup the free list
-					EntityFreeIndexList.Push(NewEntityIndex);
-				}
+		return true;
+	}
+
+	FMassEntityHandle FConcurrentEntityStorage::AcquireOne()
+	{
+		int32 EntityIndex;
+		{
+			UE::TUniqueLock FreeListLock(FreeListMutex);
+		
+			if (UNLIKELY(EntityFreeIndexList.IsEmpty()))
+			{
+				AddPage();
 			}
 
 			EntityIndex = EntityFreeIndexList.Pop(EAllowShrinking::No);
@@ -331,7 +379,7 @@ namespace UE::Mass
 		// that SerialNum == 0 means an invalid Entity.  FMassArchetypeEntityCollection uses this assumption
 		// and will fail IsValid() checks otherwise.
 		++EntityData.GenerationId;
-		EntityData.IsAllocated = 1;
+		EntityData.bIsAllocated = 1;
 		int32 SerialNumber = EntityData.GetSerialNumber();
 
 		EntityCount.fetch_add(1llu);
@@ -342,7 +390,54 @@ namespace UE::Mass
 		return Handle;
 	}
 
-	int32 FConcurrentEntityStorage::Release(TArrayView<FMassEntityHandle> Handles)
+	int32 FConcurrentEntityStorage::Acquire(const int32 Count, TArray<FMassEntityHandle>& OutEntityHandles)
+	{
+		check(Count >= 0);
+
+		int32 CountAdded = 0;
+		int32 CountLeft = Count;
+
+		OutEntityHandles.Reserve(OutEntityHandles.Num() + Count);
+
+		while (CountLeft > 0)
+		{
+			UE::TUniqueLock FreeListLock(FreeListMutex);
+
+			if (UNLIKELY(EntityFreeIndexList.IsEmpty()))
+			{
+				if (AddPage() == false)
+				{
+					break;
+				}
+			}
+
+			const int32 CountToProcess = FMath::Min(CountLeft, EntityFreeIndexList.Num());
+
+			for (int32 Iteration = 0; Iteration < CountToProcess; ++Iteration)
+			{
+				const int32 EntityIndex = EntityFreeIndexList.Pop(EAllowShrinking::No);
+
+				FEntityData& EntityData = LookupEntity(EntityIndex);
+				// NOTE: Technically should not be necessary, however FEntityHandle::IsValid() makes the assumption
+				// that SerialNum == 0 means an invalid Entity.  FMassArchetypeEntityCollection uses this assumption
+				// and will fail IsValid() checks otherwise.
+				++EntityData.GenerationId;
+				EntityData.bIsAllocated = 1;
+				const int32 SerialNumber = EntityData.GetSerialNumber();
+
+				OutEntityHandles.Add({ EntityIndex, SerialNumber });
+			}
+			
+			CountAdded += CountToProcess;
+			CountLeft -= CountToProcess;
+		}
+
+		EntityCount.fetch_add(CountAdded);
+
+		return CountAdded;
+	}
+
+	int32 FConcurrentEntityStorage::Release(TConstArrayView<FMassEntityHandle> Handles)
 	{
 		int32 DeallocateCount = 0;
 	
@@ -358,7 +453,7 @@ namespace UE::Mass
 				EntityFreeIndexList.Reserve(EntityFreeIndexList.Num() + AllocatedRunLength);
 				for (int32 IndexToFree = BeginHandlesIndexToFree; IndexToFree < BeginHandlesIndexToFree + AllocatedRunLength; ++IndexToFree)
 				{
-					FMassEntityHandle& HandleToFree = Handles[IndexToFree];
+					const FMassEntityHandle& HandleToFree = Handles[IndexToFree];
 					EntityFreeIndexList.Add(HandleToFree.Index);
 				}
 			}
@@ -368,14 +463,14 @@ namespace UE::Mass
 	
 		for (int32 Index = 0, End = Handles.Num(); Index < End; ++Index)
 		{
-			FMassEntityHandle& Handle = Handles[Index];
+			const FMassEntityHandle& Handle = Handles[Index];
 			FEntityData& EntityData = LookupEntity(Handle.Index);
 			if (EntityData.GetSerialNumber() == Handle.SerialNumber)
 			{
 				++AllocatedRunLength;
 			
 				++EntityData.GenerationId;
-				EntityData.IsAllocated = 0;
+				EntityData.bIsAllocated = 0;
 				EntityData.CurrentArchetype.Reset();
 			
 				++DeallocateCount;
@@ -402,23 +497,23 @@ namespace UE::Mass
 		return Release(MakeArrayView(&Handle, 1));
 	}
 
-	int32 FConcurrentEntityStorage::ForceRelease(TArrayView<FMassEntityHandle> Handles)
+	int32 FConcurrentEntityStorage::ForceRelease(TConstArrayView<FMassEntityHandle> Handles)
 	{
 		// ForceRelease assumes the caller knows all handles are allocated
 		// no need to have complexity of tracking "runs" of handles 
-		for (FMassEntityHandle& Handle : Handles)
+		for (const FMassEntityHandle& Handle : Handles)
 		{
 			FEntityData& EntityData = LookupEntity(Handle.Index);
 
 			++EntityData.GenerationId;
-			EntityData.IsAllocated = 0;
+			EntityData.bIsAllocated = 0;
 			EntityData.CurrentArchetype.Reset();
 		}
 
 		{
 			UE::TUniqueLock FreeListLock(FreeListMutex);
 			EntityFreeIndexList.Reserve(EntityFreeIndexList.Num() + Handles.Num());
-			for (FMassEntityHandle& Handle : Handles)
+			for (const FMassEntityHandle& Handle : Handles)
 			{
 				EntityFreeIndexList.Add(Handle.Index);
 			}
@@ -444,13 +539,6 @@ namespace UE::Mass
 		return EntityFreeIndexList.Num();
 	}
 
-	FConcurrentEntityStorage::FEntityData::~FEntityData() = default;
-
-	int32 FConcurrentEntityStorage::FEntityData::GetSerialNumber() const
-	{
-		return static_cast<int32>(GenerationId);
-	}
-
 	FConcurrentEntityStorage::FEntityData& FConcurrentEntityStorage::LookupEntity(int32 Index)
 	{
 		check(Index >= 0);
@@ -477,4 +565,40 @@ namespace UE::Mass
 	{
 		return sizeof(FEntityData) * MaxEntitiesPerPage;
 	}
+
+#if WITH_MASSENTITY_DEBUG
+	bool FConcurrentEntityStorage::DebugAssumptionsSelfTest()
+	{
+		// future proofing in case FEntityData's or TSharedPtr's internals change and make MemZero-ing not produce 
+		// the same results as default FEntityData's constructor
+		FEntityData DefaultData;
+		FEntityData ZeroedData;
+		FMemory::Memzero(&ZeroedData, sizeof(FEntityData));
+
+		if (DefaultData != ZeroedData)
+		{
+			UE_LOG(LogMass, Error, TEXT("%hs assumption about default FEntityData values is no longer true."), __FUNCTION__);
+			return false;
+		}
+
+		return true;
+	}
+#endif // WITH_MASSENTITY_DEBUG
+
+	//-----------------------------------------------------------------------------
+	// FConcurrentEntityStorage::FEntityData
+	//-----------------------------------------------------------------------------
+	FConcurrentEntityStorage::FEntityData::~FEntityData() = default;
+
+	int32 FConcurrentEntityStorage::FEntityData::GetSerialNumber() const
+	{
+		return static_cast<int32>(GenerationId);
+	}
+
+	bool FConcurrentEntityStorage::FEntityData::operator==(const FEntityData& Other) const
+	{
+		return CurrentArchetype == Other.CurrentArchetype
+			&& GenerationId == Other.GenerationId
+			&& bIsAllocated == Other.bIsAllocated;
+	}
 }
diff --git a/Engine/Plugins/Runtime/MassEntity/Source/MassEntity/Public/MassEntityManager.h b/Engine/Plugins/Runtime/MassEntity/Source/MassEntity/Public/MassEntityManager.h
--- a/Engine/Plugins/Runtime/MassEntity/Source/MassEntity/Public/MassEntityManager.h
+++ b/Engine/Plugins/Runtime/MassEntity/Source/MassEntity/Public/MassEntityManager.h
@@ -312,11 +312,17 @@ public:
 	void ReleaseReservedEntity(FMassEntityHandle Entity);
 
 	/**
-	 * Destroys all the entity in the provided array of entities
+	 * Destroys all the entities in the provided array of entities. The function will also gracefully handle entities
+	 * that have been reserved but not created yet.
 	 * @param InEntities to destroy
 	 */
 	void BatchDestroyEntities(TConstArrayView<FMassEntityHandle> InEntities);
 
+	/**
+	 * Destroys all the entities provided via the Collection. The function will also gracefully handle entities
+	 * that have been reserved but not created yet.
+	 * @param Collection to destroy
+	 */
 	void BatchDestroyEntityChunks(const FMassArchetypeEntityCollection& Collection);
 
 	void AddFragmentToEntity(FMassEntityHandle Entity, const UScriptStruct* FragmentType);
diff --git a/Engine/Plugins/Runtime/MassEntity/Source/MassEntity/Public/MassEntityManagerStorage.h b/Engine/Plugins/Runtime/MassEntity/Source/MassEntity/Public/MassEntityManagerStorage.h
--- a/Engine/Plugins/Runtime/MassEntity/Source/MassEntity/Public/MassEntityManagerStorage.h
+++ b/Engine/Plugins/Runtime/MassEntity/Source/MassEntity/Public/MassEntityManagerStorage.h
@@ -5,6 +5,7 @@
 #include "Containers/ChunkedArray.h"
 #include "Misc/TVariant.h"
 #include "Templates/SharedPointer.h"
+#include "MassProcessingTypes.h"
 
 struct FMassArchetypeData;
 struct FMassEntityHandle;
@@ -76,13 +77,15 @@ namespace UE::Mass
 		virtual bool IsValid(int32 Index) const = 0;
 
 		virtual FMassEntityHandle AcquireOne() = 0;
+		// @return number of entities actually added
+		virtual int32 Acquire(const int32 Count, TArray<FMassEntityHandle>& OutEntityHandles) = 0;
 		
-		virtual int32 Release(TArrayView<FMassEntityHandle> Handles) = 0;
+		virtual int32 Release(TConstArrayView<FMassEntityHandle> Handles) = 0;
 		virtual int32 ReleaseOne(FMassEntityHandle Handles) = 0;
 		
 		// Bypasses Serial Number Check
 		// Only use if caller has ensured serial number matches or for debug purposes
-		virtual int32 ForceRelease(TArrayView<FMassEntityHandle> Handles) = 0;
+		virtual int32 ForceRelease(TConstArrayView<FMassEntityHandle> Handles) = 0;
 		virtual int32 ForceReleaseOne(FMassEntityHandle Handle) = 0;
 
 		// Returns the number of entities that are not free
@@ -117,9 +120,10 @@ namespace UE::Mass
 		virtual SIZE_T GetAllocatedSize() const override;
 		virtual bool IsValid(int32 Index) const override;
 		virtual FMassEntityHandle AcquireOne() override;
-		virtual int32 Release(TArrayView<FMassEntityHandle> Handles) override;
+		virtual int32 Acquire(const int32 Count, TArray<FMassEntityHandle>& OutEntityHandles) override;
+		virtual int32 Release(TConstArrayView<FMassEntityHandle> Handles) override;
 		virtual int32 ReleaseOne(FMassEntityHandle Handle) override;
-		virtual int32 ForceRelease(TArrayView<FMassEntityHandle> Handles) override;
+		virtual int32 ForceRelease(TConstArrayView<FMassEntityHandle> Handles) override;
 		virtual int32 ForceReleaseOne(FMassEntityHandle Handle) override;
 		virtual int32 Num() const override;
 		virtual int32 ComputeFreeSize() const override;
@@ -169,12 +173,17 @@ namespace UE::Mass
 		virtual SIZE_T GetAllocatedSize() const override;
 		virtual bool IsValid(int32 Index) const override;
 		virtual FMassEntityHandle AcquireOne() override;
-		virtual int32 Release(TArrayView<FMassEntityHandle> Handles) override;
+		virtual int32 Acquire(const int32 Count, TArray<FMassEntityHandle>& OutEntityHandles) override;
+		virtual int32 Release(TConstArrayView<FMassEntityHandle> Handles) override;
 		virtual int32 ReleaseOne(FMassEntityHandle Handle) override;
-		virtual int32 ForceRelease(TArrayView<FMassEntityHandle> Handles) override;
+		virtual int32 ForceRelease(TConstArrayView<FMassEntityHandle> Handles) override;
 		virtual int32 ForceReleaseOne(FMassEntityHandle Handle) override;
 		virtual int32 Num() const override;
 		virtual int32 ComputeFreeSize() const override;
+#if WITH_MASSENTITY_DEBUG
+		/** @return whether the assumptions are still valid */
+		MASSENTITY_API static bool DebugAssumptionsSelfTest();
+#endif // WITH_MASSENTITY_DEBUG
 	private:	
 	
 		struct FEntityData
@@ -183,13 +192,15 @@ namespace UE::Mass
 		
 			TSharedPtr<FMassArchetypeData> CurrentArchetype;
 			// Generation ID or version of the entity in this slot
-			uint32 GenerationId  : MaxGenerationBits = 0;
+			uint32 GenerationId : MaxGenerationBits = 0;
 			// 1 if the entity is NOT free
-			uint32 IsAllocated : 1 = 0;
+			uint32 bIsAllocated : 1 = 0;
 
 			~FEntityData();
 			// Converts EntityData state into a SerialNumber for public usage
 			int32 GetSerialNumber() const;
+
+			bool operator==(const FEntityData& Other) const;
 		};
 	
 		FEntityData& LookupEntity(int32 Index);
@@ -198,6 +209,9 @@ namespace UE::Mass
 		// Returns size of a page in bytes
 		uint64 ComputePageSize() const;
 
+		// @return whether the operation was successful. Will return false when OOM
+		bool AddPage();
+
 		// Number of allocated Entities
 		std::atomic<uint64> EntityCount = 0;
 		uint32 MaximumEntityCountShift = 0;
diff --git a/Engine/Plugins/Runtime/MassEntity/Source/MassEntityTestSuite/Private/MassEntityStorageTest.cpp b/Engine/Plugins/Runtime/MassEntity/Source/MassEntityTestSuite/Private/MassEntityStorageTest.cpp
new file mode 100644
--- /dev/null
+++ b/Engine/Plugins/Runtime/MassEntity/Source/MassEntityTestSuite/Private/MassEntityStorageTest.cpp
@@ -0,0 +1,389 @@
+// Copyright Epic Games, Inc. All Rights Reserved.
+
+#include "AITestsCommon.h"
+#include "MassEntityManager.h"
+
+#define LOCTEXT_NAMESPACE "MassTest"
+
+UE_DISABLE_OPTIMIZATION_SHIP
+
+namespace FMassEntityStorageTest
+{
+
+enum class EEntityManagerMode : int32
+{
+	Sequential,
+	Concurrent
+};
+
+struct FEntityStorageTestBase : FAITestBase
+{
+	TArray<TSharedPtr<FMassEntityManager>> EntityManagers;
+	TArray<EEntityManagerMode> ManagersToCreate;
+	const uint32 MaxConcurrentEntitiesPerPage = 0;
+	const int32 TotalNumToReserve = 0;
+	int32 NumToReserveInOneIteration = 0;
+	int32 NumToReleaseInOneIteration = 0;
+	FRandomStream RandomStream;
+
+	FEntityStorageTestBase()
+		// default setup: one sequential and one concurrent
+		: ManagersToCreate({ EEntityManagerMode::Sequential, EEntityManagerMode::Concurrent })
+		, MaxConcurrentEntitiesPerPage(FMassEntityManager_InitParams_Concurrent().MaxEntitiesPerPage)
+		// note that MaxConcurrentEntitiesPerPage doesn't really have any meaning for Sequential storage, but it's
+		// as good of a value to use for this test as any other
+		, TotalNumToReserve(MaxConcurrentEntitiesPerPage * 3 / 2)
+		, RandomStream(1)
+	{
+		NumToReserveInOneIteration = TotalNumToReserve / 10;
+		NumToReleaseInOneIteration = TotalNumToReserve / 12;
+	}
+
+	enum EEntityOperation
+	{
+		Operation_IndividualReserve,
+		Operation_BatchReserve,
+		Operation_IndividualDestroy,
+		Operation_BatchDestroy,
+	};
+
+	void PerformOperation(const EEntityOperation CurrentOperation, TSharedPtr<FMassEntityManager>& EntityManager, TArray<FMassEntityHandle>& EntitiesReserved)
+	{
+		switch (CurrentOperation)
+		{
+		case Operation_IndividualReserve:
+			for (int32 Counter = 0; Counter < NumToReserveInOneIteration; ++Counter)
+			{
+				EntitiesReserved.Add(EntityManager->ReserveEntity());
+			}
+			break;
+
+		case Operation_BatchReserve:
+			EntityManager->BatchReserveEntities(NumToReserveInOneIteration, EntitiesReserved);
+			break;
+
+		case Operation_IndividualDestroy:
+		{
+			const int32 NumToRelease = FMath::Min(NumToReleaseInOneIteration, EntitiesReserved.Num());
+			for (int32 Counter = 0; Counter < NumToRelease; ++Counter)
+			{
+				const int32 Index = RandomStream.RandRange(0, EntitiesReserved.Num() - 1);
+				EntityManager->ReleaseReservedEntity(EntitiesReserved[Index]);
+				EntitiesReserved.RemoveAtSwap(Index, EAllowShrinking::No);
+			}
+		}
+		break;
+
+		case Operation_BatchDestroy:
+		{
+			const int32 NumToRelease = FMath::Min(NumToReleaseInOneIteration, EntitiesReserved.Num());
+			TArray<FMassEntityHandle> EntitiesToDestroy;
+			EntitiesToDestroy.Reserve(NumToRelease);
+			for (int32 Counter = 0; Counter < NumToRelease; ++Counter)
+			{
+				const int32 Index = RandomStream.RandRange(0, EntitiesReserved.Num() - 1);
+				EntitiesToDestroy.Add(EntitiesReserved[Index]);
+				EntitiesReserved.RemoveAtSwap(Index, EAllowShrinking::No);
+			}
+			EntityManager->BatchDestroyEntities(EntitiesToDestroy);
+		}
+		break;
+		}
+	}
+
+	virtual bool SetUp() override
+	{
+		int32 Index = 0;
+		for (const EEntityManagerMode Mode : ManagersToCreate)
+		{
+			TSharedPtr<FMassEntityManager> LocalEntityManager = MakeShareable(new FMassEntityManager(nullptr));
+			LocalEntityManager->SetDebugName(FString::Printf(TEXT("TestEntityManager_%d"), Index));
+			++Index;
+
+			FMassEntityManagerStorageInitParams InitializationParams;
+			if (Mode == EEntityManagerMode::Sequential)
+			{
+				InitializationParams.Emplace<FMassEntityManager_InitParams_SingleThreaded>();
+			}
+			else // Mode == EEntityManagerMode::Concurrent
+			{
+				InitializationParams.Emplace<FMassEntityManager_InitParams_Concurrent>();
+				InitializationParams.Get<FMassEntityManager_InitParams_Concurrent>().MaxEntitiesPerPage = MaxConcurrentEntitiesPerPage;
+			}
+			LocalEntityManager->Initialize(InitializationParams);
+
+			EntityManagers.Add(LocalEntityManager);
+		}
+		
+		return true;
+	}
+
+	/** 
+	 * @return whether there are no duplicates. The return value matches the results produced by AITEST macros, 
+	 * thus the seemingly reversed function's meaning
+	 */
+	bool ValidateUniqueAndValidEntities(TSharedPtr<FMassEntityManager>& EntityManager, TConstArrayView<FMassEntityHandle> ContainerToTest) const
+	{
+		ContainerToTest.Sort();
+		for (int32 Index = ContainerToTest.Num() - 1; Index >= 0; --Index)
+		{
+			if (Index > 0)
+			{
+				AITEST_FALSE("Checking for duplicates", ContainerToTest[Index].Index == ContainerToTest[Index - 1].Index);
+			}
+			AITEST_TRUE("We expect every handle to be valid", EntityManager->IsEntityValid(ContainerToTest[Index]));
+			AITEST_FALSE("None of the gathered entities is expected to have been built", EntityManager->IsEntityBuilt(ContainerToTest[Index]));
+		}
+
+		return true;
+	}
+};
+
+struct FEntityStorageTest_SingleEntityParity : FEntityStorageTestBase
+{	
+	virtual bool InstantTest() override
+	{
+		//AITEST_TRUE("The reserved entity should be a valid entity", EntityManager->IsEntityValid(ReservedEntity));
+		FMassEntityHandle ReservedEntitySequential = EntityManagers[int(EEntityManagerMode::Sequential)]->ReserveEntity();
+		FMassEntityHandle ReservedEntityConcurrent = EntityManagers[int(EEntityManagerMode::Concurrent)]->ReserveEntity();
+		AITEST_EQUAL("The reserved entities are expected to be the same, regardless of the storage type", ReservedEntitySequential, ReservedEntityConcurrent);
+
+		return true;
+	}
+};
+IMPLEMENT_AI_INSTANT_TEST(FEntityStorageTest_SingleEntityParity, "System.Mass.Storage.SingleEntityParity");
+
+/** This test base compares results of multiple calls to ReserveEntity and a single call to BatchReserveEntities */
+struct FEntityStorageTest_ParityBase : FEntityStorageTestBase
+{
+	virtual bool InstantTest() override
+	{
+		const int32 NumToReserve = MaxConcurrentEntitiesPerPage * 3 / 2;
+
+		TArray<FMassEntityHandle> EntitiesIndividual;
+		for (int32 Index = 0; Index < NumToReserve; ++Index)
+		{
+			EntitiesIndividual.Add(EntityManagers[0]->ReserveEntity());
+		}
+		AITEST_EQUAL("The number of individually reserved entities should match the requested count", EntitiesIndividual.Num(), NumToReserve);
+
+		TArray<FMassEntityHandle> EntitiesBatch;
+		EntityManagers[1]->BatchReserveEntities(NumToReserve, EntitiesBatch);
+		AITEST_EQUAL("The number of batch-reserved entities should match the requested count", EntitiesBatch.Num(), NumToReserve);
+
+		for (int Index = 0; Index < NumToReserve; ++Index)
+		{
+			AITEST_EQUAL("The reserved entities are expected to be the same, regardless of the storage type", EntitiesIndividual[Index].Index, EntitiesBatch[Index].Index);
+		}
+
+		return true;
+	}
+};
+
+struct FEntityStorageTest_ParitySequential : FEntityStorageTest_ParityBase
+{
+	FEntityStorageTest_ParitySequential()
+	{
+		ManagersToCreate = TArray<EEntityManagerMode>({ EEntityManagerMode::Sequential, EEntityManagerMode::Sequential });
+	}
+};
+IMPLEMENT_AI_INSTANT_TEST(FEntityStorageTest_ParitySequential, "System.Mass.Storage.Sequential.BatchParity");
+
+struct FEntityStorageTest_ParityConcurrent : FEntityStorageTest_ParityBase
+{
+	FEntityStorageTest_ParityConcurrent()
+	{
+		ManagersToCreate = TArray<EEntityManagerMode>({ EEntityManagerMode::Concurrent, EEntityManagerMode::Concurrent });
+	}
+};
+IMPLEMENT_AI_INSTANT_TEST(FEntityStorageTest_ParityConcurrent, "System.Mass.Storage.Concurrent.BatchParity");
+
+struct FEntityStorageTest_FreeListBase : FEntityStorageTestBase
+{
+	virtual bool InstantTest() override
+	{
+		// note that MaxConcurrentEntitiesPerPage doesn't really have any meaning for Sequential storage, but it's
+		// as good of a value to use for this test as any other
+		const int32 NumToReserve = MaxConcurrentEntitiesPerPage * 3 / 2;
+
+		TArray<FMassEntityHandle> EntitiesBaseline;
+		EntityManagers[0]->BatchReserveEntities(NumToReserve, EntitiesBaseline);
+
+		TArray<FMassEntityHandle> EntitiesTestedBatch;
+		// batch-reserving and batch-removing
+		{
+			EntityManagers[1]->BatchReserveEntities(NumToReserve, EntitiesTestedBatch);
+
+			TArray<FMassEntityHandle> EntitiesToModify;
+			EntitiesToModify.Reserve(EntitiesTestedBatch.Num() / 2);
+			for (int32 Index = EntitiesTestedBatch.Num() - 1; Index >= 0; --Index)
+			{
+				if (EntitiesTestedBatch[Index].Index % 2)
+				{
+					EntitiesToModify.Add(EntitiesTestedBatch[Index]);
+					EntitiesTestedBatch.RemoveAtSwap(Index, EAllowShrinking::No);
+				}
+			}
+			const int32 EntitiesRemovedCount = EntitiesToModify.Num();
+			EntityManagers[1]->BatchDestroyEntities(EntitiesToModify);
+			EntitiesToModify.Reset();
+			EntityManagers[1]->BatchReserveEntities(EntitiesRemovedCount, EntitiesToModify);
+			AITEST_EQUAL("Reserving after removing should produce the expected number of entity handles", EntitiesToModify.Num(), EntitiesRemovedCount);
+
+			EntitiesTestedBatch.Append(EntitiesToModify);
+			AITEST_EQUAL("Entity handle arrays are expected to be of the same size", EntitiesTestedBatch.Num(), EntitiesBaseline.Num());
+		}
+
+		// batch-reserving and removing one-by-one
+		TArray<FMassEntityHandle> EntitiesTestedIndividual;
+		{
+			EntityManagers[2]->BatchReserveEntities(NumToReserve, EntitiesTestedIndividual);
+
+			for (int32 Index = EntitiesTestedIndividual.Num() - 1; Index >= 0; --Index)
+			{
+				if (RandomStream.FRand() >= 0.5)
+				{
+					EntityManagers[2]->ReleaseReservedEntity(EntitiesTestedIndividual[Index]);
+					EntitiesTestedIndividual.RemoveAtSwap(Index, EAllowShrinking::No);
+				}
+			}
+
+			EntityManagers[2]->BatchReserveEntities(NumToReserve - EntitiesTestedIndividual.Num(), EntitiesTestedIndividual);
+			AITEST_EQUAL("Total number of entities reserved should be back to initial number", EntitiesTestedIndividual.Num(), NumToReserve);
+		}
+
+		EntitiesTestedIndividual.Sort();
+		EntitiesTestedBatch.Sort();
+		EntitiesBaseline.Sort();
+		for (int32 Index = 0; Index < EntitiesTestedBatch.Num(); ++Index)
+		{ 
+			AITEST_EQUAL("All entries are expected to be equivalent A", EntitiesTestedBatch[Index].Index, EntitiesBaseline[Index].Index);
+			AITEST_EQUAL("All entries are expected to be equivalent B", EntitiesTestedIndividual[Index].Index, EntitiesBaseline[Index].Index);
+		}
+
+		return true;
+	}
+};
+
+struct FEntityStorageTest_FreeListSequential : FEntityStorageTest_FreeListBase
+{
+	FEntityStorageTest_FreeListSequential()
+	{
+		ManagersToCreate = TArray<EEntityManagerMode>({ EEntityManagerMode::Sequential, EEntityManagerMode::Sequential, EEntityManagerMode::Sequential });
+	}
+};
+IMPLEMENT_AI_INSTANT_TEST(FEntityStorageTest_FreeListSequential, "System.Mass.Storage.Sequential.FreeList");
+
+struct FEntityStorageTest_FreeListConcurrent : FEntityStorageTest_FreeListBase
+{
+	FEntityStorageTest_FreeListConcurrent()
+	{
+		ManagersToCreate = TArray<EEntityManagerMode>({ EEntityManagerMode::Concurrent, EEntityManagerMode::Concurrent, EEntityManagerMode::Concurrent });
+	}
+};
+IMPLEMENT_AI_INSTANT_TEST(FEntityStorageTest_FreeListConcurrent, "System.Mass.Storage.Concurrent.FreeList");
+
+//-----------------------------------------------------------------------------
+// Special case that failed during development. 
+// The free list was not being updated properly after removing entities so a sequence of add-remove-add-add 
+// caused the last `add` to overlap with the last-but-one `add`
+//-----------------------------------------------------------------------------
+struct FEntityStorageTest_MultiReAddingBase : FEntityStorageTestBase
+{
+	virtual bool InstantTest() override
+	{
+		TSharedPtr<FMassEntityManager> EntityManager = EntityManagers[0];
+		TArray<FMassEntityHandle> EntitiesReserved;
+
+		PerformOperation(Operation_BatchReserve, EntityManager, EntitiesReserved);
+		PerformOperation(Operation_BatchDestroy, EntityManager, EntitiesReserved);
+		PerformOperation(Operation_BatchReserve, EntityManager, EntitiesReserved);
+		AITEST_TRUE("Testing for duplicates after first readding", ValidateUniqueAndValidEntities(EntityManager, EntitiesReserved))
+
+		PerformOperation(Operation_BatchReserve, EntityManager, EntitiesReserved);
+		AITEST_TRUE("Testing for duplicates after n+1 readding", ValidateUniqueAndValidEntities(EntityManager, EntitiesReserved))
+
+		return true;
+	}
+};
+
+struct FEntityStorageTest_MultiReAddingSequential : FEntityStorageTest_MultiReAddingBase
+{
+	FEntityStorageTest_MultiReAddingSequential()
+	{
+		ManagersToCreate = TArray<EEntityManagerMode>({ EEntityManagerMode::Sequential });
+	}
+};
+IMPLEMENT_AI_INSTANT_TEST(FEntityStorageTest_MultiReAddingSequential, "System.Mass.Storage.Sequential.AddRemoveLoop");
+
+struct FEntityStorageTest_MultiReAddingConcurrent : FEntityStorageTest_MultiReAddingBase
+{
+	FEntityStorageTest_MultiReAddingConcurrent()
+	{
+		ManagersToCreate = TArray<EEntityManagerMode>({ EEntityManagerMode::Concurrent });
+	}
+};
+IMPLEMENT_AI_INSTANT_TEST(FEntityStorageTest_MultiReAddingConcurrent, "System.Mass.Storage.Concurrent.AddRemoveLoop");
+
+//-----------------------------------------------------------------------------
+// catch-all test
+//-----------------------------------------------------------------------------
+struct FEntityStorageTest_MixedOperationsBase : FEntityStorageTestBase
+{
+	virtual bool InstantTest() override
+	{
+		TSharedPtr<FMassEntityManager> EntityManager = EntityManagers[0];
+
+		TArray<int32> OperationsPerformed;
+		TArray<FMassEntityHandle> EntitiesReserved;
+
+		EEntityOperation CurrentOperation = Operation_BatchReserve;
+		while (EntitiesReserved.Num() < TotalNumToReserve)
+		{
+			PerformOperation(CurrentOperation, EntityManager, EntitiesReserved);
+			const bool bValidColleciton = ValidateUniqueAndValidEntities(EntityManager, EntitiesReserved);
+			AITEST_TRUE("Testing for duplicates", bValidColleciton);
+			
+			OperationsPerformed.Add(CurrentOperation);
+			CurrentOperation = EEntityOperation(RandomStream.RandRange(0, 3));
+		}
+
+		return true;
+	}
+};
+
+struct FEntityStorageTest_MixedOperationsSequential : FEntityStorageTest_MixedOperationsBase
+{
+	FEntityStorageTest_MixedOperationsSequential()
+	{
+		ManagersToCreate = TArray<EEntityManagerMode>({ EEntityManagerMode::Sequential });
+	}
+};
+IMPLEMENT_AI_INSTANT_TEST(FEntityStorageTest_MixedOperationsSequential, "System.Mass.Storage.Sequential.MixedOperations");
+
+struct FEntityStorageTest_MixedOperationsConcurrent : FEntityStorageTest_MixedOperationsBase
+{
+	FEntityStorageTest_MixedOperationsConcurrent()
+	{
+		ManagersToCreate = TArray<EEntityManagerMode>({ EEntityManagerMode::Concurrent });
+	}
+};
+IMPLEMENT_AI_INSTANT_TEST(FEntityStorageTest_MixedOperationsConcurrent, "System.Mass.Storage.Concurrent.MixedOperations");
+
+#if WITH_MASSENTITY_DEBUG
+struct FEntityStorageTest_ConcurrentDataLayoutAssumptions : FAITestBase
+{
+	virtual bool InstantTest() override
+	{
+		const bool bAssumptionsValid = UE::Mass::FConcurrentEntityStorage::DebugAssumptionsSelfTest();
+		AITEST_TRUE("Testing assumptions", bAssumptionsValid);
+		return true;
+	}
+};
+IMPLEMENT_AI_INSTANT_TEST(FEntityStorageTest_ConcurrentDataLayoutAssumptions, "System.Mass.Storage.Concurrent.DataLayoutAssumptions");
+#endif // WITH_MASSENTITY_DEBUG
+} // FMassEntityTestTest
+
+UE_ENABLE_OPTIMIZATION_SHIP
+
+#undef LOCTEXT_NAMESPACE
