Fixed an issue with executing processors over an array of FMassArchetypeEntityCollection. The code was breaking the assumption that ExecutionContext is not expected to override stored entity collection.

diff --git a/Engine/Plugins/Runtime/MassEntity/Source/MassEntity/Private/MassExecutor.cpp b/Engine/Plugins/Runtime/MassEntity/Source/MassEntity/Private/MassExecutor.cpp
--- a/Engine/Plugins/Runtime/MassEntity/Source/MassEntity/Private/MassExecutor.cpp
+++ b/Engine/Plugins/Runtime/MassEntity/Source/MassEntity/Private/MassExecutor.cpp
@@ -122,6 +122,7 @@ void RunProcessorsView(TArrayView<UMassProcessor* const> Processors, FMassProces
 			{
 				ExecutionContext.SetEntityCollection(Collection);
 				ExecuteProcessors(*ProcessingContext.EntityManager, Processors, ExecutionContext);
+				ExecutionContext.ClearEntityCollection();
 			}
 		}
 	}
