Added a flavor of BatchReserveEntities that lets the caller use scratch-array to get the entity handles back

diff --git a/Engine/Plugins/Runtime/MassEntity/Source/MassEntity/Private/MassEntityManager.cpp b/Engine/Plugins/Runtime/MassEntity/Source/MassEntity/Private/MassEntityManager.cpp
--- a/Engine/Plugins/Runtime/MassEntity/Source/MassEntity/Private/MassEntityManager.cpp
+++ b/Engine/Plugins/Runtime/MassEntity/Source/MassEntity/Private/MassEntityManager.cpp
@@ -758,6 +758,11 @@ TConstArrayView<FMassEntityHandle> FMassEntityManager::BatchReserveEntities(cons
 	return MakeArrayView(InOutEntities.GetData() + Index, NumAdded);
 }
 
+int32 FMassEntityManager::BatchReserveEntities(TArrayView<FMassEntityHandle> InOutEntities)
+{
+	return GetEntityStorageInterface().Acquire(InOutEntities);
+}
+
 TSharedRef<FMassEntityManager::FEntityCreationContext> FMassEntityManager::BatchBuildEntities(const FMassArchetypeEntityCollectionWithPayload& EncodedEntitiesWithPayload
 	, const FMassFragmentBitSet& FragmentsAffected, const FMassArchetypeSharedFragmentValues& SharedFragmentValues, const FMassArchetypeCreationParams& CreationParams)
 {
diff --git a/Engine/Plugins/Runtime/MassEntity/Source/MassEntity/Private/MassEntityManagerStorage.cpp b/Engine/Plugins/Runtime/MassEntity/Source/MassEntity/Private/MassEntityManagerStorage.cpp
--- a/Engine/Plugins/Runtime/MassEntity/Source/MassEntity/Private/MassEntityManagerStorage.cpp
+++ b/Engine/Plugins/Runtime/MassEntity/Source/MassEntity/Private/MassEntityManagerStorage.cpp
@@ -7,6 +7,26 @@
 
 namespace UE::Mass
 {
+	//-----------------------------------------------------------------------------
+	// IEntityStorageInterface
+	//-----------------------------------------------------------------------------
+	int32 IEntityStorageInterface::Acquire(const int32 Count, TArray<FMassEntityHandle>& OutEntityHandles)
+	{
+		if (Count)
+		{
+			const int32 StartingIndex = OutEntityHandles.Num();
+			OutEntityHandles.AddZeroed(Count);
+			const int32 NumberAdded = Acquire(MakeArrayView(&OutEntityHandles[StartingIndex], Count));
+			if (UNLIKELY(NumberAdded < Count))
+			{
+				// need to remove the redundantly reserved entries
+				OutEntityHandles.RemoveAt(StartingIndex + NumberAdded, Count - NumberAdded, EAllowShrinking::No);
+			}
+			return NumberAdded;
+		}
+		return 0;
+	}
+
 	//-----------------------------------------------------------------------------
 	// FSingleThreadedEntityStorage
 	//-----------------------------------------------------------------------------
@@ -87,17 +107,16 @@ namespace UE::Mass
 		return Handle;
 	}
 
-	int32 FSingleThreadedEntityStorage::Acquire(const int32 Count, TArray<FMassEntityHandle>& OutEntityHandles)
+	int32 FSingleThreadedEntityStorage::Acquire(TArrayView<FMassEntityHandle> OutEntityHandles)
 	{
-		check(Count >= 0);
+		const int32 NumToAdd = OutEntityHandles.Num();
 
 		const int32 SerialNumber = SerialNumberGenerator.fetch_add(1);
 
 		int32 NumAdded = 0;
+		int32 CurrentEntityHandleIndex = 0;
 
-		OutEntityHandles.Reserve(OutEntityHandles.Num() + Count);
-
-		const int32 NumAvailableFromFreeList = FMath::Min(Count, EntityFreeIndexList.Num());
+		const int32 NumAvailableFromFreeList = FMath::Min(NumToAdd, EntityFreeIndexList.Num());
 		if (NumAvailableFromFreeList > 0)
 		{
 			const int32 FirstIndexToUse = EntityFreeIndexList.Num() - NumAvailableFromFreeList;
@@ -105,21 +124,21 @@ namespace UE::Mass
 			{
 				const int32 EntityIndex = EntityFreeIndexList[Index];
 				Entities[EntityIndex].SerialNumber = SerialNumber;
-				OutEntityHandles.Add({ EntityIndex, SerialNumber });
+				OutEntityHandles[CurrentEntityHandleIndex++] = { EntityIndex, SerialNumber };
 			}
 			EntityFreeIndexList.RemoveAt(FirstIndexToUse, NumAvailableFromFreeList, EAllowShrinking::No);
 			NumAdded = NumAvailableFromFreeList;
 		}
 
-		if (NumAdded < Count)
+		if (NumAdded < NumToAdd)
 		{
-			const int32 RemainingCount = Count - NumAdded;
+			const int32 RemainingCount = NumToAdd - NumAdded;
 			const int32 StartingIndex = Entities.Num();
 			Entities.Add(RemainingCount);
 			for (int32 EntityIndex = StartingIndex; EntityIndex < Entities.Num(); ++EntityIndex)
 			{
 				Entities[EntityIndex].SerialNumber = SerialNumber;
-				OutEntityHandles.Add({ EntityIndex, SerialNumber });
+				OutEntityHandles[CurrentEntityHandleIndex++] = { EntityIndex, SerialNumber };
 			}
 			NumAdded += RemainingCount;
 		}
@@ -390,14 +409,13 @@ namespace UE::Mass
 		return Handle;
 	}
 
-	int32 FConcurrentEntityStorage::Acquire(const int32 Count, TArray<FMassEntityHandle>& OutEntityHandles)
+	int32 FConcurrentEntityStorage::Acquire(TArrayView<FMassEntityHandle> OutEntityHandles)
 	{
-		check(Count >= 0);
+		const int32 NumberToAdd = OutEntityHandles.Num();
 
 		int32 CountAdded = 0;
-		int32 CountLeft = Count;
-
-		OutEntityHandles.Reserve(OutEntityHandles.Num() + Count);
+		int32 CountLeft = NumberToAdd;
+		int32 CurrentEntityHandleIndex = 0;
 
 		while (CountLeft > 0)
 		{
@@ -425,7 +443,7 @@ namespace UE::Mass
 				EntityData.bIsAllocated = 1;
 				const int32 SerialNumber = EntityData.GetSerialNumber();
 
-				OutEntityHandles.Add({ EntityIndex, SerialNumber });
+				OutEntityHandles[CurrentEntityHandleIndex++] = { EntityIndex, SerialNumber };
 			}
 			
 			CountAdded += CountToProcess;
diff --git a/Engine/Plugins/Runtime/MassEntity/Source/MassEntity/Public/MassEntityManager.h b/Engine/Plugins/Runtime/MassEntity/Source/MassEntity/Public/MassEntityManager.h
--- a/Engine/Plugins/Runtime/MassEntity/Source/MassEntity/Public/MassEntityManager.h
+++ b/Engine/Plugins/Runtime/MassEntity/Source/MassEntity/Public/MassEntityManager.h
@@ -357,6 +357,14 @@ public:
 	 * @return a view into InOutEntities containing only the freshly reserved entities
 	 */
 	TConstArrayView<FMassEntityHandle> BatchReserveEntities(const int32 Count, TArray<FMassEntityHandle>& InOutEntities);
+	
+	/**
+	 * Reserves number of entities corresponding to number of entries in the provided array view InOutEntities.
+	 * As a result InOutEntities gets filled with handles of reserved entities
+	 * @return the number of entities reserved
+	 */
+	int32 BatchReserveEntities(TArrayView<FMassEntityHandle> InOutEntities);
+
 	TSharedRef<FEntityCreationContext> BatchBuildEntities(const FMassArchetypeEntityCollectionWithPayload& EncodedEntitiesWithPayload, const FMassFragmentBitSet& FragmentsAffected
 		, const FMassArchetypeSharedFragmentValues& SharedFragmentValues = {}, const FMassArchetypeCreationParams& CreationParams = FMassArchetypeCreationParams());
 	TSharedRef<FEntityCreationContext> BatchBuildEntities(const FMassArchetypeEntityCollectionWithPayload& EncodedEntitiesWithPayload, FMassArchetypeCompositionDescriptor&& Composition
diff --git a/Engine/Plugins/Runtime/MassEntity/Source/MassEntity/Public/MassEntityManagerStorage.h b/Engine/Plugins/Runtime/MassEntity/Source/MassEntity/Public/MassEntityManagerStorage.h
--- a/Engine/Plugins/Runtime/MassEntity/Source/MassEntity/Public/MassEntityManagerStorage.h
+++ b/Engine/Plugins/Runtime/MassEntity/Source/MassEntity/Public/MassEntityManagerStorage.h
@@ -78,8 +78,9 @@ namespace UE::Mass
 
 		virtual FMassEntityHandle AcquireOne() = 0;
 		// @return number of entities actually added
-		virtual int32 Acquire(const int32 Count, TArray<FMassEntityHandle>& OutEntityHandles) = 0;
-		
+		int32 Acquire(const int32 Count, TArray<FMassEntityHandle>& OutEntityHandles);
+		virtual int32 Acquire(TArrayView<FMassEntityHandle> OutEntityHandles) = 0;
+
 		virtual int32 Release(TConstArrayView<FMassEntityHandle> Handles) = 0;
 		virtual int32 ReleaseOne(FMassEntityHandle Handles) = 0;
 		
@@ -120,7 +121,8 @@ namespace UE::Mass
 		virtual SIZE_T GetAllocatedSize() const override;
 		virtual bool IsValid(int32 Index) const override;
 		virtual FMassEntityHandle AcquireOne() override;
-		virtual int32 Acquire(const int32 Count, TArray<FMassEntityHandle>& OutEntityHandles) override;
+		using IEntityStorageInterface::Acquire;
+		virtual int32 Acquire(TArrayView<FMassEntityHandle> OutEntityHandles) override;
 		virtual int32 Release(TConstArrayView<FMassEntityHandle> Handles) override;
 		virtual int32 ReleaseOne(FMassEntityHandle Handle) override;
 		virtual int32 ForceRelease(TConstArrayView<FMassEntityHandle> Handles) override;
@@ -173,7 +175,8 @@ namespace UE::Mass
 		virtual SIZE_T GetAllocatedSize() const override;
 		virtual bool IsValid(int32 Index) const override;
 		virtual FMassEntityHandle AcquireOne() override;
-		virtual int32 Acquire(const int32 Count, TArray<FMassEntityHandle>& OutEntityHandles) override;
+		using IEntityStorageInterface::Acquire;
+		virtual int32 Acquire(TArrayView<FMassEntityHandle> OutEntityHandles) override;
 		virtual int32 Release(TConstArrayView<FMassEntityHandle> Handles) override;
 		virtual int32 ReleaseOne(FMassEntityHandle Handle) override;
 		virtual int32 ForceRelease(TConstArrayView<FMassEntityHandle> Handles) override;
