Added an array-view flavor of BatchDestroyEntityChunks and applied it to places where we called the existing version of BatchDestroyEntityChunks in a tight loop.
Added means of ensuring Mass creation context is being used in the same thread it has been created in.

diff --git a/Engine/Plugins/Runtime/InstancedActors/Source/InstancedActors/Private/InstancedActorsData.cpp b/Engine/Plugins/Runtime/InstancedActors/Source/InstancedActors/Private/InstancedActorsData.cpp
--- a/Engine/Plugins/Runtime/InstancedActors/Source/InstancedActors/Private/InstancedActorsData.cpp
+++ b/Engine/Plugins/Runtime/InstancedActors/Source/InstancedActors/Private/InstancedActorsData.cpp
@@ -1091,10 +1091,7 @@ void UInstancedActorsData::RuntimeRemoveInstances(TConstArrayView<FInstancedActo
 
 		// Destroy entities. This will also trigger actor destruction for any spawned actors
 		// for these entities.
-		for (FMassArchetypeEntityCollection& EntityCollectionToDestroy : EntityCollectionsToDestroy)
-		{
-			MassEntityManager.BatchDestroyEntityChunks(EntityCollectionToDestroy);
-		}
+		MassEntityManager.BatchDestroyEntityChunks(EntityCollectionsToDestroy);
 	}
 	// Pre-empt entity spawning and simply invalidate InstanceTransform entries, preventing them from spawning later
 	else
@@ -1132,10 +1129,7 @@ void UInstancedActorsData::RuntimeRemoveAllInstances()
 
 		// Destroy entities. This will also trigger actor destruction for any spawned actors
 		// for these entities.
-		for (FMassArchetypeEntityCollection& EntityCollectionToDestroy : EntityCollectionsToDestroy)
-		{
-			MassEntityManager.BatchDestroyEntityChunks(EntityCollectionToDestroy);
-		}
+		MassEntityManager.BatchDestroyEntityChunks(EntityCollectionsToDestroy);
 
 		// Zero out all entity handles to 'reset' them
 		FMemory::Memzero(Entities.GetData(), Entities.GetTypeSize() * Entities.Num());
diff --git a/Engine/Plugins/Runtime/MassEntity/Source/MassEntity/Private/MassEntityManager.cpp b/Engine/Plugins/Runtime/MassEntity/Source/MassEntity/Private/MassEntityManager.cpp
--- a/Engine/Plugins/Runtime/MassEntity/Source/MassEntity/Private/MassEntityManager.cpp
+++ b/Engine/Plugins/Runtime/MassEntity/Source/MassEntity/Private/MassEntityManager.cpp
@@ -47,9 +47,15 @@ namespace UE::Mass::Private
 //-----------------------------------------------------------------------------
 // FMassEntityManager::FEntityCreationContext
 //-----------------------------------------------------------------------------
+FMassEntityManager::FEntityCreationContext::FEntityCreationContext()
+	: OwnerThreadId(FPlatformTLS::GetCurrentThreadId())
+{	
+}
+
 FMassEntityManager::FEntityCreationContext::FEntityCreationContext(FMassEntityManager& InManager, const TConstArrayView<FMassEntityHandle> InCreatedEntities)
-	: CreatedEntities(InCreatedEntities)
+	: FEntityCreationContext()
 {
+	CreatedEntities = InCreatedEntities;
 	Manager = InManager.AsShared();
 }
 
@@ -85,11 +91,15 @@ TConstArrayView<FMassArchetypeEntityCollection> FMassEntityManager::FEntityCreat
 
 void FMassEntityManager::FEntityCreationContext::MarkDirty()
 {
+	checkf(OwnerThreadId == FPlatformTLS::GetCurrentThreadId(), TEXT("%hs: all FEntityCreationContext operations ere expected to be run in a single thread"), __FUNCTION__);
+
 	EntityCollections.Reset();
 }
 
 void FMassEntityManager::FEntityCreationContext::AppendEntities(const TConstArrayView<FMassEntityHandle> EntitiesToAppend)
 {
+	checkf(OwnerThreadId == FPlatformTLS::GetCurrentThreadId(), TEXT("%hs: all FEntityCreationContext operations ere expected to be run in a single thread"), __FUNCTION__);
+
 	if (EntitiesToAppend.Num())
 	{
 		if (CreatedEntities.Num())
@@ -109,6 +119,8 @@ void FMassEntityManager::FEntityCreationContext::AppendEntities(const TConstArra
 
 void FMassEntityManager::FEntityCreationContext::AppendEntities(const TConstArrayView<FMassEntityHandle> EntitiesToAppend, FMassArchetypeEntityCollection&& InEntityCollection)
 {
+	checkf(OwnerThreadId == FPlatformTLS::GetCurrentThreadId(), TEXT("%hs: all FEntityCreationContext operations ere expected to be run in a single thread"), __FUNCTION__);
+
 	if (EntitiesToAppend.Num() == 0)
 	{
 		return;
@@ -124,6 +136,11 @@ void FMassEntityManager::FEntityCreationContext::AppendEntities(const TConstArra
 	}
 }
 
+void FMassEntityManager::FEntityCreationContext::ForceUpdateCurrentThreadID()
+{
+	OwnerThreadId = FPlatformTLS::GetCurrentThreadId();
+}
+
 //-----------------------------------------------------------------------------
 // FMassEntityManager
 //-----------------------------------------------------------------------------
@@ -416,6 +433,11 @@ void FMassEntityManager::OnPostFork(EForkProcessRole Role)
 				CommandBuffer = MakeShareable(new FMassCommandBuffer());
 			}
 		}
+
+		if (TSharedPtr<FEntityCreationContext> ActiveContext = ActiveCreationContext.Pin())
+		{
+			ActiveContext->ForceUpdateCurrentThreadID();
+		}
 	}
 }
 
@@ -914,7 +936,6 @@ void FMassEntityManager::BatchDestroyEntities(TConstArrayView<FMassEntityHandle>
 	checkf(IsProcessing() == false, TEXT("Synchronous API function %hs called during mass processing. Use asynchronous API instead."), __FUNCTION__);
 	checkf(IsDuringEntityCreation() == false, TEXT("%hs: Trying to destroy entities while entity creation is under way. This operation is not supported."), __FUNCTION__);
 
-	// @todo optimize, we can make savings by implementing Archetype->RemoveEntities()
 	for (const FMassEntityHandle Entity : InEntities)
 	{
 		if (GetEntityStorageInterface().IsValidIndex(Entity.Index) == false)
@@ -940,6 +961,11 @@ void FMassEntityManager::BatchDestroyEntities(TConstArrayView<FMassEntityHandle>
 }
 
 void FMassEntityManager::BatchDestroyEntityChunks(const FMassArchetypeEntityCollection& EntityCollection)
+{
+	BatchDestroyEntityChunks(MakeArrayView(&EntityCollection, 1));
+}
+
+void FMassEntityManager::BatchDestroyEntityChunks(TConstArrayView<FMassArchetypeEntityCollection> Collections)
 {
 	TRACE_CPUPROFILER_EVENT_SCOPE(Mass_BatchDestroyEntityChunks);
 
@@ -952,22 +978,25 @@ void FMassEntityManager::BatchDestroyEntityChunks(const FMassArchetypeEntityColl
 	// destruction the commands will work on outdated information (which might result in crashes).
 	FMassProcessingContext ProcessingContext(*this, /*TimeDelta=*/0.0f);
 
-	bool bValidArchetype = EntityCollection.GetArchetype().IsValid();
-	if (bValidArchetype)
+	for (const FMassArchetypeEntityCollection& EntityCollection : Collections)
 	{
-		ProcessingContext.bFlushCommandBuffer = false;
-		ProcessingContext.CommandBuffer = MakeShareable(new FMassCommandBuffer());
-		ObserverManager.OnPreEntitiesDestroyed(ProcessingContext, EntityCollection);
+		EntitiesRemoved.Reset();
+		if (EntityCollection.GetArchetype().IsValid())
+		{
+			ProcessingContext.bFlushCommandBuffer = false;
+			ProcessingContext.CommandBuffer = MakeShareable(new FMassCommandBuffer());
+			ObserverManager.OnPreEntitiesDestroyed(ProcessingContext, EntityCollection);
 
-		FMassArchetypeData& ArchetypeData = FMassArchetypeHelper::ArchetypeDataFromHandleChecked(EntityCollection.GetArchetype());
-		ArchetypeData.BatchDestroyEntityChunks(EntityCollection.GetRanges(), EntitiesRemoved);
+			FMassArchetypeData& ArchetypeData = FMassArchetypeHelper::ArchetypeDataFromHandleChecked(EntityCollection.GetArchetype());
+			ArchetypeData.BatchDestroyEntityChunks(EntityCollection.GetRanges(), EntitiesRemoved);
 		
-		GetEntityStorageInterface().Release(EntitiesRemoved);
-	}
-	else
-	{
-		UE::Mass::Private::ConvertArchetypelessSubchunksIntoEntityHandles(EntityCollection.GetRanges(), EntitiesRemoved);
-		GetEntityStorageInterface().ForceRelease(EntitiesRemoved);
+			GetEntityStorageInterface().Release(EntitiesRemoved);
+		}
+		else
+		{
+			UE::Mass::Private::ConvertArchetypelessSubchunksIntoEntityHandles(EntityCollection.GetRanges(), EntitiesRemoved);
+			GetEntityStorageInterface().ForceRelease(EntitiesRemoved);
+		}
 	}
 }
 
diff --git a/Engine/Plugins/Runtime/MassEntity/Source/MassEntity/Public/MassCommands.h b/Engine/Plugins/Runtime/MassEntity/Source/MassEntity/Public/MassCommands.h
--- a/Engine/Plugins/Runtime/MassEntity/Source/MassEntity/Public/MassCommands.h
+++ b/Engine/Plugins/Runtime/MassEntity/Source/MassEntity/Public/MassCommands.h
@@ -188,10 +188,7 @@ protected:
 
 		TArray<FMassArchetypeEntityCollection> EntityCollectionsToDestroy;
 		UE::Mass::Utils::CreateEntityCollections(System, TargetEntities, FMassArchetypeEntityCollection::FoldDuplicates, EntityCollectionsToDestroy);
-		for (FMassArchetypeEntityCollection& Collection : EntityCollectionsToDestroy)
-		{
-			System.BatchDestroyEntityChunks(Collection);
-		}
+		System.BatchDestroyEntityChunks(EntityCollectionsToDestroy);
 	}
 };
 
diff --git a/Engine/Plugins/Runtime/MassEntity/Source/MassEntity/Public/MassEntityManager.h b/Engine/Plugins/Runtime/MassEntity/Source/MassEntity/Public/MassEntityManager.h
--- a/Engine/Plugins/Runtime/MassEntity/Source/MassEntity/Public/MassEntityManager.h
+++ b/Engine/Plugins/Runtime/MassEntity/Source/MassEntity/Public/MassEntityManager.h
@@ -209,13 +209,11 @@ public:
 	/**
 	 * A dedicated structure for ensuring the "on entities creation" observers get notified only once all other 
 	 * initialization operations are done and this creation context instance gets released. 
-	 * 
-	 * @NOTE the whole concept relies on the assumption we only ever create or "build" entities in a single thread (usually the GameThread).
 	 */
 	struct MASSENTITY_API FEntityCreationContext
 	{
 	private:
-		FEntityCreationContext() = default;
+		FEntityCreationContext();
 		explicit FEntityCreationContext(FMassEntityManager& InManager, const TConstArrayView<FMassEntityHandle> InCreatedEntities = {});
 		FEntityCreationContext(FMassEntityManager& InManager, const TConstArrayView<FMassEntityHandle> InCreatedEntities, FMassArchetypeEntityCollection&& EntityCollection);
 		
@@ -240,6 +238,14 @@ public:
 
 	private:
 		friend FMassEntityManager;
+		/** To be called in case of processor forking. */
+		void ForceUpdateCurrentThreadID();
+
+		/**
+		 * Identifies the thread where given FEntityCreationContext instance was created. All subsequent operations are 
+		 * expected to be run in the same thread.
+		 */
+		uint32 OwnerThreadId;
 		mutable TArray<FMassArchetypeEntityCollection> EntityCollections;
 		TArray<FMassEntityHandle> CreatedEntities;
 		FMassArchetypeEntityCollection::EDuplicatesHandling CollectionCreationDuplicatesHandling = FMassArchetypeEntityCollection::EDuplicatesHandling::NoDuplicates;
@@ -314,6 +320,7 @@ public:
 	/**
 	 * Destroys all the entities in the provided array of entities. The function will also gracefully handle entities
 	 * that have been reserved but not created yet.
+	 * @note the function doesn't handle duplicates in InEntities.
 	 * @param InEntities to destroy
 	 */
 	void BatchDestroyEntities(TConstArrayView<FMassEntityHandle> InEntities);
@@ -324,6 +331,7 @@ public:
 	 * @param Collection to destroy
 	 */
 	void BatchDestroyEntityChunks(const FMassArchetypeEntityCollection& Collection);
+	void BatchDestroyEntityChunks(TConstArrayView<FMassArchetypeEntityCollection> Collections);
 
 	void AddFragmentToEntity(FMassEntityHandle Entity, const UScriptStruct* FragmentType);
 	void AddFragmentToEntity(FMassEntityHandle Entity, const UScriptStruct* FragmentType, const FStructInitializationCallback& Initializer);
diff --git a/Engine/Plugins/Runtime/MassGameplay/Source/MassSpawner/Private/MassSpawnerSubsystem.cpp b/Engine/Plugins/Runtime/MassGameplay/Source/MassSpawner/Private/MassSpawnerSubsystem.cpp
--- a/Engine/Plugins/Runtime/MassGameplay/Source/MassSpawner/Private/MassSpawnerSubsystem.cpp
+++ b/Engine/Plugins/Runtime/MassGameplay/Source/MassSpawner/Private/MassSpawnerSubsystem.cpp
@@ -85,11 +85,7 @@ void UMassSpawnerSubsystem::DestroyEntities(TConstArrayView<FMassEntityHandle> E
 
 	TArray<FMassArchetypeEntityCollection> EntityCollections;
 	UE::Mass::Utils::CreateEntityCollections(*EntityManager.Get(), Entities, FMassArchetypeEntityCollection::NoDuplicates, EntityCollections);
-
-	for (const FMassArchetypeEntityCollection& Collection : EntityCollections)
-	{
-		EntityManager->BatchDestroyEntityChunks(Collection);
-	}
+	EntityManager->BatchDestroyEntityChunks(EntityCollections);
 }
 
 UMassProcessor* UMassSpawnerSubsystem::GetSpawnDataInitializer(TSubclassOf<UMassProcessor> InitializerClass)
