Added a unit test verifying the entity query behavior around fragments requested being there but not requested to be bound in any mode (via EMassFragmentAccess::None).

diff --git a/Engine/Plugins/Runtime/MassEntity/Source/MassEntity/Public/MassExecutionContext.h b/Engine/Plugins/Runtime/MassEntity/Source/MassEntity/Public/MassExecutionContext.h
--- a/Engine/Plugins/Runtime/MassEntity/Source/MassEntity/Public/MassExecutionContext.h
+++ b/Engine/Plugins/Runtime/MassEntity/Source/MassEntity/Public/MassExecutionContext.h
@@ -7,13 +7,14 @@
 #include "MassEntityQuery.h"
 #include "MassSubsystemAccess.h"
 
-
 #define CHECK_IF_VALID(View, Type) \
-	checkf(View \
+	testableCheckfReturn(View \
+		, return {} \
 		, TEXT("Requested fragment type not bound, type %s. Make sure it has been listed as required."), *GetNameSafe(Type))
 
 #define CHECK_IF_READWRITE(View) \
-	checkf(View == nullptr || View->Requirement.AccessMode == EMassFragmentAccess::ReadWrite \
+	testableCheckfReturn(View == nullptr || View->Requirement.AccessMode == EMassFragmentAccess::ReadWrite \
+		, return {} \
 		, TEXT("Requested fragment type not bound for writing, type %s. Make sure it has been listed as required in ReadWrite mode.") \
 		, View ? *GetNameSafe(View->Requirement.StructType) : TEXT("[Not found]"))
 
diff --git a/Engine/Plugins/Runtime/MassEntity/Source/MassEntityTestSuite/Private/MassQueryTest.cpp b/Engine/Plugins/Runtime/MassEntity/Source/MassEntityTestSuite/Private/MassQueryTest.cpp
--- a/Engine/Plugins/Runtime/MassEntity/Source/MassEntityTestSuite/Private/MassQueryTest.cpp
+++ b/Engine/Plugins/Runtime/MassEntity/Source/MassEntityTestSuite/Private/MassQueryTest.cpp
@@ -311,11 +311,22 @@ struct FQueryTest_FragmentPresent : FEntityTestBase
 		EntityManager->BatchCreateEntities(FloatsIntsArchetype, NumberOfEntitiesToAddB, MatchingEntities);
 		ensure(MatchingEntities.Num() == NumberOfEntitiesToAddA + NumberOfEntitiesToAddB);
 
+		bool bMutablesBound = true;
+		bool bConstsBound = true;
 		int TotalProcessed = 0;
 		FMassExecutionContext ExecContext(*EntityManager.Get());
-		Query.ForEachEntityChunk(*EntityManager, ExecContext, [&TotalProcessed](FMassExecutionContext& Context) {
-			TotalProcessed += Context.GetNumEntities();
-		});
+		{
+			AITEST_SCOPED_CHECK("Requested fragment type not bound", 4);
+			Query.ForEachEntityChunk(*EntityManager, ExecContext, [&TotalProcessed, &bMutablesBound, &bConstsBound](FMassExecutionContext& Context) {
+				TotalProcessed += Context.GetNumEntities();
+				TArrayView<FTestFragment_Int> MutableView = Context.GetMutableFragmentView<FTestFragment_Int>();
+				bMutablesBound = bMutablesBound && !MutableView.IsEmpty();
+				TConstArrayView<FTestFragment_Int> ConstView = Context.GetFragmentView<FTestFragment_Int>();
+				bConstsBound = bConstsBound && !ConstView.IsEmpty();
+			});
+		}
+		AITEST_FALSE("The fragment with None access is not being bound in RW mode", bMutablesBound);
+		AITEST_FALSE("The fragment with None access is not being bound in RO mode", bConstsBound);
 		AITEST_EQUAL("We expect the number of entities processed to match number added to matching archetypes", MatchingEntities.Num(), TotalProcessed);
 
 		return true;
