Extended Mass observers unit testing with verifying that the commands issued by observers are being flushed

diff --git a/Engine/Plugins/Runtime/MassEntity/Source/MassEntityTestSuite/Private/MassObserverTest.cpp b/Engine/Plugins/Runtime/MassEntity/Source/MassEntityTestSuite/Private/MassObserverTest.cpp
--- a/Engine/Plugins/Runtime/MassEntity/Source/MassEntityTestSuite/Private/MassObserverTest.cpp
+++ b/Engine/Plugins/Runtime/MassEntity/Source/MassEntityTestSuite/Private/MassObserverTest.cpp
@@ -32,6 +32,8 @@ struct FTagBaseOperation : FEntityTestBase
 	TArray<FMassEntityHandle> EntitiesInt;
 	TArray<FMassEntityHandle> EntitiesIntsFloat;
 	TArray<FMassEntityHandle> ExpectedEntities;
+	bool bCommandsFlushed = false;
+
 	// @return signifies if the test can continue
 	virtual bool PerformOperation() { return false; }
 
@@ -42,9 +44,14 @@ struct FTagBaseOperation : FEntityTestBase
 			ObserverProcessor = NewObject<UMassTestProcessorBase>();
 			ObserverProcessor->EntityQuery.AddRequirement<FTestFragment_Int>(EMassFragmentAccess::ReadOnly);
 			ObserverProcessor->EntityQuery.AddTagRequirement<FTagStruct>(EMassFragmentPresence::All);
-			ObserverProcessor->ForEachEntityChunkExecutionFunction = [this](FMassExecutionContext& Context)
+			ObserverProcessor->ForEachEntityChunkExecutionFunction = [bCommandsFlushedPtr = &bCommandsFlushed, AffectedEntitiesPtr = &AffectedEntities](FMassExecutionContext& Context)
 			{
-				AffectedEntities.Append(Context.GetEntities().GetData(), Context.GetEntities().Num());
+				AffectedEntitiesPtr->Append(Context.GetEntities().GetData(), Context.GetEntities().Num());
+				Context.Defer().PushCommand<FMassDeferredSetCommand>([&bCommandsFlushedPtr](FMassEntityManager&)
+					{
+						// dummy command, here just to catch if commands issue by observers got executed at all
+						*bCommandsFlushedPtr = true;
+					});
 			};
 
 			return true;
@@ -63,8 +70,9 @@ struct FTagBaseOperation : FEntityTestBase
 		if (PerformOperation())
 		{
 			EntityManager->FlushCommands();
-			AITEST_EQUAL(TEXT("The tag observer is expected to be run for predicted number of entities"), AffectedEntities.Num(), ExpectedEntities.Num());
-			
+			AITEST_EQUAL(TEXT("The observer is expected to be run for predicted number of entities"), AffectedEntities.Num(), ExpectedEntities.Num());
+			AITEST_TRUE(TEXT("The commands issued by the observer are flushed"), bCommandsFlushed);
+
 			ExpectedEntities.Sort(EntityIndexSorted);
 			AffectedEntities.Sort(EntityIndexSorted);
 
@@ -279,6 +287,8 @@ struct FFragmentTestBase : FEntityTestBase
 	TArray<FMassEntityHandle> EntitiesInt;
 	TArray<FMassEntityHandle> EntitiesIntsFloat;
 	TArray<FMassEntityHandle> ExpectedEntities;
+	bool bCommandsFlushed = false;
+
 	// @return signifies if the test can continue
 	virtual bool PerformOperation() { return false; }
 
@@ -288,9 +298,14 @@ struct FFragmentTestBase : FEntityTestBase
 		{
 			ObserverProcessor = NewObject<UMassTestProcessorBase>();
 			ObserverProcessor->EntityQuery.AddRequirement(FFragmentStruct::StaticStruct(), EMassFragmentAccess::ReadOnly);
-			ObserverProcessor->ForEachEntityChunkExecutionFunction = [this](FMassExecutionContext& Context)
+			ObserverProcessor->ForEachEntityChunkExecutionFunction = [bCommandsFlushedPtr = &bCommandsFlushed, AffectedEntitiesPtr = &AffectedEntities](FMassExecutionContext& Context)
 				{
-					AffectedEntities.Append(Context.GetEntities().GetData(), Context.GetEntities().Num());
+					AffectedEntitiesPtr->Append(Context.GetEntities().GetData(), Context.GetEntities().Num());
+					Context.Defer().PushCommand<FMassDeferredSetCommand>([&bCommandsFlushedPtr](FMassEntityManager&)
+						{
+							// dummy command, here just to catch if commands issue by observers got executed at all
+							*bCommandsFlushedPtr = true;
+						});
 				};
 
 			return true;
@@ -310,6 +325,7 @@ struct FFragmentTestBase : FEntityTestBase
 		{
 			EntityManager->FlushCommands();
 			AITEST_EQUAL(TEXT("The fragment observer is expected to be run for predicted number of entities"), AffectedEntities.Num(), ExpectedEntities.Num());
+			AITEST_TRUE(TEXT("The commands issued by the observer are flushed"), bCommandsFlushed);
 
 			ExpectedEntities.Sort(EntityIndexSorted);
 			AffectedEntities.Sort(EntityIndexSorted);
