[Mass] Unified and simplified determining whether a given archetype matches giver query requirements.

diff --git a/Engine/Plugins/Runtime/MassEntity/Source/MassEntity/Internal/MassArchetypeData.h b/Engine/Plugins/Runtime/MassEntity/Source/MassEntity/Internal/MassArchetypeData.h
--- a/Engine/Plugins/Runtime/MassEntity/Source/MassEntity/Internal/MassArchetypeData.h
+++ b/Engine/Plugins/Runtime/MassEntity/Source/MassEntity/Internal/MassArchetypeData.h
@@ -425,9 +425,11 @@ struct FMassArchetypeHelper
 	 * @param bBailOutOnFirstFail if true will skip the remaining tests as soon as a single mismatch is detected. This option
 	 *	is used when looking for matching archetypes. For debugging purposes use `false` to list all the mismatching elements.
 	 */
-	MASSENTITY_API static bool DoesArchetypeMatchRequirements(const FMassArchetypeData& Archetype, const FMassFragmentRequirements& Requirements
 #if WITH_MASSENTITY_DEBUG
-		, const bool bBailOutOnFirstFail = true, FOutputDevice* OutputDevice = nullptr
+	MASSENTITY_API static bool DoesArchetypeMatchRequirements(const FMassArchetypeData& Archetype, const FMassFragmentRequirements& Requirements
+		, const bool bBailOutOnFirstFail = true, FOutputDevice* OutputDevice = nullptr);
 #endif // WITH_MASSENTITY_DEBUG
-		);
+
+	MASSENTITY_API static bool DoesArchetypeMatchRequirements(const FMassArchetypeData& Archetype, const FMassFragmentRequirements& Requirements);
+	MASSENTITY_API static bool DoesArchetypeMatchRequirements(const FMassArchetypeCompositionDescriptor& ArchetypeComposition, const FMassFragmentRequirements& Requirements);
 };
diff --git a/Engine/Plugins/Runtime/MassEntity/Source/MassEntity/Private/MassArchetypeData.cpp b/Engine/Plugins/Runtime/MassEntity/Source/MassEntity/Private/MassArchetypeData.cpp
--- a/Engine/Plugins/Runtime/MassEntity/Source/MassEntity/Private/MassArchetypeData.cpp
+++ b/Engine/Plugins/Runtime/MassEntity/Source/MassEntity/Private/MassArchetypeData.cpp
@@ -1322,26 +1322,35 @@ void FMassArchetypeData::BatchSetFragmentValues(TConstArrayView<FMassArchetypeEn
 //-----------------------------------------------------------------------------
 // FMassArchetypeHelper
 //-----------------------------------------------------------------------------
-bool FMassArchetypeHelper::DoesArchetypeMatchRequirements(const FMassArchetypeData& Archetype, const FMassFragmentRequirements& Requirements
+bool FMassArchetypeHelper::DoesArchetypeMatchRequirements(const FMassArchetypeData& Archetype, const FMassFragmentRequirements& Requirements)
+{
+	return DoesArchetypeMatchRequirements(Archetype.GetCompositionDescriptor(), Requirements);
+}
+bool FMassArchetypeHelper::DoesArchetypeMatchRequirements(const FMassArchetypeCompositionDescriptor& ArchetypeComposition, const FMassFragmentRequirements& Requirements)
+{
+	return Requirements.DoesArchetypeMatchRequirements(ArchetypeComposition);
+}
+
 #if WITH_MASSENTITY_DEBUG
-	, const bool bBailOutOnFirstFail, FOutputDevice* OutputDevice
-#endif // WITH_MASSENTITY_DEBUG
-	)
+bool FMassArchetypeHelper::DoesArchetypeMatchRequirements(const FMassArchetypeData& Archetype, const FMassFragmentRequirements& Requirements
+	, const bool bBailOutOnFirstFail, FOutputDevice* OutputDevice)
 {
+	if (OutputDevice == nullptr)
+	{
+		return DoesArchetypeMatchRequirements(Archetype.GetCompositionDescriptor(), Requirements);
+	}
+
 	bool bResult = true;
 	if (Archetype.GetTagBitSet().HasAll(Requirements.GetRequiredAllTags()) == false)
 	{
 		// missing some required tags, skip.
-#if WITH_MASSENTITY_DEBUG
-		if (OutputDevice)
-		{
-			const FMassTagBitSet UnsatisfiedTags = Requirements.GetRequiredAllTags() - Archetype.GetTagBitSet();
-			FStringOutputDevice Description;
-			UnsatisfiedTags.DebugGetStringDesc(Description);
-			OutputDevice->Logf(TEXT("Archetype did not match due to missing tags: %s"), *Description);
-		}
-		if (!bBailOutOnFirstFail) bResult = false; else
-#endif // WITH_MASSENTITY_DEBUG
+		const FMassTagBitSet UnsatisfiedTags = Requirements.GetRequiredAllTags() - Archetype.GetTagBitSet();
+		FStringOutputDevice Description;
+		UnsatisfiedTags.DebugGetStringDesc(Description);
+		OutputDevice->Logf(TEXT("Archetype did not match due to missing tags: %s"), *Description);
+		
+		bResult = false;
+		if (bBailOutOnFirstFail)
 		{
 			return false;
 		}
@@ -1350,16 +1359,13 @@ bool FMassArchetypeHelper::DoesArchetypeMatchRequirements(const FMassArchetypeDa
 	if (Archetype.GetTagBitSet().HasNone(Requirements.GetRequiredNoneTags()) == false)
 	{
 		// has some tags required to be absent
-#if WITH_MASSENTITY_DEBUG
-		if (OutputDevice)
-		{
-			const FMassTagBitSet UnwantedTags = Requirements.GetRequiredNoneTags().GetOverlap(Archetype.GetTagBitSet());
-			FStringOutputDevice Description;
-			UnwantedTags.DebugGetStringDesc(Description);
-			OutputDevice->Logf(TEXT("Archetype has tags required absent: %s"), *Description);
-		}
-		if (!bBailOutOnFirstFail) bResult = false; else
-#endif // WITH_MASSENTITY_DEBUG
+		const FMassTagBitSet UnwantedTags = Requirements.GetRequiredNoneTags().GetOverlap(Archetype.GetTagBitSet());
+		FStringOutputDevice Description;
+		UnwantedTags.DebugGetStringDesc(Description);
+		OutputDevice->Logf(TEXT("Archetype has tags required absent: %s"), *Description);
+
+		bResult = false;
+		if (bBailOutOnFirstFail)
 		{
 			return false;
 		}
@@ -1368,15 +1374,12 @@ bool FMassArchetypeHelper::DoesArchetypeMatchRequirements(const FMassArchetypeDa
 	if (Requirements.GetRequiredAnyTags().IsEmpty() == false
 		&& Archetype.GetTagBitSet().HasAny(Requirements.GetRequiredAnyTags()) == false)
 	{
-#if WITH_MASSENTITY_DEBUG
-		if (OutputDevice)
-		{
-			FStringOutputDevice Description;
-			Requirements.GetRequiredAnyTags().DebugGetStringDesc(Description);
-			OutputDevice->Logf(TEXT("Archetype did not match due to missing \'any\' tags: %s"), *Description);
-		}
-		if (!bBailOutOnFirstFail) bResult = false; else
-#endif // WITH_MASSENTITY_DEBUG
+		FStringOutputDevice Description;
+		Requirements.GetRequiredAnyTags().DebugGetStringDesc(Description);
+		OutputDevice->Logf(TEXT("Archetype did not match due to missing \'any\' tags: %s"), *Description);
+
+		bResult = false;
+		if (bBailOutOnFirstFail)
 		{
 			return false;
 		}
@@ -1385,16 +1388,13 @@ bool FMassArchetypeHelper::DoesArchetypeMatchRequirements(const FMassArchetypeDa
 	if (Archetype.GetFragmentBitSet().HasAll(Requirements.GetRequiredAllFragments()) == false)
 	{
 		// missing some required fragments, skip.
-#if WITH_MASSENTITY_DEBUG
-		if (OutputDevice)
-		{
-			const FMassFragmentBitSet UnsatisfiedFragments = Requirements.GetRequiredAllFragments() - Archetype.GetFragmentBitSet();
-			FStringOutputDevice Description;
-			UnsatisfiedFragments.DebugGetStringDesc(Description);
-			OutputDevice->Logf(TEXT("Archetype did not match due to missing Fragments: %s"), *Description);
-		}
-		if (!bBailOutOnFirstFail) bResult = false; else
-#endif // WITH_MASSENTITY_DEBUG
+		const FMassFragmentBitSet UnsatisfiedFragments = Requirements.GetRequiredAllFragments() - Archetype.GetFragmentBitSet();
+		FStringOutputDevice Description;
+		UnsatisfiedFragments.DebugGetStringDesc(Description);
+		OutputDevice->Logf(TEXT("Archetype did not match due to missing Fragments: %s"), *Description);
+
+		bResult = false;
+		if (bBailOutOnFirstFail)
 		{
 			return false;
 		}
@@ -1403,16 +1403,13 @@ bool FMassArchetypeHelper::DoesArchetypeMatchRequirements(const FMassArchetypeDa
 	if (Archetype.GetFragmentBitSet().HasNone(Requirements.GetRequiredNoneFragments()) == false)
 	{
 		// has some Fragments required to be absent
-#if WITH_MASSENTITY_DEBUG
-		if (OutputDevice)
-		{
-			const FMassFragmentBitSet UnwantedFragments = Requirements.GetRequiredNoneFragments().GetOverlap(Archetype.GetFragmentBitSet());
-			FStringOutputDevice Description;
-			UnwantedFragments.DebugGetStringDesc(Description);
-			OutputDevice->Logf(TEXT("Archetype has Fragments required absent: %s"), *Description);
-		}
-		if (!bBailOutOnFirstFail) bResult = false; else
-#endif // WITH_MASSENTITY_DEBUG
+		const FMassFragmentBitSet UnwantedFragments = Requirements.GetRequiredNoneFragments().GetOverlap(Archetype.GetFragmentBitSet());
+		FStringOutputDevice Description;
+		UnwantedFragments.DebugGetStringDesc(Description);
+		OutputDevice->Logf(TEXT("Archetype has Fragments required absent: %s"), *Description);
+
+		bResult = false;
+		if (bBailOutOnFirstFail)
 		{
 			return false;
 		}
@@ -1421,15 +1418,12 @@ bool FMassArchetypeHelper::DoesArchetypeMatchRequirements(const FMassArchetypeDa
 	if (Requirements.GetRequiredAnyFragments().IsEmpty() == false
 		&& Archetype.GetFragmentBitSet().HasAny(Requirements.GetRequiredAnyFragments()) == false)
 	{
-#if WITH_MASSENTITY_DEBUG
-		if (OutputDevice)
-		{
-			FStringOutputDevice Description;
-			Requirements.GetRequiredAnyFragments().DebugGetStringDesc(Description);
-			OutputDevice->Logf(TEXT("Archetype did not match due to missing \'any\' fragments: %s"), *Description);
-		}
-		if (!bBailOutOnFirstFail) bResult = false; else
-#endif // WITH_MASSENTITY_DEBUG
+		FStringOutputDevice Description;
+		Requirements.GetRequiredAnyFragments().DebugGetStringDesc(Description);
+		OutputDevice->Logf(TEXT("Archetype did not match due to missing \'any\' fragments: %s"), *Description);
+
+		bResult = false;
+		if (bBailOutOnFirstFail)
 		{
 			return false;
 		}
@@ -1438,16 +1432,13 @@ bool FMassArchetypeHelper::DoesArchetypeMatchRequirements(const FMassArchetypeDa
 	if (Archetype.GetChunkFragmentBitSet().HasAll(Requirements.GetRequiredAllChunkFragments()) == false)
 	{
 		// missing some required fragments, skip.
-#if WITH_MASSENTITY_DEBUG
-		if (OutputDevice)
-		{
-			const FMassChunkFragmentBitSet UnsatisfiedFragments = Requirements.GetRequiredAllChunkFragments() - Archetype.GetChunkFragmentBitSet();
-			FStringOutputDevice Description;
-			UnsatisfiedFragments.DebugGetStringDesc(Description);
-			OutputDevice->Logf(TEXT("Archetype did not match due to missing Chunk Fragments: %s"), *Description);
-		}
-		if (!bBailOutOnFirstFail) bResult = false; else
-#endif // WITH_MASSENTITY_DEBUG
+		const FMassChunkFragmentBitSet UnsatisfiedFragments = Requirements.GetRequiredAllChunkFragments() - Archetype.GetChunkFragmentBitSet();
+		FStringOutputDevice Description;
+		UnsatisfiedFragments.DebugGetStringDesc(Description);
+		OutputDevice->Logf(TEXT("Archetype did not match due to missing Chunk Fragments: %s"), *Description);
+
+		bResult = false;
+		if (bBailOutOnFirstFail)
 		{
 			return false;
 		}
@@ -1456,16 +1447,13 @@ bool FMassArchetypeHelper::DoesArchetypeMatchRequirements(const FMassArchetypeDa
 	if (Archetype.GetChunkFragmentBitSet().HasNone(Requirements.GetRequiredNoneChunkFragments()) == false)
 	{
 		// has some Fragments required to be absent
-#if WITH_MASSENTITY_DEBUG
-		if (OutputDevice)
-		{
-			const FMassChunkFragmentBitSet UnwantedFragments = Requirements.GetRequiredNoneChunkFragments().GetOverlap(Archetype.GetChunkFragmentBitSet());
-			FStringOutputDevice Description;
-			UnwantedFragments.DebugGetStringDesc(Description);
-			OutputDevice->Logf(TEXT("Archetype has Chunk Fragments required absent: %s"), *Description);
-		}
-		if (!bBailOutOnFirstFail) bResult = false; else
-#endif // WITH_MASSENTITY_DEBUG
+		const FMassChunkFragmentBitSet UnwantedFragments = Requirements.GetRequiredNoneChunkFragments().GetOverlap(Archetype.GetChunkFragmentBitSet());
+		FStringOutputDevice Description;
+		UnwantedFragments.DebugGetStringDesc(Description);
+		OutputDevice->Logf(TEXT("Archetype has Chunk Fragments required absent: %s"), *Description);
+
+		bResult = false;
+		if (bBailOutOnFirstFail)
 		{
 			return false;
 		}
@@ -1474,16 +1462,13 @@ bool FMassArchetypeHelper::DoesArchetypeMatchRequirements(const FMassArchetypeDa
 	if (Archetype.GetSharedFragmentBitSet().HasAll(Requirements.GetRequiredAllSharedFragments()) == false)
 	{
 		// missing some required fragments, skip.
-#if WITH_MASSENTITY_DEBUG
-		if (OutputDevice)
-		{
-			const FMassSharedFragmentBitSet UnsatisfiedFragments = Requirements.GetRequiredAllSharedFragments() - Archetype.GetSharedFragmentBitSet();
-			FStringOutputDevice Description;
-			UnsatisfiedFragments.DebugGetStringDesc(Description);
-			OutputDevice->Logf(TEXT("Archetype did not match due to missing Shared Fragments: %s"), *Description);
-		}
-		if (!bBailOutOnFirstFail) bResult = false; else
-#endif // WITH_MASSENTITY_DEBUG
+		const FMassSharedFragmentBitSet UnsatisfiedFragments = Requirements.GetRequiredAllSharedFragments() - Archetype.GetSharedFragmentBitSet();
+		FStringOutputDevice Description;
+		UnsatisfiedFragments.DebugGetStringDesc(Description);
+		OutputDevice->Logf(TEXT("Archetype did not match due to missing Shared Fragments: %s"), *Description);
+
+		bResult = false;
+		if (bBailOutOnFirstFail)
 		{
 			return false;
 		}
@@ -1492,16 +1477,13 @@ bool FMassArchetypeHelper::DoesArchetypeMatchRequirements(const FMassArchetypeDa
 	if (Archetype.GetSharedFragmentBitSet().HasNone(Requirements.GetRequiredNoneSharedFragments()) == false)
 	{
 		// has some Fragments required to be absent
-#if WITH_MASSENTITY_DEBUG
-		if (OutputDevice)
-		{
-			const FMassSharedFragmentBitSet UnwantedFragments = Requirements.GetRequiredNoneSharedFragments().GetOverlap(Archetype.GetSharedFragmentBitSet());
-			FStringOutputDevice Description;
-			UnwantedFragments.DebugGetStringDesc(Description);
-			OutputDevice->Logf(TEXT("Archetype has Shared Fragments required absent: %s"), *Description);
-		}
-		if (!bBailOutOnFirstFail) bResult = false; else
-#endif // WITH_MASSENTITY_DEBUG
+		const FMassSharedFragmentBitSet UnwantedFragments = Requirements.GetRequiredNoneSharedFragments().GetOverlap(Archetype.GetSharedFragmentBitSet());
+		FStringOutputDevice Description;
+		UnwantedFragments.DebugGetStringDesc(Description);
+		OutputDevice->Logf(TEXT("Archetype has Shared Fragments required absent: %s"), *Description);
+
+		bResult = false;
+		if (bBailOutOnFirstFail)
 		{
 			return false;
 		}
@@ -1510,16 +1492,13 @@ bool FMassArchetypeHelper::DoesArchetypeMatchRequirements(const FMassArchetypeDa
 	if (Archetype.GetConstSharedFragmentBitSet().HasAll(Requirements.GetRequiredAllConstSharedFragments()) == false)
 	{
 		// missing some required fragments, skip.
-#if WITH_MASSENTITY_DEBUG
-		if (OutputDevice)
-		{
-			const FMassConstSharedFragmentBitSet UnsatisfiedFragments = Requirements.GetRequiredAllConstSharedFragments() - Archetype.GetConstSharedFragmentBitSet();
-			FStringOutputDevice Description;
-			UnsatisfiedFragments.DebugGetStringDesc(Description);
-			OutputDevice->Logf(TEXT("Archetype did not match due to missing Const Shared Fragments: %s"), *Description);
-		}
-		if (!bBailOutOnFirstFail) bResult = false; else
-#endif // WITH_MASSENTITY_DEBUG
+		const FMassConstSharedFragmentBitSet UnsatisfiedFragments = Requirements.GetRequiredAllConstSharedFragments() - Archetype.GetConstSharedFragmentBitSet();
+		FStringOutputDevice Description;
+		UnsatisfiedFragments.DebugGetStringDesc(Description);
+		OutputDevice->Logf(TEXT("Archetype did not match due to missing Const Shared Fragments: %s"), *Description);
+
+		bResult = false;
+		if (bBailOutOnFirstFail)
 		{
 			return false;
 		}
@@ -1528,18 +1507,15 @@ bool FMassArchetypeHelper::DoesArchetypeMatchRequirements(const FMassArchetypeDa
 	if (Archetype.GetConstSharedFragmentBitSet().HasNone(Requirements.GetRequiredNoneConstSharedFragments()) == false)
 	{
 		// has some Fragments required to be absent
-#if WITH_MASSENTITY_DEBUG
-		if (OutputDevice)
-		{
-			const FMassConstSharedFragmentBitSet UnwantedFragments = Requirements.GetRequiredNoneConstSharedFragments().GetOverlap(Archetype.GetConstSharedFragmentBitSet());
-			FStringOutputDevice Description;
-			UnwantedFragments.DebugGetStringDesc(Description);
-			OutputDevice->Logf(TEXT("Archetype has Const Shared Fragments required absent: %s"), *Description);
-		}
+		const FMassConstSharedFragmentBitSet UnwantedFragments = Requirements.GetRequiredNoneConstSharedFragments().GetOverlap(Archetype.GetConstSharedFragmentBitSet());
+		FStringOutputDevice Description;
+		UnwantedFragments.DebugGetStringDesc(Description);
+		OutputDevice->Logf(TEXT("Archetype has Const Shared Fragments required absent: %s"), *Description);
+
 		// could skip the test as it's the final check we're performing, but leaving it in in case more cases are added
 		// later without checking the existing code
-		if (!bBailOutOnFirstFail) bResult = false; else
-#endif // WITH_MASSENTITY_DEBUG
+		bResult = false;
+		if (bBailOutOnFirstFail)
 		{
 			return false;
 		}
@@ -1547,3 +1523,4 @@ bool FMassArchetypeHelper::DoesArchetypeMatchRequirements(const FMassArchetypeDa
 
 	return bResult;
 }
+#endif // WITH_MASSENTITY_DEBUG
\ No newline at end of file
diff --git a/Engine/Plugins/Runtime/MassEntity/Source/MassEntity/Private/MassDebugger.cpp b/Engine/Plugins/Runtime/MassEntity/Source/MassEntity/Private/MassDebugger.cpp
--- a/Engine/Plugins/Runtime/MassEntity/Source/MassEntity/Private/MassDebugger.cpp
+++ b/Engine/Plugins/Runtime/MassEntity/Source/MassEntity/Private/MassDebugger.cpp
@@ -493,10 +493,13 @@ FString FMassDebugger::GetArchetypeRequirementCompatibilityDescription(const FMa
 	}
 
 	const FMassArchetypeData& ArchetypeData = FMassArchetypeHelper::ArchetypeDataFromHandleChecked(ArchetypeHandle);
+	return FMassDebugger::GetArchetypeRequirementCompatibilityDescription(Requirements, ArchetypeData.GetCompositionDescriptor());
+}
+	
+FString FMassDebugger::GetArchetypeRequirementCompatibilityDescription(const FMassFragmentRequirements& Requirements, const FMassArchetypeCompositionDescriptor& ArchetypeComposition)
+{
 	FStringOutputDevice OutDescription;
 
-	const FMassArchetypeCompositionDescriptor& ArchetypeComposition = ArchetypeData.GetCompositionDescriptor();
-
 	if (ArchetypeComposition.Fragments.HasAll(Requirements.RequiredAllFragments) == false)
 	{
 		// missing one of the strictly required fragments
diff --git a/Engine/Plugins/Runtime/MassEntity/Source/MassEntity/Private/MassEntityManager.cpp b/Engine/Plugins/Runtime/MassEntity/Source/MassEntity/Private/MassEntityManager.cpp
--- a/Engine/Plugins/Runtime/MassEntity/Source/MassEntity/Private/MassEntityManager.cpp
+++ b/Engine/Plugins/Runtime/MassEntity/Source/MassEntity/Private/MassEntityManager.cpp
@@ -1790,11 +1790,6 @@ void FMassEntityManager::GetMatchingArchetypes(const FMassFragmentRequirements&
 		}
 	}
 
-#if WITH_MASSENTITY_DEBUG
-	FStringOutputDevice LogOutput;
-	LogOutput.SetAutoEmitLineTerminator(true);
-#endif // WITH_MASSENTITY_DEBUG
-
 	// Then verify that they contain *all* required fragments
 	for (TSharedPtr<FMassArchetypeData>& ArchetypePtr : AnyArchetypes)
 	{
@@ -1806,19 +1801,15 @@ void FMassEntityManager::GetMatchingArchetypes(const FMassFragmentRequirements&
 			continue;
 		}
 
-		if (FMassArchetypeHelper::DoesArchetypeMatchRequirements(Archetype, Requirements
-#if WITH_MASSENTITY_DEBUG
-			, /*bBailOutOnFirstFail=*/true, & LogOutput
-#endif // WITH_MASSENTITY_DEBUG
-		))
+		if (Requirements.DoesArchetypeMatchRequirements(Archetype.GetCompositionDescriptor()))
 		{
 			OutValidArchetypes.Add(ArchetypePtr);
 		}
 #if WITH_MASSENTITY_DEBUG
 		else
 		{
-			UE_VLOG_UELOG(GetOwner(), LogMass, VeryVerbose, TEXT("%s"), *LogOutput);
-			LogOutput.Reset();
+			UE_VLOG_UELOG(GetOwner(), LogMass, VeryVerbose, TEXT("%s")
+				, *FMassDebugger::GetArchetypeRequirementCompatibilityDescription(Requirements, Archetype.GetCompositionDescriptor()));
 		}
 #endif // WITH_MASSENTITY_DEBUG
 	}
diff --git a/Engine/Plugins/Runtime/MassEntity/Source/MassEntity/Private/MassRequirements.cpp b/Engine/Plugins/Runtime/MassEntity/Source/MassEntity/Private/MassRequirements.cpp
--- a/Engine/Plugins/Runtime/MassEntity/Source/MassEntity/Private/MassRequirements.cpp
+++ b/Engine/Plugins/Runtime/MassEntity/Source/MassEntity/Private/MassRequirements.cpp
@@ -126,9 +126,12 @@ bool FMassFragmentRequirements::DoesArchetypeMatchRequirements(const FMassArchet
 	check(ArchetypeHandle.IsValid());
 	const FMassArchetypeData* Archetype = FMassArchetypeHelper::ArchetypeDataFromHandle(ArchetypeHandle);
 	CA_ASSUME(Archetype);
-	
-	const FMassArchetypeCompositionDescriptor& ArchetypeComposition = Archetype->GetCompositionDescriptor();
 
+	return DoesArchetypeMatchRequirements(Archetype->GetCompositionDescriptor());
+}
+	
+bool FMassFragmentRequirements::DoesArchetypeMatchRequirements(const FMassArchetypeCompositionDescriptor& ArchetypeComposition) const
+{
 	return ArchetypeComposition.Fragments.HasAll(RequiredAllFragments)
 		&& (RequiredAnyFragments.IsEmpty() || ArchetypeComposition.Fragments.HasAny(RequiredAnyFragments))
 		&& ArchetypeComposition.Fragments.HasNone(RequiredNoneFragments)
diff --git a/Engine/Plugins/Runtime/MassEntity/Source/MassEntity/Public/MassDebugger.h b/Engine/Plugins/Runtime/MassEntity/Source/MassEntity/Public/MassDebugger.h
--- a/Engine/Plugins/Runtime/MassEntity/Source/MassEntity/Public/MassDebugger.h
+++ b/Engine/Plugins/Runtime/MassEntity/Source/MassEntity/Public/MassDebugger.h
@@ -123,6 +123,7 @@ struct MASSENTITY_API FMassDebugger
 	static FString GetSingleRequirementDescription(const FMassFragmentRequirementDescription& Requirement);
 	static FString GetRequirementsDescription(const FMassFragmentRequirements& Requirements);
 	static FString GetArchetypeRequirementCompatibilityDescription(const FMassFragmentRequirements& Requirements, const FMassArchetypeHandle& ArchetypeHandle);
+	static FString GetArchetypeRequirementCompatibilityDescription(const FMassFragmentRequirements& Requirements, const FMassArchetypeCompositionDescriptor& ArchetypeComposition);
 
 	static void OutputArchetypeDescription(FOutputDevice& Ar, const FMassArchetypeHandle& Archetype);
 	static void OutputEntityDescription(FOutputDevice& Ar, const FMassEntityManager& EntityManager, const int32 EntityIndex, const TCHAR* InPrefix = TEXT(""));
@@ -161,6 +162,7 @@ struct MASSENTITY_API FMassDebugger
 	static FString GetSingleRequirementDescription(const FMassFragmentRequirementDescription& Requirement) { return TEXT("[no debug information]"); }
 	static FString GetRequirementsDescription(const FMassFragmentRequirements& Requirements) { return TEXT("[no debug information]"); }
 	static FString GetArchetypeRequirementCompatibilityDescription(const FMassFragmentRequirements& Requirements, const FMassArchetypeHandle& ArchetypeHandle) { return TEXT("[no debug information]"); }
+	static FString GetArchetypeRequirementCompatibilityDescription(const FMassFragmentRequirements& Requirements, const FMassArchetypeCompositionDescriptor& ArchetypeComposition) { return TEXT("[no debug information]"); }
 };
 
 #define MASS_IF_ENTITY_DEBUGGED(a, b) false
diff --git a/Engine/Plugins/Runtime/MassEntity/Source/MassEntity/Public/MassRequirements.h b/Engine/Plugins/Runtime/MassEntity/Source/MassEntity/Public/MassRequirements.h
--- a/Engine/Plugins/Runtime/MassEntity/Source/MassEntity/Public/MassRequirements.h
+++ b/Engine/Plugins/Runtime/MassEntity/Source/MassEntity/Public/MassRequirements.h
@@ -428,6 +428,7 @@ public:
 	bool IsEmpty() const;
 
 	bool DoesArchetypeMatchRequirements(const FMassArchetypeHandle& ArchetypeHandle) const;
+	bool DoesArchetypeMatchRequirements(const FMassArchetypeCompositionDescriptor& ArchetypeComposition) const;
 	bool DoesRequireGameThreadExecution() const { return bRequiresGameThreadExecution; }
 	void ExportRequirements(FMassExecutionRequirements& OutRequirements) const;
 
