Added a unit test verifying proper handling of fragments containing TSharedPtrs and alike

diff --git a/Engine/Plugins/Runtime/MassEntity/Source/MassEntityTestSuite/Private/MassEntityTest.cpp b/Engine/Plugins/Runtime/MassEntity/Source/MassEntityTestSuite/Private/MassEntityTest.cpp
--- a/Engine/Plugins/Runtime/MassEntity/Source/MassEntityTestSuite/Private/MassEntityTest.cpp
+++ b/Engine/Plugins/Runtime/MassEntity/Source/MassEntityTestSuite/Private/MassEntityTest.cpp
@@ -354,6 +354,59 @@ IMPLEMENT_AI_INSTANT_TEST(FEntityTest_ReserveAPreviouslyBuiltEntity, "System.Mas
 
 #endif // WITH_MASSENTITY_DEBUG
 
+// testing handling of fragments containing shared ptrs
+struct FEntityTest_SharedPtrFragment : FEntityTestBase
+{
+	virtual bool InstantTest() override
+	{
+		TArray<FMassEntityHandle> Entities;
+		EntityManager->BatchCreateEntities(IntsArchetype, 3, Entities);
+
+		TArray<TWeakPtr<int32>> SharedPtrs;
+		for (int32 Index = 0; Index < Entities.Num(); ++Index)
+		{
+			TSharedPtr<int32> TestData = MakeShared<int32>(Index);
+			SharedPtrs.Add(TestData.ToWeakPtr());
+
+			TArray<FInstancedStruct> Array;
+			Array.AddZeroed();
+			Array[0].InitializeAs<FFragmentWithSharedPtr>(TestData);
+
+			const FMassEntityHandle& EntityHandle = Entities[Index];
+			EntityManager->AddFragmentInstanceListToEntity(EntityHandle, Array);
+		}
+
+		for (int32 Index = 0; Index < Entities.Num(); ++Index)
+		{
+			TSharedPtr<int32> TestData = SharedPtrs[Index].Pin();
+			const FMassEntityHandle& EntityHandle = Entities[Index];
+			const FFragmentWithSharedPtr& Fragment = EntityManager->GetFragmentDataChecked<FFragmentWithSharedPtr>(EntityHandle);
+			AITEST_EQUAL("Data stored in fragments is the same as the data provided initially", *Fragment.Data, *TestData);
+		}
+
+		EntityManager->AddTagToEntity(Entities[0], FTestTag_A::StaticStruct());
+		EntityManager->AddFragmentToEntity(Entities[1], FTestFragment_Float::StaticStruct());
+
+		for (int32 Index = 0; Index < Entities.Num(); ++Index)
+		{
+			TSharedPtr<int32> TestData = SharedPtrs[Index].Pin();
+			const FMassEntityHandle& EntityHandle = Entities[Index];
+			const FFragmentWithSharedPtr& Fragment = EntityManager->GetFragmentDataChecked<FFragmentWithSharedPtr>(EntityHandle);
+			AITEST_EQUAL("After move operations: Data stored in fragments is the same as the data provided initially", *Fragment.Data, *TestData);
+		}
+
+		EntityManager->BatchDestroyEntities(Entities);
+		for (int32 Index = 0; Index < Entities.Num(); ++Index)
+		{
+			TSharedPtr<int32> TestData = SharedPtrs[Index].Pin();
+			AITEST_FALSE("After deletion we expect shared data to be released", TestData.IsValid());
+		}
+
+		return true;
+	}
+};
+IMPLEMENT_AI_INSTANT_TEST(FEntityTest_SharedPtrFragment, "System.Mass.Entity.SharedPtrFragment");
+
 } // FMassEntityTestTest
 
 UE_ENABLE_OPTIMIZATION_SHIP
diff --git a/Engine/Plugins/Runtime/MassEntity/Source/MassEntityTestSuite/Public/MassEntityTestTypes.h b/Engine/Plugins/Runtime/MassEntity/Source/MassEntityTestSuite/Public/MassEntityTestTypes.h
--- a/Engine/Plugins/Runtime/MassEntity/Source/MassEntityTestSuite/Public/MassEntityTestTypes.h
+++ b/Engine/Plugins/Runtime/MassEntity/Source/MassEntityTestSuite/Public/MassEntityTestTypes.h
@@ -77,6 +77,18 @@ struct FTestFragment_Array : public FMassFragment
 	}
 };
 
+USTRUCT()
+struct FFragmentWithSharedPtr : public FMassFragment
+{
+	GENERATED_BODY()
+	TSharedPtr<int32> Data;
+
+	FFragmentWithSharedPtr() = default;
+	FFragmentWithSharedPtr(TSharedPtr<int32>& InData)
+		: Data(InData)
+	{}
+};
+
 USTRUCT()
 struct FTestChunkFragment_Int : public FMassChunkFragment
 {
