[Mass] Added support for queries containing solely optional requirements. This feature required changes to the logic of determining whether a given archetype matches a query's requirements. The changes done has a side effect of no longer requiring queries to contain some kind of fragment requirements - now anything can be a standalone requirement (so it's possible to "get all archetypes that have a given chunk fragment" or "using given shared fragment type"). Same goes for negative requirements - now it's possible to have negative-only queries!

diff --git a/Engine/Plugins/Runtime/MassEntity/Source/MassEntity/Private/MassArchetypeData.cpp b/Engine/Plugins/Runtime/MassEntity/Source/MassEntity/Private/MassArchetypeData.cpp
--- a/Engine/Plugins/Runtime/MassEntity/Source/MassEntity/Private/MassArchetypeData.cpp
+++ b/Engine/Plugins/Runtime/MassEntity/Source/MassEntity/Private/MassArchetypeData.cpp
@@ -4,6 +4,7 @@
 #include "MassEntityTypes.h"
 #include "MassExecutionContext.h"
 #include "MassEntitySettings.h"
+#include "MassDebugger.h"
 #include "Misc/StringBuilder.h"
 
 
@@ -1335,192 +1336,226 @@ bool FMassArchetypeHelper::DoesArchetypeMatchRequirements(const FMassArchetypeCo
 bool FMassArchetypeHelper::DoesArchetypeMatchRequirements(const FMassArchetypeData& Archetype, const FMassFragmentRequirements& Requirements
 	, const bool bBailOutOnFirstFail, FOutputDevice* OutputDevice)
 {
-	if (OutputDevice == nullptr)
+	if (DoesArchetypeMatchRequirements(Archetype.GetCompositionDescriptor(), Requirements))
 	{
-		return DoesArchetypeMatchRequirements(Archetype.GetCompositionDescriptor(), Requirements);
+		// nothing to log
+		return true;
 	}
-
-	bool bResult = true;
-	if (Archetype.GetTagBitSet().HasAll(Requirements.GetRequiredAllTags()) == false)
-	{
-		// missing some required tags, skip.
-		const FMassTagBitSet UnsatisfiedTags = Requirements.GetRequiredAllTags() - Archetype.GetTagBitSet();
-		FStringOutputDevice Description;
-		UnsatisfiedTags.DebugGetStringDesc(Description);
-		OutputDevice->Logf(TEXT("Archetype did not match due to missing tags: %s"), *Description);
-		
-		bResult = false;
-		if (bBailOutOnFirstFail)
-		{
-			return false;
-		}
-	}
-
-	if (Archetype.GetTagBitSet().HasNone(Requirements.GetRequiredNoneTags()) == false)
-	{
-		// has some tags required to be absent
-		const FMassTagBitSet UnwantedTags = Requirements.GetRequiredNoneTags().GetOverlap(Archetype.GetTagBitSet());
-		FStringOutputDevice Description;
-		UnwantedTags.DebugGetStringDesc(Description);
-		OutputDevice->Logf(TEXT("Archetype has tags required absent: %s"), *Description);
-
-		bResult = false;
-		if (bBailOutOnFirstFail)
-		{
-			return false;
-		}
-	}
-
-	if (Requirements.GetRequiredAnyTags().IsEmpty() == false
-		&& Archetype.GetTagBitSet().HasAny(Requirements.GetRequiredAnyTags()) == false)
-	{
-		FStringOutputDevice Description;
-		Requirements.GetRequiredAnyTags().DebugGetStringDesc(Description);
-		OutputDevice->Logf(TEXT("Archetype did not match due to missing \'any\' tags: %s"), *Description);
-
-		bResult = false;
-		if (bBailOutOnFirstFail)
-		{
-			return false;
-		}
-	}
-
-	if (Archetype.GetFragmentBitSet().HasAll(Requirements.GetRequiredAllFragments()) == false)
-	{
-		// missing some required fragments, skip.
-		const FMassFragmentBitSet UnsatisfiedFragments = Requirements.GetRequiredAllFragments() - Archetype.GetFragmentBitSet();
-		FStringOutputDevice Description;
-		UnsatisfiedFragments.DebugGetStringDesc(Description);
-		OutputDevice->Logf(TEXT("Archetype did not match due to missing Fragments: %s"), *Description);
-
-		bResult = false;
-		if (bBailOutOnFirstFail)
-		{
-			return false;
-		}
-	}
-
-	if (Archetype.GetFragmentBitSet().HasNone(Requirements.GetRequiredNoneFragments()) == false)
-	{
-		// has some Fragments required to be absent
-		const FMassFragmentBitSet UnwantedFragments = Requirements.GetRequiredNoneFragments().GetOverlap(Archetype.GetFragmentBitSet());
-		FStringOutputDevice Description;
-		UnwantedFragments.DebugGetStringDesc(Description);
-		OutputDevice->Logf(TEXT("Archetype has Fragments required absent: %s"), *Description);
-
-		bResult = false;
-		if (bBailOutOnFirstFail)
-		{
-			return false;
-		}
-	}
-
-	if (Requirements.GetRequiredAnyFragments().IsEmpty() == false
-		&& Archetype.GetFragmentBitSet().HasAny(Requirements.GetRequiredAnyFragments()) == false)
-	{
-		FStringOutputDevice Description;
-		Requirements.GetRequiredAnyFragments().DebugGetStringDesc(Description);
-		OutputDevice->Logf(TEXT("Archetype did not match due to missing \'any\' fragments: %s"), *Description);
-
-		bResult = false;
-		if (bBailOutOnFirstFail)
-		{
-			return false;
-		}
-	}
-
-	if (Archetype.GetChunkFragmentBitSet().HasAll(Requirements.GetRequiredAllChunkFragments()) == false)
-	{
-		// missing some required fragments, skip.
-		const FMassChunkFragmentBitSet UnsatisfiedFragments = Requirements.GetRequiredAllChunkFragments() - Archetype.GetChunkFragmentBitSet();
-		FStringOutputDevice Description;
-		UnsatisfiedFragments.DebugGetStringDesc(Description);
-		OutputDevice->Logf(TEXT("Archetype did not match due to missing Chunk Fragments: %s"), *Description);
-
-		bResult = false;
-		if (bBailOutOnFirstFail)
-		{
-			return false;
-		}
-	}
-
-	if (Archetype.GetChunkFragmentBitSet().HasNone(Requirements.GetRequiredNoneChunkFragments()) == false)
-	{
-		// has some Fragments required to be absent
-		const FMassChunkFragmentBitSet UnwantedFragments = Requirements.GetRequiredNoneChunkFragments().GetOverlap(Archetype.GetChunkFragmentBitSet());
-		FStringOutputDevice Description;
-		UnwantedFragments.DebugGetStringDesc(Description);
-		OutputDevice->Logf(TEXT("Archetype has Chunk Fragments required absent: %s"), *Description);
-
-		bResult = false;
-		if (bBailOutOnFirstFail)
-		{
-			return false;
-		}
-	}
-
-	if (Archetype.GetSharedFragmentBitSet().HasAll(Requirements.GetRequiredAllSharedFragments()) == false)
-	{
-		// missing some required fragments, skip.
-		const FMassSharedFragmentBitSet UnsatisfiedFragments = Requirements.GetRequiredAllSharedFragments() - Archetype.GetSharedFragmentBitSet();
-		FStringOutputDevice Description;
-		UnsatisfiedFragments.DebugGetStringDesc(Description);
-		OutputDevice->Logf(TEXT("Archetype did not match due to missing Shared Fragments: %s"), *Description);
-
-		bResult = false;
-		if (bBailOutOnFirstFail)
-		{
-			return false;
-		}
-	}
-
-	if (Archetype.GetSharedFragmentBitSet().HasNone(Requirements.GetRequiredNoneSharedFragments()) == false)
-	{
-		// has some Fragments required to be absent
-		const FMassSharedFragmentBitSet UnwantedFragments = Requirements.GetRequiredNoneSharedFragments().GetOverlap(Archetype.GetSharedFragmentBitSet());
-		FStringOutputDevice Description;
-		UnwantedFragments.DebugGetStringDesc(Description);
-		OutputDevice->Logf(TEXT("Archetype has Shared Fragments required absent: %s"), *Description);
-
-		bResult = false;
-		if (bBailOutOnFirstFail)
-		{
-			return false;
-		}
-	}
-
-	if (Archetype.GetConstSharedFragmentBitSet().HasAll(Requirements.GetRequiredAllConstSharedFragments()) == false)
-	{
-		// missing some required fragments, skip.
-		const FMassConstSharedFragmentBitSet UnsatisfiedFragments = Requirements.GetRequiredAllConstSharedFragments() - Archetype.GetConstSharedFragmentBitSet();
-		FStringOutputDevice Description;
-		UnsatisfiedFragments.DebugGetStringDesc(Description);
-		OutputDevice->Logf(TEXT("Archetype did not match due to missing Const Shared Fragments: %s"), *Description);
-
-		bResult = false;
-		if (bBailOutOnFirstFail)
-		{
-			return false;
-		}
-	}
-
-	if (Archetype.GetConstSharedFragmentBitSet().HasNone(Requirements.GetRequiredNoneConstSharedFragments()) == false)
-	{
-		// has some Fragments required to be absent
-		const FMassConstSharedFragmentBitSet UnwantedFragments = Requirements.GetRequiredNoneConstSharedFragments().GetOverlap(Archetype.GetConstSharedFragmentBitSet());
-		FStringOutputDevice Description;
-		UnwantedFragments.DebugGetStringDesc(Description);
-		OutputDevice->Logf(TEXT("Archetype has Const Shared Fragments required absent: %s"), *Description);
-
-		// could skip the test as it's the final check we're performing, but leaving it in in case more cases are added
-		// later without checking the existing code
-		bResult = false;
-		if (bBailOutOnFirstFail)
-		{
-			return false;
-		}
-	}
-
-	return bResult;
+	
+	if (OutputDevice)
+	{
+		// do logging
+		OutputDevice->Logf(TEXT("% s")
+			, *FMassDebugger::GetArchetypeRequirementCompatibilityDescription(Requirements, Archetype.GetCompositionDescriptor()));
+	}
+
+	return false;
+
+	//bool bResult = true;
+	//if (Requirements.IsOptionalsOnly() == false)
+	//{
+	//	if (Archetype.GetTagBitSet().HasAll(Requirements.GetRequiredAllTags()) == false)
+	//	{
+	//		// missing some required tags, skip.
+	//		const FMassTagBitSet UnsatisfiedTags = Requirements.GetRequiredAllTags() - Archetype.GetTagBitSet();
+	//		FStringOutputDevice Description;
+	//		UnsatisfiedTags.DebugGetStringDesc(Description);
+	//		OutputDevice->Logf(TEXT("Archetype did not match due to missing tags: %s"), *Description);
+
+	//		bResult = false;
+	//		if (bBailOutOnFirstFail)
+	//		{
+	//			return false;
+	//		}
+	//	}
+
+	//	if (Archetype.GetTagBitSet().HasNone(Requirements.GetRequiredNoneTags()) == false)
+	//	{
+	//		// has some tags required to be absent
+	//		const FMassTagBitSet UnwantedTags = Requirements.GetRequiredNoneTags().GetOverlap(Archetype.GetTagBitSet());
+	//		FStringOutputDevice Description;
+	//		UnwantedTags.DebugGetStringDesc(Description);
+	//		OutputDevice->Logf(TEXT("Archetype has tags required absent: %s"), *Description);
+
+	//		bResult = false;
+	//		if (bBailOutOnFirstFail)
+	//		{
+	//			return false;
+	//		}
+	//	}
+
+	//	if (Requirements.GetRequiredAnyTags().IsEmpty() == false
+	//		&& Archetype.GetTagBitSet().HasAny(Requirements.GetRequiredAnyTags()) == false)
+	//	{
+	//		FStringOutputDevice Description;
+	//		Requirements.GetRequiredAnyTags().DebugGetStringDesc(Description);
+	//		OutputDevice->Logf(TEXT("Archetype did not match due to missing \'any\' tags: %s"), *Description);
+
+	//		bResult = false;
+	//		if (bBailOutOnFirstFail)
+	//		{
+	//			return false;
+	//		}
+	//	}
+
+	//	if (Archetype.GetFragmentBitSet().HasAll(Requirements.GetRequiredAllFragments()) == false)
+	//	{
+	//		// missing some required fragments, skip.
+	//		const FMassFragmentBitSet UnsatisfiedFragments = Requirements.GetRequiredAllFragments() - Archetype.GetFragmentBitSet();
+	//		FStringOutputDevice Description;
+	//		UnsatisfiedFragments.DebugGetStringDesc(Description);
+	//		OutputDevice->Logf(TEXT("Archetype did not match due to missing Fragments: %s"), *Description);
+
+	//		bResult = false;
+	//		if (bBailOutOnFirstFail)
+	//		{
+	//			return false;
+	//		}
+	//	}
+
+	//	if (Archetype.GetFragmentBitSet().HasNone(Requirements.GetRequiredNoneFragments()) == false)
+	//	{
+	//		// has some Fragments required to be absent
+	//		const FMassFragmentBitSet UnwantedFragments = Requirements.GetRequiredNoneFragments().GetOverlap(Archetype.GetFragmentBitSet());
+	//		FStringOutputDevice Description;
+	//		UnwantedFragments.DebugGetStringDesc(Description);
+	//		OutputDevice->Logf(TEXT("Archetype has Fragments required absent: %s"), *Description);
+
+	//		bResult = false;
+	//		if (bBailOutOnFirstFail)
+	//		{
+	//			return false;
+	//		}
+	//	}
+
+	//	if (Requirements.GetRequiredAnyFragments().IsEmpty() == false
+	//		&& Archetype.GetFragmentBitSet().HasAny(Requirements.GetRequiredAnyFragments()) == false)
+	//	{
+	//		FStringOutputDevice Description;
+	//		Requirements.GetRequiredAnyFragments().DebugGetStringDesc(Description);
+	//		OutputDevice->Logf(TEXT("Archetype did not match due to missing \'any\' fragments: %s"), *Description);
+
+	//		bResult = false;
+	//		if (bBailOutOnFirstFail)
+	//		{
+	//			return false;
+	//		}
+	//	}
+
+	//	if (Archetype.GetChunkFragmentBitSet().HasAll(Requirements.GetRequiredAllChunkFragments()) == false)
+	//	{
+	//		// missing some required fragments, skip.
+	//		const FMassChunkFragmentBitSet UnsatisfiedFragments = Requirements.GetRequiredAllChunkFragments() - Archetype.GetChunkFragmentBitSet();
+	//		FStringOutputDevice Description;
+	//		UnsatisfiedFragments.DebugGetStringDesc(Description);
+	//		OutputDevice->Logf(TEXT("Archetype did not match due to missing Chunk Fragments: %s"), *Description);
+
+	//		bResult = false;
+	//		if (bBailOutOnFirstFail)
+	//		{
+	//			return false;
+	//		}
+	//	}
+
+	//	if (Archetype.GetChunkFragmentBitSet().HasNone(Requirements.GetRequiredNoneChunkFragments()) == false)
+	//	{
+	//		// has some Fragments required to be absent
+	//		const FMassChunkFragmentBitSet UnwantedFragments = Requirements.GetRequiredNoneChunkFragments().GetOverlap(Archetype.GetChunkFragmentBitSet());
+	//		FStringOutputDevice Description;
+	//		UnwantedFragments.DebugGetStringDesc(Description);
+	//		OutputDevice->Logf(TEXT("Archetype has Chunk Fragments required absent: %s"), *Description);
+
+	//		bResult = false;
+	//		if (bBailOutOnFirstFail)
+	//		{
+	//			return false;
+	//		}
+	//	}
+
+	//	if (Archetype.GetSharedFragmentBitSet().HasAll(Requirements.GetRequiredAllSharedFragments()) == false)
+	//	{
+	//		// missing some required fragments, skip.
+	//		const FMassSharedFragmentBitSet UnsatisfiedFragments = Requirements.GetRequiredAllSharedFragments() - Archetype.GetSharedFragmentBitSet();
+	//		FStringOutputDevice Description;
+	//		UnsatisfiedFragments.DebugGetStringDesc(Description);
+	//		OutputDevice->Logf(TEXT("Archetype did not match due to missing Shared Fragments: %s"), *Description);
+
+	//		bResult = false;
+	//		if (bBailOutOnFirstFail)
+	//		{
+	//			return false;
+	//		}
+	//	}
+
+	//	if (Archetype.GetSharedFragmentBitSet().HasNone(Requirements.GetRequiredNoneSharedFragments()) == false)
+	//	{
+	//		// has some Fragments required to be absent
+	//		const FMassSharedFragmentBitSet UnwantedFragments = Requirements.GetRequiredNoneSharedFragments().GetOverlap(Archetype.GetSharedFragmentBitSet());
+	//		FStringOutputDevice Description;
+	//		UnwantedFragments.DebugGetStringDesc(Description);
+	//		OutputDevice->Logf(TEXT("Archetype has Shared Fragments required absent: %s"), *Description);
+
+	//		bResult = false;
+	//		if (bBailOutOnFirstFail)
+	//		{
+	//			return false;
+	//		}
+	//	}
+
+	//	if (Archetype.GetConstSharedFragmentBitSet().HasAll(Requirements.GetRequiredAllConstSharedFragments()) == false)
+	//	{
+	//		// missing some required fragments, skip.
+	//		const FMassConstSharedFragmentBitSet UnsatisfiedFragments = Requirements.GetRequiredAllConstSharedFragments() - Archetype.GetConstSharedFragmentBitSet();
+	//		FStringOutputDevice Description;
+	//		UnsatisfiedFragments.DebugGetStringDesc(Description);
+	//		OutputDevice->Logf(TEXT("Archetype did not match due to missing Const Shared Fragments: %s"), *Description);
+
+	//		bResult = false;
+	//		if (bBailOutOnFirstFail)
+	//		{
+	//			return false;
+	//		}
+	//	}
+
+	//	if (Archetype.GetConstSharedFragmentBitSet().HasNone(Requirements.GetRequiredNoneConstSharedFragments()) == false)
+	//	{
+	//		// has some Fragments required to be absent
+	//		const FMassConstSharedFragmentBitSet UnwantedFragments = Requirements.GetRequiredNoneConstSharedFragments().GetOverlap(Archetype.GetConstSharedFragmentBitSet());
+	//		FStringOutputDevice Description;
+	//		UnwantedFragments.DebugGetStringDesc(Description);
+	//		OutputDevice->Logf(TEXT("Archetype has Const Shared Fragments required absent: %s"), *Description);
+
+	//		// could skip the test as it's the final check we're performing, but leaving it in in case more cases are added
+	//		// later without checking the existing code
+	//		bResult = false;
+	//		if (bBailOutOnFirstFail)
+	//		{
+	//			return false;
+	//		}
+	//	}
+	//}
+	//else
+	//{
+	//	// test if contains any of the optional elements
+	//	if (Archetype.GetFragmentBitSet().HasNone(Requirements.GetRequiredOptionalFragments())
+	//		&& Archetype.GetTagBitSet().HasNone(Requirements.GetRequiredOptionalTags())
+	//		&& Archetype.GetChunkFragmentBitSet().HasNone(Requirements.GetRequiredOptionalChunkFragments())
+	//		&& Archetype.GetSharedFragmentBitSet().HasNone(Requirements.GetRequiredOptionalSharedFragments())
+	//		&& Archetype.GetConstSharedFragmentBitSet().HasNone(Requirements.GetRequiredOptionalConstSharedFragments()))
+	//	{
+	//		FStringOutputDevice Description;
+	//		Requirements.GetRequiredOptionalFragments().DebugGetStringDesc(Description);
+	//		Requirements.GetRequiredOptionalTags().DebugGetStringDesc(Description);
+	//		Requirements.GetRequiredOptionalChunkFragments().DebugGetStringDesc(Description);
+	//		Requirements.GetRequiredOptionalSharedFragments().DebugGetStringDesc(Description);
+	//		Requirements.GetRequiredOptionalConstSharedFragments().DebugGetStringDesc(Description);
+	//		
+	//		OutputDevice->Logf(TEXT("Archetype has none of the optional elements: %s"), *Description);
+
+	//		bResult = false;
+	//	}
+	//}
+
+	//return bResult;
 }
 #endif // WITH_MASSENTITY_DEBUG
\ No newline at end of file
diff --git a/Engine/Plugins/Runtime/MassEntity/Source/MassEntity/Private/MassDebugger.cpp b/Engine/Plugins/Runtime/MassEntity/Source/MassEntity/Private/MassDebugger.cpp
--- a/Engine/Plugins/Runtime/MassEntity/Source/MassEntity/Private/MassDebugger.cpp
+++ b/Engine/Plugins/Runtime/MassEntity/Source/MassEntity/Private/MassDebugger.cpp
@@ -401,7 +401,7 @@ TConstArrayView<FMassEntityQuery*> FMassDebugger::GetUpToDateProcessorQueries(co
 UE::Mass::Debug::FQueryRequirementsView FMassDebugger::GetQueryRequirements(const FMassEntityQuery& Query)
 {
 	UE::Mass::Debug::FQueryRequirementsView View = { Query.FragmentRequirements, Query.ChunkFragmentRequirements, Query.ConstSharedFragmentRequirements, Query.SharedFragmentRequirements
-		, Query.RequiredAllTags, Query.RequiredAnyTags, Query.RequiredNoneTags
+		, Query.RequiredAllTags, Query.RequiredAnyTags, Query.RequiredNoneTags, Query.RequiredOptionalTags
 		, Query.RequiredConstSubsystems, Query.RequiredMutableSubsystems };
 
 	return View;
@@ -499,60 +499,107 @@ FString FMassDebugger::GetArchetypeRequirementCompatibilityDescription(const FMa
 {
 	FStringOutputDevice OutDescription;
 
-	if (ArchetypeComposition.Fragments.HasAll(Requirements.RequiredAllFragments) == false)
+	if (Requirements.HasNegativeRequirements())
 	{
-		// missing one of the strictly required fragments
-		OutDescription += TEXT("\nMissing required fragments: ");
-		(Requirements.RequiredAllFragments - ArchetypeComposition.Fragments).DebugGetStringDesc(OutDescription);
-	}
+		if (ArchetypeComposition.Fragments.HasNone(Requirements.RequiredNoneFragments) == false)
+		{
+			// has some of the fragments required absent
+			OutDescription += TEXT("\nHas fragments required absent: ");
+			(Requirements.RequiredNoneFragments & ArchetypeComposition.Fragments).DebugGetStringDesc(OutDescription);
+		}
 
-	if (Requirements.RequiredAnyFragments.IsEmpty() == false && ArchetypeComposition.Fragments.HasAny(Requirements.RequiredAnyFragments) == false)
-	{
-		// missing all of the "any" fragments
-		OutDescription += TEXT("\nMissing all \'any\' fragments: ");
-		Requirements.RequiredAnyFragments.DebugGetStringDesc(OutDescription);
-	}
+		if (ArchetypeComposition.Tags.HasNone(Requirements.RequiredNoneTags) == false)
+		{
+			// has some of the tags required absent
+			OutDescription += TEXT("\nHas tags required absent: ");
+			(Requirements.RequiredNoneTags & ArchetypeComposition.Tags).DebugGetStringDesc(OutDescription);
+		}
 
-	if (ArchetypeComposition.Fragments.HasNone(Requirements.RequiredNoneFragments) == false)
-	{
-		// has some of the fragments required absent
-		OutDescription += TEXT("\nHas fragments required absent: ");
-		Requirements.RequiredNoneFragments.DebugGetStringDesc(OutDescription);
-	}
+		if (ArchetypeComposition.ChunkFragments.HasNone(Requirements.RequiredNoneChunkFragments) == false)
+		{
+			// has some of the chunk fragments required absent
+			OutDescription += TEXT("\nHas chunk fragments required absent: ");
+			(Requirements.RequiredNoneChunkFragments & ArchetypeComposition.ChunkFragments).DebugGetStringDesc(OutDescription);
+		}
 
-	if (ArchetypeComposition.Tags.HasAll(Requirements.RequiredAllTags) == false)
-	{
-		// missing one of the strictly required tags
-		OutDescription += TEXT("\nMissing required tags: ");
-		(Requirements.RequiredAllTags - ArchetypeComposition.Tags).DebugGetStringDesc(OutDescription);
-	}
+		if (ArchetypeComposition.SharedFragments.HasNone(Requirements.RequiredNoneSharedFragments) == false)
+		{
+			// has some of the chunk fragments required absent
+			OutDescription += TEXT("\nHas shared fragments required absent: ");
+			(Requirements.RequiredNoneSharedFragments & ArchetypeComposition.SharedFragments).DebugGetStringDesc(OutDescription);
+		}
 
-	if (Requirements.RequiredAnyTags.IsEmpty() == false && ArchetypeComposition.Tags.HasAny(Requirements.RequiredAnyTags) == false)
-	{
-		// missing all of the "any" tags
-		OutDescription += TEXT("\nMissing all \'any\' tags: ");
-		Requirements.RequiredAnyTags.DebugGetStringDesc(OutDescription);
+		if (ArchetypeComposition.ConstSharedFragments.HasNone(Requirements.RequiredNoneConstSharedFragments) == false)
+		{
+			// has some of the chunk fragments required absent
+			OutDescription += TEXT("\nHas shared fragments required absent: ");
+			(Requirements.RequiredNoneConstSharedFragments & ArchetypeComposition.ConstSharedFragments).DebugGetStringDesc(OutDescription);
+		}
 	}
 
-	if (ArchetypeComposition.Tags.HasNone(Requirements.RequiredNoneTags) == false)
+	// if we have regular (i.e. non-optional) positive requirements then these are the determining factor, we don't check optionals
+	if (Requirements.HasPositiveRequirements())
 	{
-		// has some of the tags required absent
-		OutDescription += TEXT("\nHas tags required absent: ");
-		Requirements.RequiredNoneTags.DebugGetStringDesc(OutDescription);
-	}
+		if (ArchetypeComposition.Fragments.HasAll(Requirements.RequiredAllFragments) == false)
+		{
+			// missing one of the strictly required fragments
+			OutDescription += TEXT("\nMissing required fragments: ");
+			(Requirements.RequiredAllFragments - ArchetypeComposition.Fragments).DebugGetStringDesc(OutDescription);
+		}
 
-	if (ArchetypeComposition.ChunkFragments.HasAll(Requirements.RequiredAllChunkFragments) == false)
-	{
-		// missing one of the strictly required chunk fragments
-		OutDescription += TEXT("\nMissing required chunk fragments: ");
-		(Requirements.RequiredAllChunkFragments - ArchetypeComposition.ChunkFragments).DebugGetStringDesc(OutDescription);
-	}
+		if (Requirements.RequiredAnyFragments.IsEmpty() == false && ArchetypeComposition.Fragments.HasAny(Requirements.RequiredAnyFragments) == false)
+		{
+			// missing all of the "any" fragments
+			OutDescription += TEXT("\nMissing all \'any\' fragments: ");
+			Requirements.RequiredAnyFragments.DebugGetStringDesc(OutDescription);
+		}
+
+		if (ArchetypeComposition.Tags.HasAll(Requirements.RequiredAllTags) == false)
+		{
+			// missing one of the strictly required tags
+			OutDescription += TEXT("\nMissing required tags: ");
+			(Requirements.RequiredAllTags - ArchetypeComposition.Tags).DebugGetStringDesc(OutDescription);
+		}
+
+		if (Requirements.RequiredAnyTags.IsEmpty() == false && ArchetypeComposition.Tags.HasAny(Requirements.RequiredAnyTags) == false)
+		{
+			// missing all of the "any" tags
+			OutDescription += TEXT("\nMissing all \'any\' tags: ");
+			Requirements.RequiredAnyTags.DebugGetStringDesc(OutDescription);
+		}
+
+		if (ArchetypeComposition.ChunkFragments.HasAll(Requirements.RequiredAllChunkFragments) == false)
+		{
+			// missing one of the strictly required chunk fragments
+			OutDescription += TEXT("\nMissing required chunk fragments: ");
+			(Requirements.RequiredAllChunkFragments - ArchetypeComposition.ChunkFragments).DebugGetStringDesc(OutDescription);
+		}
+
+		if (ArchetypeComposition.SharedFragments.HasAll(Requirements.RequiredAllSharedFragments) == false)
+		{
+			// missing one of the strictly required Shared fragments
+			OutDescription += TEXT("\nMissing required Shared fragments: ");
+			(Requirements.RequiredAllSharedFragments - ArchetypeComposition.SharedFragments).DebugGetStringDesc(OutDescription);
+		}
 
-	if (ArchetypeComposition.ChunkFragments.HasNone(Requirements.RequiredNoneChunkFragments) == false)
+		if (ArchetypeComposition.ConstSharedFragments.HasAll(Requirements.RequiredAllConstSharedFragments) == false)
+		{
+			// missing one of the strictly required Shared fragments
+			OutDescription += TEXT("\nMissing required Shared fragments: ");
+			(Requirements.RequiredAllConstSharedFragments - ArchetypeComposition.ConstSharedFragments).DebugGetStringDesc(OutDescription);
+		}
+	}
+	// else we check if there are any optionals and if so test them
+	else if (Requirements.HasOptionalRequirements() && (Requirements.DoesMatchAnyOptionals(ArchetypeComposition) == false))
 	{
-		// has some of the chunk fragments required absent
-		OutDescription += TEXT("\nHas chunk fragments required absent: ");
-		Requirements.RequiredNoneChunkFragments.DebugGetStringDesc(OutDescription);
+		// we report that none of the optionals has been met
+		OutDescription += TEXT("\nNone of the optionals were safisfied while not having other positive hard requirements: ");
+
+		Requirements.RequiredOptionalTags.DebugGetStringDesc(OutDescription);
+		Requirements.RequiredOptionalFragments.DebugGetStringDesc(OutDescription);
+		Requirements.RequiredOptionalChunkFragments.DebugGetStringDesc(OutDescription);
+		Requirements.RequiredOptionalSharedFragments.DebugGetStringDesc(OutDescription);
+		Requirements.RequiredOptionalConstSharedFragments.DebugGetStringDesc(OutDescription);
 	}
 
 	return OutDescription.Len() > 0 ? static_cast<FString>(OutDescription) : TEXT("Match");
diff --git a/Engine/Plugins/Runtime/MassEntity/Source/MassEntity/Private/MassEntityManager.cpp b/Engine/Plugins/Runtime/MassEntity/Source/MassEntity/Private/MassEntityManager.cpp
--- a/Engine/Plugins/Runtime/MassEntity/Source/MassEntity/Private/MassEntityManager.cpp
+++ b/Engine/Plugins/Runtime/MassEntity/Source/MassEntity/Private/MassEntityManager.cpp
@@ -1735,75 +1735,19 @@ void FMassEntityManager::CheckIfEntityIsActive(FMassEntityHandle Entity) const
 
 void FMassEntityManager::GetMatchingArchetypes(const FMassFragmentRequirements& Requirements, TArray<FMassArchetypeHandle>& OutValidArchetypes, const uint32 FromArchetypeDataVersion) const
 {
-	//@TODO: Not optimized yet, but we call this rarely now, so not a big deal.
-
-	// First get set of all archetypes that contain *any* fragment
-	TSet<TSharedPtr<FMassArchetypeData>> AnyArchetypes;
-	TConstArrayView<FMassFragmentRequirementDescription> FragmentRequirements = Requirements.GetFragmentRequirements();
-
-	// We need to find out if the query does indeed strictly require any fragments to be there - it's for example 
-	// possible to build a query consisting of only optional fragment requirements. In that case we need to test for 
-	// tags as well, independently of fragments, because none of the fragments might be there and we still need to 
-	// respect the tags.
-	// A practical example would be a query asking for all entities that have a given fragment OR a given tag.
-	bool bHasHardFragmentdRequirements = false;
-
-	if (!FragmentRequirements.IsEmpty())
+	for (int32 ArchetypeIndex = FromArchetypeDataVersion; ArchetypeIndex < AllArchetypes.Num(); ++ArchetypeIndex)
 	{
-		for (const FMassFragmentRequirementDescription& Requirement : FragmentRequirements)
-		{
-			check(Requirement.StructType);
-			if (Requirement.Presence != EMassFragmentPresence::None)
-			{
-				if (const TArray<TSharedPtr<FMassArchetypeData>>* pData = FragmentTypeToArchetypeMap.Find(Requirement.StructType))
-				{
-					AnyArchetypes.Append(*pData);
-					// only Any and All count as a "hard requirement" since other modes don't strictly require a fragment to be there
-					bHasHardFragmentdRequirements = bHasHardFragmentdRequirements
-						|| Requirement.Presence == EMassFragmentPresence::All
-						|| Requirement.Presence == EMassFragmentPresence::Any;
-				}
-			}
-		}
-	}
-	
-	// If the query does not have any hard requirements for fragments, but has tag requirements.
-	// We need to search all archetypes, because we don't have a Tag->Archetype Map (and we don't want to create one 
-	// just for this one edge case) so we need to go through all archetypes and match tags. On the bright side,
-	// we can ignore all archetypes with data version < FromArchetypeDataVersion
-	// Note that we collect all archetypes even remotely matching to let the regular process below to filter out
-	// the failing archetypes while justifying the reasons. 
-	if (bHasHardFragmentdRequirements == false
-		&& (Requirements.GetRequiredAllTags().IsEmpty() == false || Requirements.GetRequiredAnyTags().IsEmpty() == false))
-	{
-		const FMassTagBitSet RoughTagFilter = Requirements.GetRequiredAllTags() | Requirements.GetRequiredAnyTags();
-		for (int32 ArchetypeIndex = FromArchetypeDataVersion; ArchetypeIndex < AllArchetypes.Num(); ++ArchetypeIndex)
-		{
-			const TSharedPtr<FMassArchetypeData>& ArchetypeData = AllArchetypes[ArchetypeIndex];
-			if (ArchetypeData.IsValid())
-			{
-				if (ArchetypeData->GetTagBitSet().HasAny(RoughTagFilter))
-				{
-					AnyArchetypes.Add(ArchetypeData);
-				}
-			}
-		}
-	}
+		checkf(AllArchetypes[ArchetypeIndex].IsValid(), TEXT("We never expect to get any invalid shared ptrs in AllArchetypes"));
 
-	// Then verify that they contain *all* required fragments
-	for (TSharedPtr<FMassArchetypeData>& ArchetypePtr : AnyArchetypes)
-	{
-		FMassArchetypeData& Archetype = *(ArchetypePtr.Get());
+		FMassArchetypeData& Archetype = *(AllArchetypes[ArchetypeIndex].Get());
 
 		// Only return archetypes with a newer created version than the specified version, this is for incremental query updates
-		if (Archetype.GetCreatedArchetypeDataVersion() <= FromArchetypeDataVersion)
-		{
-			continue;
-		}
+		ensureMsgf(Archetype.GetCreatedArchetypeDataVersion() > FromArchetypeDataVersion
+			, TEXT("There's a stron assumption that archetype's data version corresponds to its index in AllArchetypes"));
 
 		if (Requirements.DoesArchetypeMatchRequirements(Archetype.GetCompositionDescriptor()))
 		{
-			OutValidArchetypes.Add(ArchetypePtr);
+			OutValidArchetypes.Add(AllArchetypes[ArchetypeIndex]);
 		}
 #if WITH_MASSENTITY_DEBUG
 		else
diff --git a/Engine/Plugins/Runtime/MassEntity/Source/MassEntity/Private/MassProcessorDependencySolver.cpp b/Engine/Plugins/Runtime/MassEntity/Source/MassEntity/Private/MassProcessorDependencySolver.cpp
--- a/Engine/Plugins/Runtime/MassEntity/Source/MassEntity/Private/MassProcessorDependencySolver.cpp
+++ b/Engine/Plugins/Runtime/MassEntity/Source/MassEntity/Private/MassProcessorDependencySolver.cpp
@@ -56,6 +56,7 @@ void FMassExecutionRequirements::Append(const FMassExecutionRequirements& Other)
 	RequiredAllTags += Other.RequiredAllTags;
 	RequiredAnyTags += Other.RequiredAnyTags;
 	RequiredNoneTags += Other.RequiredNoneTags;
+	// note that we're deliberately ignoring optional tags, they play no role here.
 
 	// signal that it requires recalculation;
 	ResourcesUsedCount = INDEX_NONE;
diff --git a/Engine/Plugins/Runtime/MassEntity/Source/MassEntity/Private/MassRequirements.cpp b/Engine/Plugins/Runtime/MassEntity/Source/MassEntity/Private/MassRequirements.cpp
--- a/Engine/Plugins/Runtime/MassEntity/Source/MassEntity/Private/MassRequirements.cpp
+++ b/Engine/Plugins/Runtime/MassEntity/Source/MassEntity/Private/MassRequirements.cpp
@@ -98,27 +98,57 @@ void FMassFragmentRequirements::SortRequirements()
 	SharedFragmentRequirements.Sort(FScriptStructSortOperator());
 }
 
-bool FMassFragmentRequirements::CheckValidity() const
+FORCEINLINE void FMassFragmentRequirements::CachePropreties() const
 {
-	if (bValidityIsCached == false)
+	if (bPropertiesCached == false)
 	{
-		bAreRequirementsValid = RequiredAllFragments.IsEmpty() == false || RequiredAnyFragments.IsEmpty() == false
-			|| RequiredOptionalFragments.IsEmpty() == false	|| RequiredAllTags.IsEmpty() == false || RequiredAnyTags.IsEmpty() == false;
-		bValidityIsCached = true;
+		bHasPositiveRequirements = !(RequiredAllTags.IsEmpty()
+			&& RequiredAnyTags.IsEmpty()
+			&& RequiredAllFragments.IsEmpty()
+			&& RequiredAnyFragments.IsEmpty()
+			&& RequiredAllChunkFragments.IsEmpty()
+			&& RequiredAllSharedFragments.IsEmpty()
+			&& RequiredAllConstSharedFragments.IsEmpty());
+
+		bHasNegativeRequirements = !(RequiredNoneTags.IsEmpty()
+			&& RequiredNoneFragments.IsEmpty()
+			&& RequiredNoneChunkFragments.IsEmpty()
+			&& RequiredNoneSharedFragments.IsEmpty()
+			&& RequiredNoneConstSharedFragments.IsEmpty());
+
+		bHasOptionalRequirements = !(RequiredOptionalFragments.IsEmpty()
+				&& RequiredOptionalTags.IsEmpty()
+				&& RequiredOptionalChunkFragments.IsEmpty()
+				&& RequiredOptionalSharedFragments.IsEmpty()
+				&& RequiredOptionalConstSharedFragments.IsEmpty());
+
+		bPropertiesCached = true;
 	}
-	return bAreRequirementsValid;
+}
+
+bool FMassFragmentRequirements::CheckValidity() const
+{
+	CachePropreties();
+	// @todo we need to add more sophisticated testing somewhere to detect contradicting requirements - like having and not having a given tag.
+	return bHasPositiveRequirements || bHasNegativeRequirements || bHasOptionalRequirements;
 }
 
 bool FMassFragmentRequirements::IsEmpty() const
 {
-	if (bEmptynessIsCached == false)
-	{
-		bAreRequirementsEmpty = FragmentRequirements.IsEmpty() && ChunkFragmentRequirements.IsEmpty() && ConstSharedFragmentRequirements.IsEmpty()
-			&& SharedFragmentRequirements.IsEmpty() && RequiredAllTags.IsEmpty() && RequiredAnyTags.IsEmpty() && RequiredNoneTags.IsEmpty();
-		bEmptynessIsCached = true;
-	}
+	CachePropreties();
+	// note that even though at the moment the following condition is the same as negation of current CheckValidity value
+	// that will change in the future (with additional validity checks).
+	return !bHasPositiveRequirements && !bHasNegativeRequirements && !bHasOptionalRequirements;
+}
 
-	return bAreRequirementsEmpty;
+bool FMassFragmentRequirements::DoesMatchAnyOptionals(const FMassArchetypeCompositionDescriptor& ArchetypeComposition) const
+{
+	return bHasOptionalRequirements
+		&& (ArchetypeComposition.Fragments.HasAny(RequiredOptionalFragments)
+			|| ArchetypeComposition.Tags.HasAny(RequiredOptionalTags)
+			|| ArchetypeComposition.ChunkFragments.HasAny(RequiredOptionalChunkFragments)
+			|| ArchetypeComposition.SharedFragments.HasAny(RequiredOptionalSharedFragments)
+			|| ArchetypeComposition.ConstSharedFragments.HasAny(RequiredOptionalConstSharedFragments));
 }
 
 bool FMassFragmentRequirements::DoesArchetypeMatchRequirements(const FMassArchetypeHandle& ArchetypeHandle) const
@@ -132,18 +162,35 @@ bool FMassFragmentRequirements::DoesArchetypeMatchRequirements(const FMassArchet
 	
 bool FMassFragmentRequirements::DoesArchetypeMatchRequirements(const FMassArchetypeCompositionDescriptor& ArchetypeComposition) const
 {
-	return ArchetypeComposition.Fragments.HasAll(RequiredAllFragments)
-		&& (RequiredAnyFragments.IsEmpty() || ArchetypeComposition.Fragments.HasAny(RequiredAnyFragments))
-		&& ArchetypeComposition.Fragments.HasNone(RequiredNoneFragments)
-		&& ArchetypeComposition.Tags.HasAll(RequiredAllTags)
-		&& (RequiredAnyTags.IsEmpty() || ArchetypeComposition.Tags.HasAny(RequiredAnyTags))
-		&& ArchetypeComposition.Tags.HasNone(RequiredNoneTags)
-		&& ArchetypeComposition.ChunkFragments.HasAll(RequiredAllChunkFragments)
-		&& ArchetypeComposition.ChunkFragments.HasNone(RequiredNoneChunkFragments)
-		&& ArchetypeComposition.SharedFragments.HasAll(RequiredAllSharedFragments)
-		&& ArchetypeComposition.SharedFragments.HasNone(RequiredNoneSharedFragments)
-		&& ArchetypeComposition.ConstSharedFragments.HasAll(RequiredAllConstSharedFragments)
-		&& ArchetypeComposition.ConstSharedFragments.HasNone(RequiredNoneConstSharedFragments);
+	CachePropreties();
+
+	const bool bPassNegativeFilter = bHasNegativeRequirements == false
+		|| (ArchetypeComposition.Fragments.HasNone(RequiredNoneFragments)
+			&& ArchetypeComposition.Tags.HasNone(RequiredNoneTags)
+			&& ArchetypeComposition.ChunkFragments.HasNone(RequiredNoneChunkFragments)
+			&& ArchetypeComposition.SharedFragments.HasNone(RequiredNoneSharedFragments)
+			&& ArchetypeComposition.ConstSharedFragments.HasNone(RequiredNoneConstSharedFragments));
+	
+	if (bPassNegativeFilter)
+	{
+		if (bHasPositiveRequirements)
+		{
+			return ArchetypeComposition.Fragments.HasAll(RequiredAllFragments)
+				&& (RequiredAnyFragments.IsEmpty() || ArchetypeComposition.Fragments.HasAny(RequiredAnyFragments))
+				&& ArchetypeComposition.Tags.HasAll(RequiredAllTags)
+				&& (RequiredAnyTags.IsEmpty() || ArchetypeComposition.Tags.HasAny(RequiredAnyTags))
+				&& ArchetypeComposition.ChunkFragments.HasAll(RequiredAllChunkFragments)
+				&& ArchetypeComposition.SharedFragments.HasAll(RequiredAllSharedFragments)
+				&& ArchetypeComposition.ConstSharedFragments.HasAll(RequiredAllConstSharedFragments);
+		}
+		else if (bHasOptionalRequirements)
+		{
+			return DoesMatchAnyOptionals(ArchetypeComposition);
+		}
+		// else - it's fine, we passed all the filters that have been set up
+		return true;
+	}
+	return false;
 }
 
 void FMassFragmentRequirements::ExportRequirements(FMassExecutionRequirements& OutRequirements) const
@@ -157,6 +204,7 @@ void FMassFragmentRequirements::ExportRequirements(FMassExecutionRequirements& O
 	OutRequirements.RequiredAllTags = RequiredAllTags;
 	OutRequirements.RequiredAnyTags = RequiredAnyTags;
 	OutRequirements.RequiredNoneTags = RequiredNoneTags;
+	// not exporting optional tags by design
 }
 
 void FMassFragmentRequirements::Reset()
diff --git a/Engine/Plugins/Runtime/MassEntity/Source/MassEntity/Public/MassDebugger.h b/Engine/Plugins/Runtime/MassEntity/Source/MassEntity/Public/MassDebugger.h
--- a/Engine/Plugins/Runtime/MassEntity/Source/MassEntity/Public/MassDebugger.h
+++ b/Engine/Plugins/Runtime/MassEntity/Source/MassEntity/Public/MassDebugger.h
@@ -65,6 +65,7 @@ namespace UE::Mass::Debug
 		const FMassTagBitSet& RequiredAllTags;
 		const FMassTagBitSet& RequiredAnyTags;
 		const FMassTagBitSet& RequiredNoneTags;
+		const FMassTagBitSet& RequiredOptionalTags;
 		const FMassExternalSubsystemBitSet& RequiredConstSubsystems;
 		const FMassExternalSubsystemBitSet& RequiredMutableSubsystems;
 	};
diff --git a/Engine/Plugins/Runtime/MassEntity/Source/MassEntity/Public/MassRequirements.h b/Engine/Plugins/Runtime/MassEntity/Source/MassEntity/Public/MassRequirements.h
--- a/Engine/Plugins/Runtime/MassEntity/Source/MassEntity/Public/MassRequirements.h
+++ b/Engine/Plugins/Runtime/MassEntity/Source/MassEntity/Public/MassRequirements.h
@@ -232,7 +232,7 @@ public:
 
 	void AddTagRequirement(const UScriptStruct& TagType, const EMassFragmentPresence Presence)
 	{
-		checkf(int(Presence) < int(EMassFragmentPresence::Optional), TEXT("Optional and MAX presence are not valid calues for AddTagRequirement"));
+		checkf(int(Presence) != int(EMassFragmentPresence::MAX), TEXT("MAX presence is not a valid value for AddTagRequirement"));
 		switch (Presence)
 		{
 		case EMassFragmentPresence::All:
@@ -244,6 +244,9 @@ public:
 		case EMassFragmentPresence::None:
 			RequiredNoneTags.Add(TagType);
 			break;
+		case EMassFragmentPresence::Optional:
+			RequiredOptionalTags.Add(TagType);
+			break;
 		}
 		IncrementChangeCounter();
 	}
@@ -251,7 +254,7 @@ public:
 	template<typename T>
 	FMassFragmentRequirements& AddTagRequirement(const EMassFragmentPresence Presence)
 	{
-		checkf(int(Presence) < int(EMassFragmentPresence::Optional), TEXT("Optional and MAX presence are not valid calues for AddTagRequirement"));
+		checkf(int(Presence) != int(EMassFragmentPresence::MAX), TEXT("MAX presence is not a valid value for AddTagRequirement"));
 		static_assert(TIsDerivedFrom<T, FMassTag>::IsDerived, "Given struct doesn't represent a valid tag type. Make sure to inherit from FMassFragment or one of its child-types.");
 		switch (Presence)
 		{
@@ -264,6 +267,9 @@ public:
 			case EMassFragmentPresence::None:
 				RequiredNoneTags.Add<T>();
 				break;
+			case EMassFragmentPresence::Optional:
+				RequiredOptionalTags.Add<T>();
+				break;
 		}
 		IncrementChangeCounter();
 		return *this;
@@ -307,7 +313,7 @@ public:
 	template<typename T>
 	FMassFragmentRequirements& AddConstSharedRequirement(const EMassFragmentPresence Presence = EMassFragmentPresence::All)
 	{
-		static_assert(TIsDerivedFrom<T, FMassConstSharedFragment>::IsDerived, "Given struct doesn't represent a valid shared fragment type. Make sure to inherit from FMassConstSharedFragment or one of its child-types.");
+		static_assert(TIsDerivedFrom<T, FMassConstSharedFragment>::IsDerived, "Given struct doesn't represent a valid const shared fragment type. Make sure to inherit from FMassConstSharedFragment or one of its child-types.");
 		checkf(ConstSharedFragmentRequirements.FindByPredicate([](const FMassFragmentRequirementDescription& Item) { return Item.StructType == T::StaticStruct(); }) == nullptr
 			, TEXT("Duplicated requirements are not supported. %s already present"), *T::StaticStruct()->GetName());
 		checkf(Presence != EMassFragmentPresence::Any, TEXT("\'Any\' is not a valid Presence value for AddConstSharedRequirement."));
@@ -333,7 +339,7 @@ public:
 	FMassFragmentRequirements& AddConstSharedRequirement(const UScriptStruct* FragmentType, const EMassFragmentPresence Presence = EMassFragmentPresence::All)
 	{
 		if (!ensureMsgf(FragmentType->IsChildOf(FMassConstSharedFragment::StaticStruct())
-			, TEXT("Given struct doesn't represent a valid shared fragment type. Make sure to inherit from FMassConstSharedFragment or one of its child-types.")))
+			, TEXT("Given struct doesn't represent a valid const shared fragment type. Make sure to inherit from FMassConstSharedFragment or one of its child-types.")))
 		{
 			return *this;
 		}
@@ -415,6 +421,7 @@ public:
 	const FMassTagBitSet& GetRequiredAllTags() const { return RequiredAllTags; }
 	const FMassTagBitSet& GetRequiredAnyTags() const { return RequiredAnyTags; }
 	const FMassTagBitSet& GetRequiredNoneTags() const { return RequiredNoneTags; }
+	const FMassTagBitSet& GetRequiredOptionalTags() const { return RequiredOptionalTags; }
 	const FMassChunkFragmentBitSet& GetRequiredAllChunkFragments() const { return RequiredAllChunkFragments; }
 	const FMassChunkFragmentBitSet& GetRequiredOptionalChunkFragments() const { return RequiredOptionalChunkFragments; }
 	const FMassChunkFragmentBitSet& GetRequiredNoneChunkFragments() const { return RequiredNoneChunkFragments; }
@@ -426,16 +433,25 @@ public:
 	const FMassConstSharedFragmentBitSet& GetRequiredNoneConstSharedFragments() const { return RequiredNoneConstSharedFragments; }
 
 	bool IsEmpty() const;
+	bool HasPositiveRequirements() const { return bHasPositiveRequirements; }
+	bool HasNegativeRequirements() const { return bHasNegativeRequirements; }
+	bool HasOptionalRequirements() const { return bHasOptionalRequirements; }
 
 	bool DoesArchetypeMatchRequirements(const FMassArchetypeHandle& ArchetypeHandle) const;
 	bool DoesArchetypeMatchRequirements(const FMassArchetypeCompositionDescriptor& ArchetypeComposition) const;
+	bool DoesMatchAnyOptionals(const FMassArchetypeCompositionDescriptor& ArchetypeComposition) const;
+
 	bool DoesRequireGameThreadExecution() const { return bRequiresGameThreadExecution; }
 	void ExportRequirements(FMassExecutionRequirements& OutRequirements) const;
 
 protected:
 	void SortRequirements();
 
-	FORCEINLINE void IncrementChangeCounter() { ++IncrementalChangesCount; }
+	FORCEINLINE void IncrementChangeCounter() 
+	{ 
+		++IncrementalChangesCount; 
+		bPropertiesCached = false;
+	}
 	void ConsumeIncrementalChangesCount() { IncrementalChangesCount = 0; }
 	bool HasIncrementalChanges() const{ return IncrementalChangesCount > 0; }
 
@@ -449,6 +465,11 @@ protected:
 	FMassTagBitSet RequiredAllTags;
 	FMassTagBitSet RequiredAnyTags;
 	FMassTagBitSet RequiredNoneTags;
+	/**
+	 * note that optional tags have meaning only if there are no other strict requirements, i.e. everything is optional,
+	 * so we're looking for anything matching any of the optionals (both tags as well as fragments).
+	 */
+	FMassTagBitSet RequiredOptionalTags;
 	FMassFragmentBitSet RequiredAllFragments;
 	FMassFragmentBitSet RequiredAnyFragments;
 	FMassFragmentBitSet RequiredOptionalFragments;
@@ -464,10 +485,17 @@ protected:
 	FMassConstSharedFragmentBitSet RequiredNoneConstSharedFragments;
 
 private:
-	mutable uint16 bValidityIsCached : 1 = false;
-	mutable uint16 bAreRequirementsValid : 1 = false;
-	mutable uint16 bEmptynessIsCached : 1 = false;
-	mutable uint16 bAreRequirementsEmpty: 1 = false;
+	FORCEINLINE void CachePropreties() const;
+	mutable uint16 bPropertiesCached : 1 = false;
+	mutable uint16 bHasPositiveRequirements : 1 = false;
+	mutable uint16 bHasNegativeRequirements : 1 = false;
+	/** 
+	 * Indicates that the requirements specify only optional elements, which means any composition having any one of 
+	 * the optional elements will be accepted. Note that RequiredNone* requirements are handled separately and if specified 
+	 * still need to be satisfied.
+	 */
+	//mutable uint16 bOptionalsOnly : 1 = false;
+	mutable uint16 bHasOptionalRequirements : 1 = false;
 
 	uint16 IncrementalChangesCount = 0;
 
@@ -500,3 +528,12 @@ FORCEINLINE FMassFragmentRequirements& FMassFragmentRequirements::AddTagRequirem
 	IncrementChangeCounter();
 	return *this;
 }
+
+template<>
+FORCEINLINE FMassFragmentRequirements& FMassFragmentRequirements::AddTagRequirements<EMassFragmentPresence::Optional>(const FMassTagBitSet& TagBitSet)
+{
+	RequiredOptionalTags += TagBitSet;
+	// force recaching the next time this query is used or the following CacheArchetypes call.
+	IncrementChangeCounter();
+	return *this;
+}
