[Mass] Added a unit test for "all optional" queries as well as "non-fragment requirement" and "negative only" queries

diff --git a/Engine/Plugins/Runtime/MassEntity/Source/MassEntityTestSuite/Private/MassQueryTest.cpp b/Engine/Plugins/Runtime/MassEntity/Source/MassEntityTestSuite/Private/MassQueryTest.cpp
--- a/Engine/Plugins/Runtime/MassEntity/Source/MassEntityTestSuite/Private/MassQueryTest.cpp
+++ b/Engine/Plugins/Runtime/MassEntity/Source/MassEntityTestSuite/Private/MassQueryTest.cpp
@@ -324,51 +324,35 @@ struct FQueryTest_FragmentPresent : FEntityTestBase
 IMPLEMENT_AI_INSTANT_TEST(FQueryTest_FragmentPresent, "System.Mass.Query.FragmentPresent");
 
 
-struct FQueryTest_OnlySingleAbsentFragment : FEntityTestBase
+struct FQueryTest_OnlyAbsentFragments : FEntityTestBase
 {
 	virtual bool InstantTest() override
 	{
 		CA_ASSUME(EntityManager);
 
 		FMassEntityQuery Query;
-		Query.AddRequirement<FTestFragment_Int>(EMassFragmentAccess::ReadOnly, EMassFragmentPresence::None);
-		
-		AITEST_FALSE("The query is not valid", Query.CheckValidity());
-
-		GetTestRunner().AddExpectedError(TEXT("requirements not valid"), EAutomationExpectedErrorFlags::Contains, 1);
-		Query.CacheArchetypes(*EntityManager);
-
-		// this is an invalid query. We expect an error and an empty valid archetypes array
-		AITEST_EQUAL("The query is invalid and we expect no archetypes match", Query.GetArchetypes().Num(), 0);
+		AITEST_FALSE("The empty query is not valid", Query.CheckValidity());
 
-		return true;
-	}
-};
-IMPLEMENT_AI_INSTANT_TEST(FQueryTest_OnlySingleAbsentFragment, "System.Mass.Query.OnlySingleAbsentFragment");
-
-
-struct FQueryTest_OnlyMultipleAbsentFragments : FEntityTestBase
-{
-	virtual bool InstantTest() override
-	{
-		CA_ASSUME(EntityManager);
-
-		FMassEntityQuery Query;
 		Query.AddRequirement<FTestFragment_Int>(EMassFragmentAccess::ReadOnly, EMassFragmentPresence::None);
-		Query.AddRequirement<FTestFragment_Float>(EMassFragmentAccess::ReadOnly, EMassFragmentPresence::None);
-
-		AITEST_FALSE("The query is not valid", Query.CheckValidity());
+		AITEST_TRUE("Single negative requirement is valid", Query.CheckValidity());
 
-		GetTestRunner().AddExpectedError(TEXT("requirements not valid"), EAutomationExpectedErrorFlags::Contains, 1);
+		Query.AddRequirement<FTestFragment_Float>(EMassFragmentAccess::ReadOnly, EMassFragmentPresence::None);
+		Query.AddRequirement<FTestFragment_Bool>(EMassFragmentAccess::ReadOnly, EMassFragmentPresence::None);
+		AITEST_TRUE("Multiple negative requirement is valid", Query.CheckValidity());
+		
 		Query.CacheArchetypes(*EntityManager);
+		AITEST_EQUAL("There's only one default test archetype matching the query", Query.GetArchetypes().Num(), 1);
+		AITEST_TRUE("Only the Empty archetype matches the query", Query.GetArchetypes()[0] == EmptyArchetype);
 
-		// this is an invalid query. We expect an error and an empty valid archetypes array
-		AITEST_EQUAL("The query is invalid and we expect no archetypes match", Query.GetArchetypes().Num(), 0);
+		const FMassArchetypeHandle NewMatchingArchetypeHandle = EntityManager->CreateArchetype({ FTestFragment_Large::StaticStruct() });
+		Query.CacheArchetypes(*EntityManager);
+		AITEST_EQUAL("The number of matching queries matches expectations", Query.GetArchetypes().Num(), 2);
+		AITEST_TRUE("The new archetype matches the query", Query.GetArchetypes()[1] == NewMatchingArchetypeHandle);
 
 		return true;
 	}
 };
-IMPLEMENT_AI_INSTANT_TEST(FQueryTest_OnlyMultipleAbsentFragments, "System.Mass.Query.OnlyMultipleAbsentFragments");
+IMPLEMENT_AI_INSTANT_TEST(FQueryTest_OnlyAbsentFragments, "System.Mass.Query.OnlyAbsentFragments");
 
 
 struct FQueryTest_AbsentAndPresentFragments : FEntityTestBase
@@ -601,6 +585,304 @@ struct FQueryTest_AutoRecache : FEntityTestBase
 };
 IMPLEMENT_AI_INSTANT_TEST(FQueryTest_AutoRecache, "System.Mass.Query.AutoReCaching");
 
+
+struct FQueryTest_AllOptional : FEntityTestBase
+{
+	virtual bool InstantTest() override
+	{
+		CA_ASSUME(EntityManager);
+
+		FMassEntityQuery Query;
+		Query.AddRequirement<FTestFragment_Float>(EMassFragmentAccess::None, EMassFragmentPresence::Optional);
+		Query.AddTagRequirement<FTestTag_A>(EMassFragmentPresence::Optional);
+		Query.AddChunkRequirement<FTestChunkFragment_Int>(EMassFragmentAccess::None, EMassFragmentPresence::Optional);
+		Query.AddSharedRequirement<FTestSharedFragment_Int>(EMassFragmentAccess::None, EMassFragmentPresence::Optional);
+		Query.AddConstSharedRequirement<FTestConstSharedFragment_Int>(EMassFragmentPresence::Optional);
+
+		Query.CacheArchetypes(*EntityManager);
+
+		int32 ExpectedNumOfArchetypes = 2;
+		// only the FloatsArchetype and FloatsIntsArchetype should match
+		AITEST_TRUE("Initial number of matching archetypes matches expectations", Query.GetArchetypes().Num() == ExpectedNumOfArchetypes);
+
+		TArray<FMassEntityHandle> Entities;
+		EntityManager->BatchCreateEntities(IntsArchetype, 10, Entities);
+
+		int32 CurrentEntityIndex = 0;
+
+		EntityManager->AddTagToEntity(Entities[CurrentEntityIndex++], FTestTag_A::StaticStruct());
+		++ExpectedNumOfArchetypes;
+		EntityManager->AddTagToEntity(Entities[CurrentEntityIndex++], FTestTag_B::StaticStruct());
+		Query.CacheArchetypes(*EntityManager);
+		AITEST_EQUAL("A: number of matching archetypes matches expectations.", Query.GetArchetypes().Num(), ExpectedNumOfArchetypes);
+
+		{
+			FMassArchetypeCompositionDescriptor Descriptor(EntityManager->GetArchetypeComposition(IntsArchetype));
+			Descriptor.ChunkFragments.Add<FTestChunkFragment_Int>();
+			EntityManager->CreateArchetype(Descriptor);
+			++ExpectedNumOfArchetypes;
+		}
+		{
+			FMassArchetypeCompositionDescriptor Descriptor(EntityManager->GetArchetypeComposition(IntsArchetype));
+			Descriptor.ChunkFragments.Add<FTestChunkFragment_Float>();
+			EntityManager->CreateArchetype(Descriptor);
+		}
+		Query.CacheArchetypes(*EntityManager);
+		AITEST_EQUAL("B: number of matching archetypes matches expectations.", Query.GetArchetypes().Num(), ExpectedNumOfArchetypes);
+
+		{
+			FTestSharedFragment_Int FragmentInstance;
+			FSharedStruct SharedFragmentInstance = FSharedStruct::Make(FragmentInstance);
+			FMassArchetypeSharedFragmentValues SharedFragmentValues;
+			SharedFragmentValues.AddSharedFragment(SharedFragmentInstance);
+
+			FMassArchetypeEntityCollection Collection(IntsArchetype, MakeArrayView(&Entities[CurrentEntityIndex++], 1), FMassArchetypeEntityCollection::EDuplicatesHandling::NoDuplicates);
+			EntityManager->BatchAddSharedFragmentsForEntities(MakeArrayView(&Collection, 1), SharedFragmentValues);
+			++ExpectedNumOfArchetypes;
+		}
+		{
+			FTestSharedFragment_Float FragmentInstance;
+			FSharedStruct SharedFragmentInstance = FSharedStruct::Make(FragmentInstance);
+			FMassArchetypeSharedFragmentValues SharedFragmentValues;
+			SharedFragmentValues.AddSharedFragment(SharedFragmentInstance);
+
+			FMassArchetypeEntityCollection Collection(IntsArchetype, MakeArrayView(&Entities[CurrentEntityIndex++], 1), FMassArchetypeEntityCollection::EDuplicatesHandling::NoDuplicates);
+			EntityManager->BatchAddSharedFragmentsForEntities(MakeArrayView(&Collection, 1), SharedFragmentValues);
+		}
+		Query.CacheArchetypes(*EntityManager);
+		AITEST_EQUAL("C: number of matching archetypes matches expectations.", Query.GetArchetypes().Num(), ExpectedNumOfArchetypes);
+
+		{
+			FTestConstSharedFragment_Int FragmentInstance;
+			FConstSharedStruct SharedFragmentInstance = FSharedStruct::Make(FragmentInstance);
+			EntityManager->AddConstSharedFragmentToEntity(Entities[CurrentEntityIndex++], SharedFragmentInstance);
+			++ExpectedNumOfArchetypes;
+		}
+		{
+			FTestConstSharedFragment_Float FragmentInstance;
+			FConstSharedStruct SharedFragmentInstance = FSharedStruct::Make(FragmentInstance);
+			EntityManager->AddConstSharedFragmentToEntity(Entities[CurrentEntityIndex++], SharedFragmentInstance);
+		}
+		Query.CacheArchetypes(*EntityManager);
+		AITEST_EQUAL("D: number of matching archetypes matches expectations.", Query.GetArchetypes().Num(), ExpectedNumOfArchetypes);
+
+		return true;
+	}
+};
+IMPLEMENT_AI_INSTANT_TEST(FQueryTest_AllOptional, "System.Mass.Query.AllOptional");
+
+struct FQueryTest_JustATag : FEntityTestBase
+{
+	virtual bool InstantTest() override
+	{
+		CA_ASSUME(EntityManager);
+
+		FMassEntityQuery Query;
+		Query.AddTagRequirement<FTestTag_A>(EMassFragmentPresence::All);
+		Query.CacheArchetypes(*EntityManager);
+
+		int32 ExpectedNumOfArchetypes = 0;
+		// only the FloatsArchetype and FloatsIntsArchetype should match
+		AITEST_TRUE("Initial number of matching archetypes matches expectations", Query.GetArchetypes().Num() == ExpectedNumOfArchetypes);
+
+		{
+			FMassArchetypeCompositionDescriptor Descriptor(EntityManager->GetArchetypeComposition(IntsArchetype));
+			Descriptor.Tags.Add<FTestTag_A>();
+			EntityManager->CreateArchetype(Descriptor);
+			++ExpectedNumOfArchetypes;
+		}
+		{
+			FMassArchetypeCompositionDescriptor Descriptor(EntityManager->GetArchetypeComposition(IntsArchetype));
+			Descriptor.Tags.Add<FTestTag_B>();
+			EntityManager->CreateArchetype(Descriptor);
+		}
+		Query.CacheArchetypes(*EntityManager);
+		AITEST_EQUAL("A: number of matching archetypes matches expectations.", Query.GetArchetypes().Num(), ExpectedNumOfArchetypes);
+
+		{
+			FMassArchetypeCompositionDescriptor Descriptor(EntityManager->GetArchetypeComposition(IntsArchetype));
+			Descriptor.Tags.Add<FTestTag_A>();
+			Descriptor.Tags.Add<FTestTag_C>();
+			Descriptor.Tags.Add<FTestTag_D>();
+			EntityManager->CreateArchetype(Descriptor);
+			++ExpectedNumOfArchetypes;
+		}
+		{
+			FMassArchetypeCompositionDescriptor Descriptor(EntityManager->GetArchetypeComposition(IntsArchetype));
+			Descriptor.Tags.Add<FTestTag_B>();
+			Descriptor.Tags.Add<FTestTag_C>();
+			Descriptor.Tags.Add<FTestTag_D>();
+			EntityManager->CreateArchetype(Descriptor);
+		}
+		Query.CacheArchetypes(*EntityManager);
+		AITEST_EQUAL("B: number of matching archetypes matches expectations.", Query.GetArchetypes().Num(), ExpectedNumOfArchetypes);
+		
+		return true;
+	}
+};
+IMPLEMENT_AI_INSTANT_TEST(FQueryTest_JustATag, "System.Mass.Query.JustATag");
+
+struct FQueryTest_JustAChunkFragment : FEntityTestBase
+{
+	virtual bool InstantTest() override
+	{
+		CA_ASSUME(EntityManager);
+
+		FMassArchetypeHandle TargetArchetype;
+
+		FMassEntityQuery Query;
+		Query.AddChunkRequirement<FTestChunkFragment_Int>(EMassFragmentAccess::ReadOnly, EMassFragmentPresence::All);
+		Query.CacheArchetypes(*EntityManager);
+
+		int32 ExpectedNumOfArchetypes = 0;
+		// no matching archetypes at this time
+		AITEST_TRUE("Initial number of matching archetypes matches expectations", Query.GetArchetypes().Num() == ExpectedNumOfArchetypes);
+
+		{
+			FMassArchetypeCompositionDescriptor Descriptor(EntityManager->GetArchetypeComposition(IntsArchetype));
+			Descriptor.ChunkFragments.Add<FTestChunkFragment_Int>();
+			TargetArchetype = EntityManager->CreateArchetype(Descriptor);
+			++ExpectedNumOfArchetypes;
+		}
+		{
+			FMassArchetypeCompositionDescriptor Descriptor(EntityManager->GetArchetypeComposition(IntsArchetype));
+			Descriptor.ChunkFragments.Add<FTestChunkFragment_Float>();
+			EntityManager->CreateArchetype(Descriptor);
+		}
+		Query.CacheArchetypes(*EntityManager);
+		AITEST_EQUAL("Number of matching archetypes matches expectations.", Query.GetArchetypes().Num(), ExpectedNumOfArchetypes);
+		
+		// try to access the chunk fragment
+		{
+			EntityManager->CreateEntity(TargetArchetype);
+
+			FMassExecutionContext ExecContext(*EntityManager.Get());
+			bool bExecuted = false;
+			Query.ForEachEntityChunk(*EntityManager, ExecContext, [&bExecuted](FMassExecutionContext& Context)
+				{
+					const FTestChunkFragment_Int& ChunkFragment = Context.GetChunkFragment<FTestChunkFragment_Int>();
+					bExecuted = true;
+				});
+			AITEST_TRUE("The tested query did execute and bounding was successful", bExecuted);
+		}
+
+		return true;
+	}
+};
+IMPLEMENT_AI_INSTANT_TEST(FQueryTest_JustAChunkFragment, "System.Mass.Query.JustAChunkFragment");
+
+
+struct FQueryTest_JustASharedFragment : FEntityTestBase
+{
+	virtual bool InstantTest() override
+	{
+		CA_ASSUME(EntityManager);
+
+		FMassArchetypeHandle TargetArchetype;
+
+		FMassEntityQuery Query;
+		Query.AddSharedRequirement<FTestSharedFragment_Int>(EMassFragmentAccess::ReadOnly, EMassFragmentPresence::All);
+		Query.CacheArchetypes(*EntityManager);
+
+		int32 ExpectedNumOfArchetypes = 0;
+		// no matching archetypes at this time
+		AITEST_TRUE("Initial number of matching archetypes matches expectations", Query.GetArchetypes().Num() == ExpectedNumOfArchetypes);
+
+		TArray<FMassEntityHandle> Entities;
+		EntityManager->BatchCreateEntities(IntsArchetype, 10, Entities);
+		int32 CurrentEntityIndex = 0;
+
+		{
+			FTestSharedFragment_Int FragmentInstance;
+			FSharedStruct SharedFragmentInstance = FSharedStruct::Make(FragmentInstance);
+			FMassArchetypeSharedFragmentValues SharedFragmentValues;
+			SharedFragmentValues.AddSharedFragment(SharedFragmentInstance);
+
+			FMassArchetypeEntityCollection Collection(IntsArchetype, MakeArrayView(&Entities[CurrentEntityIndex++], 1), FMassArchetypeEntityCollection::EDuplicatesHandling::NoDuplicates);
+			EntityManager->BatchAddSharedFragmentsForEntities(MakeArrayView(&Collection, 1), SharedFragmentValues);
+			++ExpectedNumOfArchetypes;
+		}
+		{
+			FTestSharedFragment_Float FragmentInstance;
+			FSharedStruct SharedFragmentInstance = FSharedStruct::Make(FragmentInstance);
+			FMassArchetypeSharedFragmentValues SharedFragmentValues;
+			SharedFragmentValues.AddSharedFragment(SharedFragmentInstance);
+
+			FMassArchetypeEntityCollection Collection(IntsArchetype, MakeArrayView(&Entities[CurrentEntityIndex++], 1), FMassArchetypeEntityCollection::EDuplicatesHandling::NoDuplicates);
+			EntityManager->BatchAddSharedFragmentsForEntities(MakeArrayView(&Collection, 1), SharedFragmentValues);
+		}
+		Query.CacheArchetypes(*EntityManager);
+		AITEST_EQUAL("Number of matching archetypes matches expectations.", Query.GetArchetypes().Num(), ExpectedNumOfArchetypes);
+
+		// try to access the shared fragment
+		{
+			bool bExecuted = false;
+			FMassExecutionContext ExecContext(*EntityManager.Get());
+			Query.ForEachEntityChunk(*EntityManager, ExecContext, [&bExecuted](FMassExecutionContext& Context)
+				{
+					const FTestSharedFragment_Int& SharedFragment = Context.GetSharedFragment<FTestSharedFragment_Int>();
+					bExecuted = true;
+				});
+			AITEST_TRUE("The tested query did execute and bounding was successful", bExecuted);
+		}
+
+		return true;
+	}
+};
+IMPLEMENT_AI_INSTANT_TEST(FQueryTest_JustASharedFragment, "System.Mass.Query.JustASharedFragment");
+
+
+struct FQueryTest_JustAConstSharedFragment : FEntityTestBase
+{
+	virtual bool InstantTest() override
+	{
+		CA_ASSUME(EntityManager);
+
+		FMassArchetypeHandle TargetArchetype;
+
+		FMassEntityQuery Query;
+		Query.AddConstSharedRequirement<FTestConstSharedFragment_Int>(EMassFragmentPresence::All);
+		Query.CacheArchetypes(*EntityManager);
+
+		int32 ExpectedNumOfArchetypes = 0;
+		// no matching archetypes at this time
+		AITEST_TRUE("Initial number of matching archetypes matches expectations", Query.GetArchetypes().Num() == ExpectedNumOfArchetypes);
+
+		TArray<FMassEntityHandle> Entities;
+		EntityManager->BatchCreateEntities(IntsArchetype, 10, Entities);
+		int32 CurrentEntityIndex = 0;
+
+		{
+			FTestConstSharedFragment_Int FragmentInstance;
+			FConstSharedStruct SharedFragmentInstance = FSharedStruct::Make(FragmentInstance);
+			EntityManager->AddConstSharedFragmentToEntity(Entities[CurrentEntityIndex++], SharedFragmentInstance);
+			++ExpectedNumOfArchetypes;
+		}
+		{
+			FTestConstSharedFragment_Float FragmentInstance;
+			FConstSharedStruct SharedFragmentInstance = FSharedStruct::Make(FragmentInstance);
+			EntityManager->AddConstSharedFragmentToEntity(Entities[CurrentEntityIndex++], SharedFragmentInstance);
+		}
+		Query.CacheArchetypes(*EntityManager);
+		AITEST_EQUAL("Number of matching archetypes matches expectations.", Query.GetArchetypes().Num(), ExpectedNumOfArchetypes);
+
+		// try to access the shared fragment
+		{
+			bool bExecuted = false;
+			FMassExecutionContext ExecContext(*EntityManager.Get());
+			Query.ForEachEntityChunk(*EntityManager, ExecContext, [&bExecuted](FMassExecutionContext& Context)
+				{
+					const FTestConstSharedFragment_Int& SharedFragment = Context.GetConstSharedFragment<FTestConstSharedFragment_Int>();
+					bExecuted = true;
+				});
+			AITEST_TRUE("The tested query did execute and bounding was successful", bExecuted);
+		}
+
+		return true;
+	}
+};
+IMPLEMENT_AI_INSTANT_TEST(FQueryTest_JustAConstSharedFragment, "System.Mass.Query.JustAConstSharedFragment");
+
 } // FMassQueryTest
 
 #undef LOCTEXT_NAMESPACE
diff --git a/Engine/Plugins/Runtime/MassEntity/Source/MassEntityTestSuite/Public/MassEntityTestTypes.h b/Engine/Plugins/Runtime/MassEntity/Source/MassEntityTestSuite/Public/MassEntityTestTypes.h
--- a/Engine/Plugins/Runtime/MassEntity/Source/MassEntityTestSuite/Public/MassEntityTestTypes.h
+++ b/Engine/Plugins/Runtime/MassEntity/Source/MassEntityTestSuite/Public/MassEntityTestTypes.h
@@ -98,6 +98,15 @@ struct FTestChunkFragment_Int : public FMassChunkFragment
 	FTestChunkFragment_Int(const int32 InValue = 0) : Value(InValue) {}
 };
 
+USTRUCT()
+struct FTestChunkFragment_Float : public FMassChunkFragment
+{
+	GENERATED_BODY()
+	float Value = 0.f;
+
+	FTestChunkFragment_Float(const float InValue = 0.f) : Value(InValue) {}
+};
+
 USTRUCT()
 struct FTestSharedFragment_Int : public FMassSharedFragment
 {
