[MassEntity]
- added new processor execution flags 'EditorWorls' and 'AllWorldModes' that could be used for processors that need to run when associated to a world in either client, server, standalone or Editor mode
- added DetermineProcessorSupportedTickTypes so processors executed in an Editor world can be ticked in all tick modes

diff --git a/Engine/Plugins/Runtime/MassEntity/Source/MassEntity/Private/MassEntityUtils.cpp b/Engine/Plugins/Runtime/MassEntity/Source/MassEntity/Private/MassEntityUtils.cpp
--- a/Engine/Plugins/Runtime/MassEntity/Source/MassEntity/Private/MassEntityUtils.cpp
+++ b/Engine/Plugins/Runtime/MassEntity/Source/MassEntity/Private/MassEntityUtils.cpp
@@ -3,38 +3,38 @@
 #include "MassEntityUtils.h"
 #include "MassEntityTypes.h"
 #include "MassArchetypeTypes.h"
-#include "MassArchetypeData.h"
-#include "MassRequirements.h"
 #include "MassEntityManager.h"
 #include "MassEntitySubsystem.h"
 #include "Engine/EngineBaseTypes.h"
 #include "Engine/World.h"
+#if WITH_EDITOR
+#include "Editor.h"
+#endif // WITH_EDITOR
 
 namespace UE::Mass::Utils
 {
 EProcessorExecutionFlags GetProcessorExecutionFlagsForWorld(const UWorld& World)
 {
+#if WITH_EDITOR
 	if (World.IsEditorWorld() && !World.IsGameWorld())
 	{
-		return EProcessorExecutionFlags::Editor;
+		return EProcessorExecutionFlags::EditorWorld;
 	}
-	else
+#endif // WITH_EDITOR
+
+	switch (const ENetMode NetMode = World.GetNetMode())
 	{
-		const ENetMode NetMode = World.GetNetMode();
-		switch (NetMode)
-		{
-		case NM_ListenServer:
-			return EProcessorExecutionFlags::Client | EProcessorExecutionFlags::Server;
-		case NM_DedicatedServer:
-			return EProcessorExecutionFlags::Server;
-		case NM_Client:
-			return EProcessorExecutionFlags::Client;
-		case NM_Standalone:
-			return EProcessorExecutionFlags::Standalone;
-		default:
-			checkf(false, TEXT("Unsupported ENetMode type (%i) found while determining MASS processor execution flags."), NetMode);
-			return EProcessorExecutionFlags::None;
-		}
+	case NM_ListenServer:
+		return EProcessorExecutionFlags::Client | EProcessorExecutionFlags::Server;
+	case NM_DedicatedServer:
+		return EProcessorExecutionFlags::Server;
+	case NM_Client:
+		return EProcessorExecutionFlags::Client;
+	case NM_Standalone:
+		return EProcessorExecutionFlags::Standalone;
+	default:
+		checkf(false, TEXT("Unsupported ENetMode type (%i) found while determining MASS processor execution flags."), NetMode);
+		return EProcessorExecutionFlags::None;
 	}
 }
 
@@ -46,11 +46,29 @@ EProcessorExecutionFlags DetermineProcessorExecutionFlags(const UWorld* World, E
 	}
 	if (World)
 	{
-		return UE::Mass::Utils::GetProcessorExecutionFlagsForWorld(*World);
+		return GetProcessorExecutionFlagsForWorld(*World);
 	}
+
+#if WITH_EDITOR
+	if (GEditor)
+	{
+		return EProcessorExecutionFlags::Editor;
+	}
+#endif // WITH_EDITOR
 	return EProcessorExecutionFlags::All;
 }
 
+uint8 DetermineProcessorSupportedTickTypes(const UWorld* World)
+{
+#if WITH_EDITOR
+	if (World != nullptr && GetProcessorExecutionFlagsForWorld(*World) == EProcessorExecutionFlags::EditorWorld)
+	{
+		return MAX_uint8;
+	}
+#endif // WITH_EDITOR
+	return (1 << LEVELTICK_All) | (1 << LEVELTICK_TimeOnly);
+}
+
 void CreateEntityCollections(const FMassEntityManager& EntityManager, const TConstArrayView<FMassEntityHandle> Entities
 	, const FMassArchetypeEntityCollection::EDuplicatesHandling DuplicatesHandling, TArray<FMassArchetypeEntityCollection>& OutEntityCollections)
 {
diff --git a/Engine/Plugins/Runtime/MassEntity/Source/MassEntity/Private/MassProcessingPhaseManager.cpp b/Engine/Plugins/Runtime/MassEntity/Source/MassEntity/Private/MassProcessingPhaseManager.cpp
--- a/Engine/Plugins/Runtime/MassEntity/Source/MassEntity/Private/MassProcessingPhaseManager.cpp
+++ b/Engine/Plugins/Runtime/MassEntity/Source/MassEntity/Private/MassProcessingPhaseManager.cpp
@@ -206,16 +206,19 @@ void FMassProcessingPhaseManager::Initialize(UObject& InOwner, TConstArrayView<F
 	ProcessingPhasesConfig = InProcessingPhasesConfig;
 
 	ProcessorExecutionFlags = UE::Mass::Utils::DetermineProcessorExecutionFlags(World, ProcessorExecutionFlags);
+	const uint8 SupportedTickTypes = UE::Mass::Utils::DetermineProcessorSupportedTickTypes(World);
 
 	for (int PhaseAsInt = 0; PhaseAsInt < int(EMassProcessingPhase::MAX); ++PhaseAsInt)
 	{		
 		const EMassProcessingPhase Phase = EMassProcessingPhase(PhaseAsInt);
+		FMassProcessingPhase& ProcessingPhase = ProcessingPhases[PhaseAsInt];
 
 		UMassCompositeProcessor* PhaseProcessor = NewObject<UMassCompositeProcessor>(&InOwner, UMassCompositeProcessor::StaticClass()
 			, *FString::Printf(TEXT("ProcessingPhase_%s"), *UEnum::GetDisplayValueAsText(Phase).ToString()));
 	
 		check(PhaseProcessor);
-		ProcessingPhases[PhaseAsInt].Initialize(*this, Phase, UE::Mass::Private::PhaseToTickingGroup[PhaseAsInt], *PhaseProcessor);
+		ProcessingPhase.Initialize(*this, Phase, UE::Mass::Private::PhaseToTickingGroup[PhaseAsInt], *PhaseProcessor);
+		ProcessingPhase.SupportedTickTypes = SupportedTickTypes;
 
 		REDIRECT_OBJECT_TO_VLOG(PhaseProcessor, &InOwner);
 		PhaseProcessor->SetProcessingPhase(Phase);
diff --git a/Engine/Plugins/Runtime/MassEntity/Source/MassEntity/Public/MassEntityUtils.h b/Engine/Plugins/Runtime/MassEntity/Source/MassEntity/Public/MassEntityUtils.h
--- a/Engine/Plugins/Runtime/MassEntity/Source/MassEntity/Public/MassEntityUtils.h
+++ b/Engine/Plugins/Runtime/MassEntity/Source/MassEntity/Public/MassEntityUtils.h
@@ -15,9 +15,12 @@ namespace UE::Mass::Utils
 /** returns the current execution mode for the processors calculated from the world network mode */
 MASSENTITY_API extern EProcessorExecutionFlags GetProcessorExecutionFlagsForWorld(const UWorld& World);
 
-/** based on given World (which can be null) and ExecutionFlagsOverride the function determines the execution flags to use */
+/** based on the given World (which can be null) and 'ExecutionFlagsOverride', the function determines the execution flags to use */
 MASSENTITY_API extern EProcessorExecutionFlags DetermineProcessorExecutionFlags(const UWorld* World, EProcessorExecutionFlags ExecutionFlagsOverride = EProcessorExecutionFlags::None);
 
+/** based on the given World (which can be null), the function determines additional level tick types for the processing phases */
+uint8 DetermineProcessorSupportedTickTypes(const UWorld* World);
+
 UE_DEPRECATED(5.3, "This function is deprecated. Please use the properly spelled GetProcessorExecutionFlagsForWorld instead.")
 FORCEINLINE EProcessorExecutionFlags GetProcessorExecutionFlagsForWold(const UWorld& World)
 {
diff --git a/Engine/Plugins/Runtime/MassEntity/Source/MassEntity/Public/MassProcessingTypes.h b/Engine/Plugins/Runtime/MassEntity/Source/MassEntity/Public/MassProcessingTypes.h
--- a/Engine/Plugins/Runtime/MassEntity/Source/MassEntity/Public/MassProcessingTypes.h
+++ b/Engine/Plugins/Runtime/MassEntity/Source/MassEntity/Public/MassProcessingTypes.h
@@ -25,8 +25,10 @@ enum class EProcessorExecutionFlags : uint8
 	Server = 1 << 1,
 	Client = 1 << 2,
 	Editor = 1 << 3,
+	EditorWorld = 1 << 4,
 	AllNetModes = Standalone | Server | Client UMETA(Hidden),
-	All = Standalone | Server | Client | Editor UMETA(Hidden)
+	AllWorldModes = Standalone | Server | Client | EditorWorld UMETA(Hidden),
+	All = Standalone | Server | Client | Editor | EditorWorld UMETA(Hidden)
 };
 ENUM_CLASS_FLAGS(EProcessorExecutionFlags);
 
