[MassEntity] fixed some command debug names

diff --git a/Engine/Plugins/Runtime/MassEntity/Source/MassEntity/Public/MassCommands.h b/Engine/Plugins/Runtime/MassEntity/Source/MassEntity/Public/MassCommands.h
--- a/Engine/Plugins/Runtime/MassEntity/Source/MassEntity/Public/MassCommands.h
+++ b/Engine/Plugins/Runtime/MassEntity/Source/MassEntity/Public/MassCommands.h
@@ -176,9 +176,11 @@ protected:
 
 struct FMassCommandDestroyEntities : public FMassBatchedEntityCommand
 {
+	using Super = FMassBatchedEntityCommand;
+
 	FMassCommandDestroyEntities()
+		: Super(EMassCommandOperationType::Destroy DEBUG_NAME("DestroyEntities"))
 	{
-		OperationType = EMassCommandOperationType::Destroy;
 	}
 
 protected:
@@ -198,9 +200,9 @@ protected:
 template<EMassCommandCheckTime CheckTime, typename... TTypes>
 struct FMassCommandAddFragmentsInternal : public FMassBatchedEntityCommand
 {
-	using Super = FMassBatchedEntityCommand; 
+	using Super = FMassBatchedEntityCommand;
 	FMassCommandAddFragmentsInternal()
-		: Super(EMassCommandOperationType::Add)
+		: Super(EMassCommandOperationType::Add DEBUG_NAME("AddFragments"))
 		, FragmentsAffected(UE::Mass::Utils::ConstructFragmentBitSet<CheckTime, TTypes...>())
 	{}
 
@@ -223,7 +225,7 @@ struct FMassCommandRemoveFragmentsInternal : public FMassBatchedEntityCommand
 {
 	using Super = FMassBatchedEntityCommand;
 	FMassCommandRemoveFragmentsInternal()
-		: Super(EMassCommandOperationType::Remove)
+		: Super(EMassCommandOperationType::Remove DEBUG_NAME("RemoveFragments"))
 		, FragmentsAffected(UE::Mass::Utils::ConstructFragmentBitSet<CheckTime, TTypes...>())
 	{}
 
@@ -248,11 +250,11 @@ struct FMassCommandChangeTags : public FMassBatchedEntityCommand
 {
 	using Super = FMassBatchedEntityCommand;
 	FMassCommandChangeTags()
-		: Super(EMassCommandOperationType::ChangeComposition)
+		: Super(EMassCommandOperationType::ChangeComposition DEBUG_NAME("ChangeTags"))
 	{}
 
 	FMassCommandChangeTags(EMassCommandOperationType OperationType, FMassTagBitSet TagsToAdd, FMassTagBitSet TagsToRemove)
-		: Super(OperationType)
+		: Super(OperationType DEBUG_NAME("ChangeTags"))
 		, TagsToAdd(TagsToAdd)
 		, TagsToRemove(TagsToRemove)
 	{}
@@ -293,7 +295,7 @@ struct FMassCommandAddTagsInternal : public FMassCommandChangeTags
 			EMassCommandOperationType::Add, 
 			UE::Mass::Utils::ConstructTagBitSet<CheckTime, TTypes...>(),
 			{} 
-			DEBUG_NAME("CommandAddTag"))
+			DEBUG_NAME("AddTags"))
 	{}
 };
 
@@ -312,7 +314,7 @@ struct FMassCommandRemoveTagsInternal : public FMassCommandChangeTags
 			EMassCommandOperationType::Remove, 
 			{}, 
 			UE::Mass::Utils::ConstructTagBitSet<CheckTime, TTypes...>()
-			DEBUG_NAME("CommandRemoveTag"))
+			DEBUG_NAME("RemoveTags"))
 	{}
 };
 
@@ -331,7 +333,7 @@ struct FMassCommandSwapTagsInternal : public FMassCommandChangeTags
 			EMassCommandOperationType::ChangeComposition,
 			UE::Mass::Utils::ConstructTagBitSet<CheckTime, TNew>(),
 			UE::Mass::Utils::ConstructTagBitSet<CheckTime, TOld>()
-			DEBUG_NAME("CommandSwapTags"))
+			DEBUG_NAME("SwapTags"))
 	{}
 };
 
@@ -351,6 +353,12 @@ struct FMassCommandAddFragmentInstances : public FMassBatchedEntityCommand
 		, FragmentsAffected(UE::Mass::Utils::ConstructFragmentBitSet<EMassCommandCheckTime::CompileTimeCheck, TOthers...>())
 	{}
 
+#if CSV_PROFILER || WITH_MASSENTITY_DEBUG
+	FMassCommandAddFragmentInstances(EMassCommandOperationType OperationType, FName DebugName)
+		: Super(OperationType, DebugName)
+	{}
+#endif // CSV_PROFILER || WITH_MASSENTITY_DEBUG
+
 	void Add(FMassEntityHandle Entity, TOthers... InFragments)
 	{
 		Super::Add(Entity);
@@ -391,16 +399,14 @@ protected:
 template<typename... TOthers>
 struct FMassCommandBuildEntity : public FMassCommandAddFragmentInstances<TOthers...>
 {
+	using Super = FMassCommandAddFragmentInstances<TOthers...>;
+
 	FMassCommandBuildEntity()
+		: Super(EMassCommandOperationType::Create DEBUG_NAME("BuildEntity"))
 	{
-		Super::OperationType = EMassCommandOperationType::Create;
-#if CSV_PROFILER_STATS || WITH_MASSENTITY_DEBUG
-		Super::DebugName = TEXT("FMassCommandBuildEntity");
-#endif // CSV_PROFILER_STATS || WITH_MASSENTITY_DEBUG
 	}
 
 protected:
-	using Super = FMassCommandAddFragmentInstances<TOthers...>;
 
 	virtual void Execute(FMassEntityManager& System) const override
 	{
