[MassEntity] updated FMassArchetypeCompositionDescriptor::DebugOutputDescription to include shared fragments

diff --git a/Engine/Plugins/Runtime/MassEntity/Source/MassEntity/Public/MassEntityTypes.h b/Engine/Plugins/Runtime/MassEntity/Source/MassEntity/Public/MassEntityTypes.h
--- a/Engine/Plugins/Runtime/MassEntity/Source/MassEntity/Public/MassEntityTypes.h
+++ b/Engine/Plugins/Runtime/MassEntity/Source/MassEntity/Public/MassEntityTypes.h
@@ -265,15 +265,46 @@ struct FMassArchetypeCompositionDescriptor
 	void DebugOutputDescription(FOutputDevice& Ar) const
 	{
 #if WITH_MASSENTITY_DEBUG
+		if (Fragments.IsEmpty()
+			&& Tags.IsEmpty()
+			&& ChunkFragments.IsEmpty())
+		{
+			Ar.Logf(TEXT("Empty"));
+			return;
+		}
+
 		const bool bAutoLineEnd = Ar.GetAutoEmitLineTerminator();
 		Ar.SetAutoEmitLineTerminator(false);
 
-		Ar.Logf(TEXT("Fragments:\n"));
-		Fragments.DebugGetStringDesc(Ar);
-		Ar.Logf(TEXT("Tags:\n"));
-		Tags.DebugGetStringDesc(Ar);
-		Ar.Logf(TEXT("ChunkFragments:\n"));
-		ChunkFragments.DebugGetStringDesc(Ar);
+		if (!Fragments.IsEmpty())
+		{
+			Ar.Logf(TEXT("Fragments:\n"));
+			Fragments.DebugGetStringDesc(Ar);
+		}
+
+		if (!Tags.IsEmpty())
+		{
+			Ar.Logf(TEXT("Tags:\n"));
+			Tags.DebugGetStringDesc(Ar);
+		}
+
+		if (!ChunkFragments.IsEmpty())
+		{
+			Ar.Logf(TEXT("ChunkFragments:\n"));
+			ChunkFragments.DebugGetStringDesc(Ar);
+		}
+
+		if (!SharedFragments.IsEmpty())
+		{
+			Ar.Logf(TEXT("SharedFragments:\n"));
+			SharedFragments.DebugGetStringDesc(Ar);
+		}
+
+		if (!ConstSharedFragments.IsEmpty())
+		{
+			Ar.Logf(TEXT("ConstSharedFragments:\n"));
+			ConstSharedFragments.DebugGetStringDesc(Ar);
+		}
 
 		Ar.SetAutoEmitLineTerminator(bAutoLineEnd);
 #endif // WITH_MASSENTITY_DEBUG
