Fix to FMassArchetypeSharedFragmentValues::HasExactSharedFragmentTypesMatch and HasExactConstSharedFragmentTypesMatch::FMassArchetypeSharedFragmentValues - the previous version was failing when A.HasAll(B) was true but B.HasAll(A) was not (which can happen in practive with the shared fragment removal added recently).

diff --git a/Engine/Plugins/Runtime/MassEntity/Source/MassEntityTestSuite/Private/MassSharedFragmentsTest.cpp b/Engine/Plugins/Runtime/MassEntity/Source/MassEntityTestSuite/Private/MassSharedFragmentsTest.cpp
--- a/Engine/Plugins/Runtime/MassEntity/Source/MassEntityTestSuite/Private/MassSharedFragmentsTest.cpp
+++ b/Engine/Plugins/Runtime/MassEntity/Source/MassEntityTestSuite/Private/MassSharedFragmentsTest.cpp
@@ -560,6 +560,45 @@ struct FSharedFragment_BatchAddToEmpty : FEntityTestBase
 };
 IMPLEMENT_AI_INSTANT_TEST(FSharedFragment_BatchAddToEmpty, "System.Mass.SharedFragments.BatchAddToEmpty");
 
+struct FSharedFragmentValues_TypeEquivalency : FEntityTestBase
+{
+	virtual bool InstantTest() override
+	{
+		constexpr int32 TestIntValue = 32;
+
+		FMassArchetypeSharedFragmentValues Values;
+		const FMassSharedFragmentBitSet EmptySharedFragmentBitSet;
+		const FMassConstSharedFragmentBitSet EmptyConstSharedFragmentBitSet;
+
+		AITEST_TRUE("Empty shared values match type with empty bitset", Values.HasExactSharedFragmentTypesMatch(EmptySharedFragmentBitSet));
+		AITEST_TRUE("Empty const shared values match type with empty const bitset", Values.HasExactConstSharedFragmentTypesMatch(EmptyConstSharedFragmentBitSet));
+
+		const FMassSharedFragmentBitSet IntSharedFragmentBitSet = FMassSharedFragmentBitSet::GetTypeBitSet<FTestSharedFragment_Int>();
+		FMassSharedFragmentBitSet IntFloatSharedFragmentBitSet = IntSharedFragmentBitSet;
+		IntFloatSharedFragmentBitSet.Add<FTestSharedFragment_Float>();
+		Values.AddSharedFragment(FSharedStruct::Make<FTestSharedFragment_Int>(TestIntValue));
+		AITEST_TRUE("Single shared value type matches expected bitset", Values.HasExactSharedFragmentTypesMatch(IntSharedFragmentBitSet));
+		AITEST_FALSE("Single shared value type doesn't match two-type bitset", Values.HasExactSharedFragmentTypesMatch(IntFloatSharedFragmentBitSet));
+		AITEST_FALSE("Single shared value type doesn't match empty", Values.HasExactSharedFragmentTypesMatch(EmptySharedFragmentBitSet));
+
+		const FMassConstSharedFragmentBitSet IntConstSharedFragmentBitSet = FMassConstSharedFragmentBitSet::GetTypeBitSet<FTestConstSharedFragment_Int>();
+		FMassConstSharedFragmentBitSet IntFloatConstSharedFragmentBitSet = IntConstSharedFragmentBitSet;
+		IntFloatConstSharedFragmentBitSet.Add<FTestConstSharedFragment_Float>();
+		Values.AddConstSharedFragment(FSharedStruct::Make<FTestConstSharedFragment_Int>(TestIntValue));
+		AITEST_TRUE("Single const shared value type matches expected bitset", Values.HasExactConstSharedFragmentTypesMatch(IntConstSharedFragmentBitSet));
+		AITEST_FALSE("Single const shared value type doesn't match two-type bitset", Values.HasExactConstSharedFragmentTypesMatch(IntFloatConstSharedFragmentBitSet));
+		AITEST_FALSE("Single const shared value type doesn't match empty", Values.HasExactConstSharedFragmentTypesMatch(EmptyConstSharedFragmentBitSet));
+
+		Values.Remove(IntSharedFragmentBitSet);
+		AITEST_TRUE("Emptied shared values match type with empty bitset", Values.HasExactSharedFragmentTypesMatch(EmptySharedFragmentBitSet));
+		Values.Remove(IntConstSharedFragmentBitSet);
+		AITEST_TRUE("Emptied const shared values match type with empty const bitset", Values.HasExactConstSharedFragmentTypesMatch(EmptyConstSharedFragmentBitSet));
+
+		return true;
+	}
+};
+IMPLEMENT_AI_INSTANT_TEST(FSharedFragmentValues_TypeEquivalency, "System.Mass.SharedFragments.TypeEquivalency");
+
 } // FMassEntityTest
 
 UE_ENABLE_OPTIMIZATION_SHIP
diff --git a/Engine/Source/Runtime/MassEntity/Public/MassEntityTypes.h b/Engine/Source/Runtime/MassEntity/Public/MassEntityTypes.h
--- a/Engine/Source/Runtime/MassEntity/Public/MassEntityTypes.h
+++ b/Engine/Source/Runtime/MassEntity/Public/MassEntityTypes.h
@@ -360,7 +360,7 @@ struct MASSENTITY_API FMassArchetypeSharedFragmentValues
 
 	FORCEINLINE bool HasExactSharedFragmentTypesMatch(const FMassSharedFragmentBitSet& InSharedFragmentBitSet) const
 	{
-		return SharedFragmentBitSet == InSharedFragmentBitSet;
+		return SharedFragmentBitSet.IsEquivalent(InSharedFragmentBitSet);
 	}
 
 	FORCEINLINE bool HasAllRequiredSharedFragmentTypes(const FMassSharedFragmentBitSet& InSharedFragmentBitSet) const
@@ -370,7 +370,7 @@ struct MASSENTITY_API FMassArchetypeSharedFragmentValues
 
 	FORCEINLINE bool HasExactConstSharedFragmentTypesMatch(const FMassConstSharedFragmentBitSet& InConstSharedFragmentBitSet) const
 	{
-		return ConstSharedFragmentBitSet == InConstSharedFragmentBitSet;
+		return ConstSharedFragmentBitSet.IsEquivalent(InConstSharedFragmentBitSet);
 	}
 
 	FORCEINLINE bool HasAllRequiredConstSharedFragmentTypes(const FMassConstSharedFragmentBitSet& InConstSharedFragmentBitSet) const
