[Mass] Moved some editor-time code around as a fallout after making MassEntity an engine module.
This change fixes "MassEntity" MessageLog category now showing up (it does, but the name is missing).

diff --git a/Engine/Plugins/Runtime/MassEntity/Source/MassEntityEditor/MassEntityEditor.Build.cs b/Engine/Plugins/Runtime/MassEntity/Source/MassEntityEditor/MassEntityEditor.Build.cs
--- a/Engine/Plugins/Runtime/MassEntity/Source/MassEntityEditor/MassEntityEditor.Build.cs
+++ b/Engine/Plugins/Runtime/MassEntity/Source/MassEntityEditor/MassEntityEditor.Build.cs
@@ -24,7 +24,7 @@ namespace UnrealBuildTool.Rules
 				"AssetTools",
 				"UnrealEd",
 				"Slate",
-				"SlateCore",				
+				"SlateCore",
 				"PropertyEditor",
 				"MassEntity",
 				"DetailCustomizations",
@@ -43,7 +43,11 @@ namespace UnrealBuildTool.Rules
 				"ToolMenus",
 			}
 			);
-		}
 
+			if (Target.bBuildDeveloperTools == true)
+			{
+				PrivateDependencyModuleNames.Add("MessageLog");
+			}
+		}
 	}
 }
diff --git a/Engine/Plugins/Runtime/MassEntity/Source/MassEntityEditor/Private/MassEntityEditorModule.cpp b/Engine/Plugins/Runtime/MassEntity/Source/MassEntityEditor/Private/MassEntityEditorModule.cpp
--- a/Engine/Plugins/Runtime/MassEntity/Source/MassEntityEditor/Private/MassEntityEditorModule.cpp
+++ b/Engine/Plugins/Runtime/MassEntity/Source/MassEntityEditor/Private/MassEntityEditorModule.cpp
@@ -3,7 +3,11 @@
 #include "MassEntityEditorModule.h"
 #include "MassEditorStyle.h"
 #include "Modules/ModuleManager.h"
-
+#if WITH_UNREAL_DEVELOPER_TOOLS
+#include "MessageLogModule.h"
+#include "Engine/World.h"
+#include "Logging/MessageLog.h"
+#endif // WITH_UNREAL_DEVELOPER_TOOLS
 
 #define LOCTEXT_NAMESPACE "MassEntityEditor"
 
@@ -15,6 +19,16 @@ void FMassEntityEditorModule::StartupModule()
 	ToolBarExtensibilityManager = MakeShareable(new FExtensibilityManager);
 
 	FMassEntityEditorStyle::Initialize();
+
+#if WITH_UNREAL_DEVELOPER_TOOLS
+	FMessageLogModule& MessageLogModule = FModuleManager::LoadModuleChecked<FMessageLogModule>("MessageLog");
+	FMessageLogInitializationOptions InitOptions;
+	InitOptions.bShowPages = true;
+	InitOptions.bShowFilters = true;
+	MessageLogModule.RegisterLogListing("MassEntity", LOCTEXT("MassEntity", "MassEntity"), InitOptions);
+
+	OnWorldCleanupHandle = FWorldDelegates::OnWorldCleanup.AddStatic(&FMassEntityEditorModule::OnWorldCleanup);
+#endif // WITH_UNREAL_DEVELOPER_TOOLS
 }
 
 void FMassEntityEditorModule::ShutdownModule()
@@ -24,5 +38,18 @@ void FMassEntityEditorModule::ShutdownModule()
 	ToolBarExtensibilityManager.Reset();
 
 	FMassEntityEditorStyle::Shutdown();
+
+#if WITH_UNREAL_DEVELOPER_TOOLS
+	FWorldDelegates::OnWorldCleanup.Remove(OnWorldCleanupHandle);
+#endif // WITH_UNREAL_DEVELOPER_TOOLS
 }
+
+#if WITH_UNREAL_DEVELOPER_TOOLS
+void FMassEntityEditorModule::OnWorldCleanup(UWorld* /*World*/, bool /*bSessionEnded*/, bool /*bCleanupResources*/)
+{
+	// clearing out messages from the world being cleaned up
+	FMessageLog("MassEntity").NewPage(FText::FromString(TEXT("MassEntity")));
+}
+#endif // WITH_UNREAL_DEVELOPER_TOOLS
+
 #undef LOCTEXT_NAMESPACE
diff --git a/Engine/Plugins/Runtime/MassEntity/Source/MassEntityEditor/Public/MassEntityEditorModule.h b/Engine/Plugins/Runtime/MassEntity/Source/MassEntityEditor/Public/MassEntityEditorModule.h
--- a/Engine/Plugins/Runtime/MassEntity/Source/MassEntityEditor/Public/MassEntityEditorModule.h
+++ b/Engine/Plugins/Runtime/MassEntity/Source/MassEntityEditor/Public/MassEntityEditorModule.h
@@ -4,10 +4,12 @@
 
 #include "Modules/ModuleInterface.h"
 #include "Toolkits/AssetEditorToolkit.h"
+#include "Delegates/IDelegateInstance.h"
 
 class IMassEntityEditor;
 struct FGraphPanelNodeFactory;
 struct FGraphNodeClassHelper;
+class UWorld;
 
 /**
 * The public interface to this module
@@ -26,6 +28,11 @@ public:
 	TSharedPtr<FGraphNodeClassHelper> GetProcessorClassCache() { return ProcessorClassCache; }
 
 protected:
+#if WITH_UNREAL_DEVELOPER_TOOLS
+	static void OnWorldCleanup(UWorld* /*World*/, bool /*bSessionEnded*/, bool /*bCleanupResources*/);
+	FDelegateHandle OnWorldCleanupHandle;
+#endif // WITH_UNREAL_DEVELOPER_TOOLS
+
 	TSharedPtr<FGraphNodeClassHelper> ProcessorClassCache;
 
 	TSharedPtr<FExtensibilityManager> MenuExtensibilityManager;
diff --git a/Engine/Source/Runtime/MassEntity/MassEntity.Build.cs b/Engine/Source/Runtime/MassEntity/MassEntity.Build.cs
--- a/Engine/Source/Runtime/MassEntity/MassEntity.Build.cs
+++ b/Engine/Source/Runtime/MassEntity/MassEntity.Build.cs
@@ -25,11 +25,6 @@ namespace UnrealBuildTool.Rules
 				PrivateDependencyModuleNames.Add("EditorSubsystem");
 			}
 
-			if (Target.bBuildDeveloperTools == true)
-			{
-				PrivateDependencyModuleNames.Add("MessageLog");
-			}
-
 			if (Target.Configuration != UnrealTargetConfiguration.Shipping
 				&& Target.Configuration != UnrealTargetConfiguration.Test)
 			{
diff --git a/Engine/Source/Runtime/MassEntity/Private/MassEntityModule.cpp b/Engine/Source/Runtime/MassEntity/Private/MassEntityModule.cpp
deleted file mode 100644
diff --git a/Engine/Source/Runtime/MassEntity/Private/MassProcessingTypes.cpp b/Engine/Source/Runtime/MassEntity/Private/MassProcessingTypes.cpp
--- a/Engine/Source/Runtime/MassEntity/Private/MassProcessingTypes.cpp
+++ b/Engine/Source/Runtime/MassEntity/Private/MassProcessingTypes.cpp
@@ -6,11 +6,29 @@
 #include "MassEntityUtils.h"
 #include "VisualLogger/VisualLogger.h"
 #include "MassDebugger.h"
+#include "Misc/CoreDelegates.h"
 
 #include UE_INLINE_GENERATED_CPP_BY_NAME(MassProcessingTypes)
 
 DEFINE_LOG_CATEGORY(LogMass);
 
+
+namespace UE::Mass::Private
+{
+	struct FMassBootReporter
+	{
+		FMassBootReporter()
+		{
+#if MASS_DO_PARALLEL
+			UE_LOG(LogMass, Log, TEXT("MassEntity running with MULTITHREADING support."));
+#else
+			UE_LOG(LogMass, Log, TEXT("MassEntity running in game thread."));
+#endif // MASS_DO_PARALLEL
+		}
+	};
+	static const FMassBootReporter MassBootReporter;
+}
+
 //----------------------------------------------------------------------//
 //  FMassProcessingContext
 //----------------------------------------------------------------------//
diff --git a/Engine/Source/Runtime/MassEntity/Public/IMassEntityModule.h b/Engine/Source/Runtime/MassEntity/Public/IMassEntityModule.h
deleted file mode 100644
