Add RemoveConstSharedFragment functionality to Mass
* Some care taken to remove null SharedStructs when removed from the MassArchetypeSharedFragmentValues class.   FMassArchetypeSharedFragmentValues::CalculateHash would return different hash value if nulls present resulting in unnecessary new chunks in FMassArchetypeData::AddEntityInternal

diff --git a/Engine/Plugins/Runtime/MassEntity/Source/MassEntityTestSuite/Private/MassSharedFragmentsTest.cpp b/Engine/Plugins/Runtime/MassEntity/Source/MassEntityTestSuite/Private/MassSharedFragmentsTest.cpp
--- a/Engine/Plugins/Runtime/MassEntity/Source/MassEntityTestSuite/Private/MassSharedFragmentsTest.cpp
+++ b/Engine/Plugins/Runtime/MassEntity/Source/MassEntityTestSuite/Private/MassSharedFragmentsTest.cpp
@@ -454,6 +454,43 @@ struct FSharedFragment_AddToEntity : FEntityTestBase
 };
 IMPLEMENT_AI_INSTANT_TEST(FSharedFragment_AddToEntity, "System.Mass.SharedFragments.AddToEntity");
 
+struct FSharedFragment_RemoveFromEntity : FEntityTestBase
+{
+	virtual bool InstantTest() override
+	{
+		constexpr int32 TestIntValue = 1023;
+		FTestConstSharedFragment_Int FragmentInstance(TestIntValue);
+		FConstSharedStruct SharedFragmentInstance = FSharedStruct::Make(FragmentInstance);
+
+		const FMassEntityHandle EntityHandle = EntityManager->CreateEntity(FloatsArchetype);
+
+		FTestConstSharedFragment_Int* EntitySharedFragment = EntityManager->GetConstSharedFragmentDataPtr<FTestConstSharedFragment_Int>(EntityHandle);
+		AITEST_NULL("Initially the entity is not expected to have the shared fragment", EntitySharedFragment);
+
+		AITEST_FALSE("Attempt to remove shared fragment from entity that doesn't have shared fragment should return false and do nothing",
+			EntityManager->RemoveConstSharedFragmentFromEntity(EntityHandle, *FTestConstSharedFragment_Int::StaticStruct()));
+
+		AITEST_TRUE("Adding shared fragment to entity should succeed",
+			EntityManager->AddConstSharedFragmentToEntity(EntityHandle, SharedFragmentInstance));
+
+		EntitySharedFragment = EntityManager->GetConstSharedFragmentDataPtr<FTestConstSharedFragment_Int>(EntityHandle);
+		AITEST_NOT_NULL("The entity is expected to have the shared fragment after the operation", EntitySharedFragment);
+		AITEST_EQUAL("The the shared fragment is expected to store the configured value", EntitySharedFragment->Value, TestIntValue);
+
+		AITEST_TRUE(
+			"Removing shared fragment from entity that has the shared fragment should succeed",
+			EntityManager->RemoveConstSharedFragmentFromEntity(EntityHandle, *FTestConstSharedFragment_Int::StaticStruct()));
+
+		EntitySharedFragment = EntityManager->GetConstSharedFragmentDataPtr<FTestConstSharedFragment_Int>(EntityHandle);
+		AITEST_NULL("The entity is not expected to have the shared fragment after the operation", EntitySharedFragment);
+		
+		AITEST_EQUAL("The the entity's new archetype is the same as the initial one", EntityManager->GetArchetypeForEntity(EntityHandle), FloatsArchetype);
+
+		return true;
+	}
+};
+IMPLEMENT_AI_INSTANT_TEST(FSharedFragment_RemoveFromEntity, "System.Mass.SharedFragments.RemoveFromEntity");
+
 struct FSharedFragment_BatchAddToEntity : FEntityTestBase
 {
 	virtual bool InstantTest() override
diff --git a/Engine/Source/Runtime/MassEntity/Private/MassEntityManager.cpp b/Engine/Source/Runtime/MassEntity/Private/MassEntityManager.cpp
--- a/Engine/Source/Runtime/MassEntity/Private/MassEntityManager.cpp
+++ b/Engine/Source/Runtime/MassEntity/Private/MassEntityManager.cpp
@@ -1388,6 +1388,47 @@ bool FMassEntityManager::AddConstSharedFragmentToEntity(const FMassEntityHandle
 	return true;
 }
 
+bool FMassEntityManager::RemoveConstSharedFragmentFromEntity(const FMassEntityHandle Entity, const UScriptStruct& ConstSharedFragmentType)
+{
+	if (!ensureMsgf(ConstSharedFragmentType.IsChildOf(FMassConstSharedFragment::StaticStruct()), TEXT("%hs parameter ConstSharedFragmentType is expected to be a FMassConstSharedFragment"), __FUNCTION__))
+	{
+		return false;
+	}
+	
+	CheckIfEntityIsActive(Entity);
+	
+	FMassArchetypeData* CurrentArchetype = GetEntityStorageInterface().GetArchetypeAsShared(Entity.Index).Get();
+	check(CurrentArchetype);
+	
+	if (!CurrentArchetype->GetCompositionDescriptor().ConstSharedFragments.Contains(ConstSharedFragmentType))
+	{
+		// Nothing to do
+		return true;
+	}
+
+	FMassArchetypeCompositionDescriptor NewComposition(CurrentArchetype->GetCompositionDescriptor());
+	NewComposition.ConstSharedFragments.Remove(ConstSharedFragmentType);
+	const FMassArchetypeHandle NewArchetypeHandle = CreateArchetype(NewComposition);
+	check(NewArchetypeHandle.IsValid());
+	FMassArchetypeData* NewArchetype = NewArchetypeHandle.DataPtr.Get();
+	check(NewArchetype);
+	
+	const FMassArchetypeSharedFragmentValues& OldSharedFragmentValues = CurrentArchetype->GetSharedFragmentValues(Entity.Index);
+	check(OldSharedFragmentValues.ContainsType(&ConstSharedFragmentType));
+	FMassArchetypeSharedFragmentValues NewSharedFragmentValues(OldSharedFragmentValues);
+	
+	const FMassConstSharedFragmentBitSet ToRemove(ConstSharedFragmentType);
+	NewSharedFragmentValues.Remove(ToRemove);
+	NewSharedFragmentValues.Sort();
+	
+	CurrentArchetype->MoveEntityToAnotherArchetype(Entity, *NewArchetype, &NewSharedFragmentValues);
+
+	// Change the entity archetype
+	GetEntityStorageInterface().SetArchetypeFromShared(Entity.Index, NewArchetypeHandle.DataPtr);
+	
+	return true;
+}
+
 void FMassEntityManager::BatchChangeTagsForEntities(TConstArrayView<FMassArchetypeEntityCollection> EntityCollections, const FMassTagBitSet& TagsToAdd, const FMassTagBitSet& TagsToRemove)
 {
 	TRACE_CPUPROFILER_EVENT_SCOPE(Mass_BatchChangeTagsForEntities);
diff --git a/Engine/Source/Runtime/MassEntity/Private/MassEntityTypes.cpp b/Engine/Source/Runtime/MassEntity/Private/MassEntityTypes.cpp
--- a/Engine/Source/Runtime/MassEntity/Private/MassEntityTypes.cpp
+++ b/Engine/Source/Runtime/MassEntity/Private/MassEntityTypes.cpp
@@ -220,6 +220,7 @@ int32 FMassArchetypeSharedFragmentValues::Remove(const FMassSharedFragmentBitSet
 
 	if (RemovedCount)
 	{
+		SharedFragments.RemoveAllSwap([](const FSharedStruct& SharedStruct) { return !SharedStruct.IsValid(); });
 		SharedFragmentBitSet -= CommonFragments;
 		DirtyHashCache();
 	}
@@ -248,6 +249,7 @@ int32 FMassArchetypeSharedFragmentValues::Remove(const FMassConstSharedFragmentB
 
 	if (RemovedCount)
 	{
+		ConstSharedFragments.RemoveAllSwap([](const FConstSharedStruct& SharedStruct) { return !SharedStruct.IsValid(); });
 		ConstSharedFragmentBitSet -= CommonFragments;
 		DirtyHashCache();
 	}
diff --git a/Engine/Source/Runtime/MassEntity/Public/MassEntityManager.h b/Engine/Source/Runtime/MassEntity/Public/MassEntityManager.h
--- a/Engine/Source/Runtime/MassEntity/Public/MassEntityManager.h
+++ b/Engine/Source/Runtime/MassEntity/Public/MassEntityManager.h
@@ -362,6 +362,13 @@ public:
 	 */
 	bool AddConstSharedFragmentToEntity(const FMassEntityHandle Entity, const FConstSharedStruct& InConstSharedFragment);
 
+	/**
+	 * Removes a const shared fragment of the given type from the entity.
+	 * Will do nothing if entity did not have the shared fragment.
+	 * @return True if fragment removed from entity, false otherwise.
+	 */
+	bool RemoveConstSharedFragmentFromEntity(const FMassEntityHandle Entity, const UScriptStruct& ConstSharedFragmentType);
+
 	/** 
 	 * Reserves Count number of entities and appends them to InOutEntities
 	 * @return a view into InOutEntities containing only the freshly reserved entities
