[Mass] Marked the MassEntity plugin as deprecated in 5.5, following the move of all of its code over to the engine.

diff --git a/Engine/Plugins/AI/MassAI/MassAI.uplugin b/Engine/Plugins/AI/MassAI/MassAI.uplugin
--- a/Engine/Plugins/AI/MassAI/MassAI.uplugin
+++ b/Engine/Plugins/AI/MassAI/MassAI.uplugin
@@ -47,10 +47,6 @@
 		}
 	],
 	"Plugins": [
-		{
-			"Name": "MassEntity",
-			"Enabled": true
-		},
 		{
 			"Name": "MassGameplay",
 			"Enabled": true
diff --git a/Engine/Plugins/AI/MassCrowd/MassCrowd.uplugin b/Engine/Plugins/AI/MassCrowd/MassCrowd.uplugin
--- a/Engine/Plugins/AI/MassCrowd/MassCrowd.uplugin
+++ b/Engine/Plugins/AI/MassCrowd/MassCrowd.uplugin
@@ -27,10 +27,6 @@
 			"Name": "ZoneGraph",
 			"Enabled": true
 		},
-		{
-			"Name": "MassEntity",
-			"Enabled": true
-		},
 		{
 			"Name": "MassGameplay",
 			"Enabled": true
diff --git a/Engine/Plugins/Runtime/InstancedActors/InstancedActors.uplugin b/Engine/Plugins/Runtime/InstancedActors/InstancedActors.uplugin
--- a/Engine/Plugins/Runtime/InstancedActors/InstancedActors.uplugin
+++ b/Engine/Plugins/Runtime/InstancedActors/InstancedActors.uplugin
@@ -30,10 +30,6 @@
 		}
 	],
 	"Plugins": [
-		{
-			"Name": "MassEntity",
-			"Enabled": true
-		},
 		{
 			"Name": "MassGameplay",
 			"Enabled": true
diff --git a/Engine/Plugins/Runtime/MassEntity/MassEntity.uplugin b/Engine/Plugins/Runtime/MassEntity/MassEntity.uplugin
--- a/Engine/Plugins/Runtime/MassEntity/MassEntity.uplugin
+++ b/Engine/Plugins/Runtime/MassEntity/MassEntity.uplugin
@@ -3,7 +3,7 @@
 	"Version": 1,
 	"VersionName": "0.4",
 	"FriendlyName": "MassEntity",
-	"Description": "Gameplay-focused framework supporting data-oriented processing",
+	"Description": "Gameplay-focused framework supporting data-oriented processing. This plugin is now DEPRECATED. All of MassEntity code has been moved to the engine and there's no need to enable this plugin anymore.'",
 	"Category": "Gameplay",
 	"CreatedBy": "Epic Games, Inc.",
 	"CreatedByURL": "https://epicgames.com",
@@ -14,6 +14,7 @@
 	"IsBetaVersion": true,
 	"Installed": false,
 	"EnabledByDefault": false,
+	"DeprecatedEngineVersion": "5.5",
 	"Modules": [],
 	"Plugins": []
 }
\ No newline at end of file
diff --git a/Engine/Plugins/Runtime/MassGameplay/MassGameplay.uplugin b/Engine/Plugins/Runtime/MassGameplay/MassGameplay.uplugin
--- a/Engine/Plugins/Runtime/MassGameplay/MassGameplay.uplugin
+++ b/Engine/Plugins/Runtime/MassGameplay/MassGameplay.uplugin
@@ -84,10 +84,6 @@
 		}
 	],
 	"Plugins": [
-		{
-			"Name": "MassEntity",
-			"Enabled": true
-		},
 		{
 			"Name": "ZoneGraph",
 			"Enabled": true
